#ttp 331398 - EDITOR: Logspam when opening Visual Studio
#proj UE4
#branch UE4
#summary Improved the error detection when VSAccessor fails to access a Visual Studio instance
#extra The non-DTE version of VSAccessor is now available as a fallback should the DTE version fail at runtime.
#change The AccessVisualStudio functions now return more information about why they failed so that the calling code can determine whether the function failed because Visual Studio was closed, or because the function failed to acccess COM data. This lets the calling code abort its attempt if the access function failed for a legitimate reason.
#change Improved the logging should AccessVisualStudio fail.

diff --git a/Engine/Source/Developer/Windows/VSAccessor/Private/VSAccessorModule.cpp b/Engine/Source/Developer/Windows/VSAccessor/Private/VSAccessorModule.cpp
--- a/Engine/Source/Developer/Windows/VSAccessor/Private/VSAccessorModule.cpp
+++ b/Engine/Source/Developer/Windows/VSAccessor/Private/VSAccessorModule.cpp
@@ -27,11 +27,10 @@
 	#pragma warning(default: 4191)
 	#pragma warning(default: 4146)
 	#pragma warning(default: 4278)
-#else
-	#include <tlhelp32.h>
-	#include <wbemidl.h>
-	#pragma comment(lib, "wbemuuid.lib")
 #endif
+#include <tlhelp32.h>
+#include <wbemidl.h>
+#pragma comment(lib, "wbemuuid.lib")
 #include "HideWindowsPlatformTypes.h"
 
 DEFINE_LOG_CATEGORY_STATIC(LogVSAccessor, Log, All);
@@ -41,6 +40,17 @@ IMPLEMENT_MODULE( FVSAccessorModule, VSAccessor );
 /** The VS query in progress message */
 TWeakPtr<class SNotificationItem> VSNotificationPtr = NULL;
 
+/** Return codes when trying to access an existing VS instance */
+enum class EAccessVisualStudioResult : uint8
+{
+	/** An instance of Visual Studio is available, and the relevant output variables have been filled in */
+	VSInstanceIsOpen,
+	/** An instance of Visual Studio is not available */
+	VSInstanceIsNotOpen,
+	/** It is unknown whether an instance of Visual Studio is available, as an error occurred when performing the check */
+	VSInstanceUnknown,
+};
+
 /** save all open documents in visual studio, when recompiling */
 void SaveVisualStudioDocuments()
 {
@@ -48,46 +58,6 @@ void SaveVisualStudioDocuments()
 	VSAccessorModule.SaveAllOpenVisualStudioDocuments();
 }
 
-bool RunVisualStudioAndOpenSolutionAndFile(const FString& ExecutablePath, const FString& SolutionPath, const FString& FilePath, const int32 FileLine)
-{
-	bool bSuccess = false;
-
-	FString Params;
-
-	// Only open the solution if it exists
-	if (!SolutionPath.IsEmpty() && FPaths::FileExists(SolutionPath))
-	{
-		Params += TEXT("\"");
-		Params += SolutionPath;
-		Params += TEXT("\"");
-	}
-
-	// Only open the file if it exists
-	if (!FilePath.IsEmpty() && FPaths::FileExists(FilePath))
-	{
-		Params += TEXT(" \"");
-		Params += FilePath;
-		Params += TEXT("\"");
-
-		if (FileLine > 0)
-		{
-			Params += FString::Printf(TEXT(" /command \"edit.goto %d\""), FileLine);
-		}
-	}
-
-	if (!Params.IsEmpty())
-	{
-		FProcHandle WorkerHandle = FPlatformProcess::CreateProc(*ExecutablePath, *Params, true, false, false, nullptr, 0, nullptr, nullptr);
-		if (WorkerHandle.IsValid())
-		{
-			bSuccess = true;
-		}
-		WorkerHandle.Close();
-	}
-
-	return bSuccess;
-}
-
 void FVSAccessorModule::StartupModule()
 {
 	VSLaunchTime = 0.0;
@@ -113,7 +83,7 @@ void FVSAccessorModule::ShutdownModule()
 
 #if VSACCESSOR_HAS_DTE
 
-bool IsVisualStudioMoniker(const FString& InName, const TArray<FVSAccessorModule::VisualStudioLocation>& InLocations)
+bool IsVisualStudioDTEMoniker(const FString& InName, const TArray<FVSAccessorModule::VisualStudioLocation>& InLocations)
 {
 	for(int32 VSVersion = 0; VSVersion < InLocations.Num(); VSVersion++)
 	{
@@ -127,62 +97,86 @@ bool IsVisualStudioMoniker(const FString& InName, const TArray<FVSAccessorModule
 }
 
 /** Accesses the correct visual studio instance if possible. */
-bool AccessVisualStudio(CComPtr<EnvDTE::_DTE>& OutDTE, const FString& InSolutionPath, const TArray<FVSAccessorModule::VisualStudioLocation>& InLocations)
+EAccessVisualStudioResult AccessVisualStudioViaDTE(CComPtr<EnvDTE::_DTE>& OutDTE, const FString& InSolutionPath, const TArray<FVSAccessorModule::VisualStudioLocation>& InLocations)
 {
-	bool bSuccess = false;
+	EAccessVisualStudioResult AccessResult = EAccessVisualStudioResult::VSInstanceIsNotOpen;
 
 	// Open the Running Object Table (ROT)
 	IRunningObjectTable* RunningObjectTable;
-	if (SUCCEEDED(GetRunningObjectTable(0, &RunningObjectTable)) &&
-		RunningObjectTable)
+	if (SUCCEEDED(GetRunningObjectTable(0, &RunningObjectTable)) && RunningObjectTable)
 	{
 		IEnumMoniker* MonikersTable;
-		RunningObjectTable->EnumRunning(&MonikersTable);
-		MonikersTable->Reset();
-
-		// Look for all visual studio instances in the ROT
-		IMoniker* CurrentMoniker;
-		while (!bSuccess && MonikersTable->Next(1, &CurrentMoniker, NULL) == S_OK)
+		if (SUCCEEDED(RunningObjectTable->EnumRunning(&MonikersTable)))
 		{
-			IBindCtx* BindContext;
-			LPOLESTR OutName;
-			CComPtr<IUnknown> ComObject;
-			if (SUCCEEDED(CreateBindCtx(0, &BindContext)) &&
-				SUCCEEDED(CurrentMoniker->GetDisplayName(BindContext, NULL, &OutName)) &&
-				IsVisualStudioMoniker(FString(OutName), InLocations) &&
-				SUCCEEDED(RunningObjectTable->GetObject(CurrentMoniker, &ComObject)))
+			MonikersTable->Reset();
+
+			// Look for all visual studio instances in the ROT
+			IMoniker* CurrentMoniker;
+			while (AccessResult != EAccessVisualStudioResult::VSInstanceIsOpen && MonikersTable->Next(1, &CurrentMoniker, NULL) == S_OK)
 			{
-				CComPtr<EnvDTE::_DTE> TempDTE;
-				TempDTE = ComObject;
-
-				// Get the solution path for this instance
-				// If it equals the solution we would have opened above in RunVisualStudio(), we'll take that
-				CComPtr<EnvDTE::_Solution> Solution;
-				LPOLESTR OutPath;
-				if (SUCCEEDED(TempDTE->get_Solution(&Solution)) &&
-					SUCCEEDED(Solution->get_FullName(&OutPath)))
+				IBindCtx* BindContext;
+				LPOLESTR OutName;
+				if (SUCCEEDED(CreateBindCtx(0, &BindContext)) && SUCCEEDED(CurrentMoniker->GetDisplayName(BindContext, NULL, &OutName)))
 				{
-					FString Filename(OutPath);
-					FPaths::NormalizeFilename(Filename);
-
-					if( Filename == InSolutionPath )
+					if (IsVisualStudioDTEMoniker(FString(OutName), InLocations))
 					{
-						OutDTE = TempDTE;
-						bSuccess = true;
+						CComPtr<IUnknown> ComObject;
+						if (SUCCEEDED(RunningObjectTable->GetObject(CurrentMoniker, &ComObject)))
+						{
+							CComPtr<EnvDTE::_DTE> TempDTE;
+							TempDTE = ComObject;
+
+							// Get the solution path for this instance
+							// If it equals the solution we would have opened above in RunVisualStudio(), we'll take that
+							CComPtr<EnvDTE::_Solution> Solution;
+							LPOLESTR OutPath;
+							if (SUCCEEDED(TempDTE->get_Solution(&Solution)) &&
+								SUCCEEDED(Solution->get_FullName(&OutPath)))
+							{
+								FString Filename(OutPath);
+								FPaths::NormalizeFilename(Filename);
+
+								if( Filename == InSolutionPath )
+								{
+									OutDTE = TempDTE;
+									AccessResult = EAccessVisualStudioResult::VSInstanceIsOpen;
+								}
+							}
+						}
+						else
+						{
+							UE_LOG(LogVSAccessor, Warning, TEXT("Couldn't get Visual Studio COM object"));
+							AccessResult = EAccessVisualStudioResult::VSInstanceUnknown;
+						}
 					}
 				}
+				else
+				{
+					UE_LOG(LogVSAccessor, Warning, TEXT("Couldn't get display name"));
+					AccessResult = EAccessVisualStudioResult::VSInstanceUnknown;
+				}
+				BindContext->Release();
+				CurrentMoniker->Release();
 			}
-			BindContext->Release();
-			CurrentMoniker->Release();
+			MonikersTable->Release();
+		}
+		else
+		{
+			UE_LOG(LogVSAccessor, Warning, TEXT("Couldn't enumerate ROT table"));
+			AccessResult = EAccessVisualStudioResult::VSInstanceUnknown;
 		}
-		MonikersTable->Release();
 		RunningObjectTable->Release();
 	}
+	else
+	{
+		UE_LOG(LogVSAccessor, Warning, TEXT("Couldn't get ROT table"));
+		AccessResult = EAccessVisualStudioResult::VSInstanceUnknown;
+	}
 
-	return bSuccess;
+	return AccessResult;
 }
 
-bool FVSAccessorModule::OpenVisualStudioSolution()
+bool FVSAccessorModule::OpenVisualStudioSolutionViaDTE()
 {
 	// Initialize the com library, if not already by this thread
 	if (!FWindowsPlatformMisc::CoInitialize())
@@ -194,27 +188,38 @@ bool FVSAccessorModule::OpenVisualStudioSolution()
 	bool bSuccess = false;
 
 	CComPtr<EnvDTE::_DTE> DTE;
-	if (AccessVisualStudio(DTE, SolutionPath, Locations))
+	switch (AccessVisualStudioViaDTE(DTE, SolutionPath, Locations))
 	{
-		// Set Focus on Visual Studio
-		CComPtr<EnvDTE::Window> MainWindow;
-		if (SUCCEEDED(DTE->get_MainWindow(&MainWindow)) &&
-			SUCCEEDED(MainWindow->Activate()))
-		{
-			bSuccess = true;
-		}
-		else
+	case EAccessVisualStudioResult::VSInstanceIsOpen:
 		{
-			UE_LOG(LogVSAccessor, Warning, TEXT("Couldn't set focus on Visual Studio."));
+			// Set Focus on Visual Studio
+			CComPtr<EnvDTE::Window> MainWindow;
+			if (SUCCEEDED(DTE->get_MainWindow(&MainWindow)) &&
+				SUCCEEDED(MainWindow->Activate()))
+			{
+				bSuccess = true;
+			}
+			else
+			{
+				UE_LOG(LogVSAccessor, Warning, TEXT("Couldn't set focus on Visual Studio."));
+			}
 		}
-	}
-	else
-	{
-		// Automatically fail if there's already an attempt in progress
-		if (!IsVSLaunchInProgress())
+		break;
+
+	case EAccessVisualStudioResult::VSInstanceIsNotOpen:
 		{
-			bSuccess = RunVisualStudioAndOpenSolution();
+			// Automatically fail if there's already an attempt in progress
+			if (!IsVSLaunchInProgress())
+			{
+				bSuccess = RunVisualStudioAndOpenSolution();
+			}
 		}
+		break;
+
+	default:
+		// Do nothing if we failed the VS detection, otherwise we could get stuck in a loop of constantly 
+		// trying to open a VS instance since we can't detect that one is already running
+		break;
 	}
 
 	// Uninitialize the com library, if we initialized it above (don't call if S_FALSE)
@@ -223,15 +228,7 @@ bool FVSAccessorModule::OpenVisualStudioSolution()
 	return bSuccess;
 }
 
-bool FVSAccessorModule::OpenVisualStudioFileAtLineInternal(const FString& FullPath, int32 LineNumber, int32 ColumnNumber)
-{
-	TArray<FileOpenRequest> Requests;
-	Requests.Add(FileOpenRequest(FullPath, LineNumber, ColumnNumber));
-
-	return OpenVisualStudioFilesInternal(Requests);
-}
-
-bool FVSAccessorModule::OpenVisualStudioFilesInternal(const TArray<FileOpenRequest>& Requests)
+bool FVSAccessorModule::OpenVisualStudioFilesInternalViaDTE(const TArray<FileOpenRequest>& Requests, bool& bWasDeferred)
 {
 	// Initialize the com library, if not already by this thread
 	if (!FWindowsPlatformMisc::CoInitialize())
@@ -242,97 +239,109 @@ bool FVSAccessorModule::OpenVisualStudioFilesInternal(const TArray<FileOpenReque
 	
 	bool bDefer = false, bSuccess = false;
 	CComPtr<EnvDTE::_DTE> DTE;
-	if (AccessVisualStudio(DTE, SolutionPath, Locations))
+	switch (AccessVisualStudioViaDTE(DTE, SolutionPath, Locations))
 	{
-		// Set Focus on Visual Studio
-		CComPtr<EnvDTE::Window> MainWindow;
-		if (SUCCEEDED(DTE->get_MainWindow(&MainWindow)) &&
-			SUCCEEDED(MainWindow->Activate()))
+	case EAccessVisualStudioResult::VSInstanceIsOpen:
 		{
-			// Get ItemOperations
-			CComPtr<EnvDTE::ItemOperations> ItemOperations;
-			if (SUCCEEDED(DTE->get_ItemOperations(&ItemOperations)))
+			// Set Focus on Visual Studio
+			CComPtr<EnvDTE::Window> MainWindow;
+			if (SUCCEEDED(DTE->get_MainWindow(&MainWindow)) &&
+				SUCCEEDED(MainWindow->Activate()))
 			{
-				for ( const FileOpenRequest& Request : Requests )
+				// Get ItemOperations
+				CComPtr<EnvDTE::ItemOperations> ItemOperations;
+				if (SUCCEEDED(DTE->get_ItemOperations(&ItemOperations)))
 				{
-					// Check that the file actually exists first
-					if ( !FPaths::FileExists(Request.FullPath) )
+					for ( const FileOpenRequest& Request : Requests )
 					{
-						OpenFileFailed.Broadcast(Request.FullPath);
-						bSuccess |= false;
-						continue;
-					}
+						// Check that the file actually exists first
+						if ( !FPaths::FileExists(Request.FullPath) )
+						{
+							OpenFileFailed.Broadcast(Request.FullPath);
+							bSuccess |= false;
+							continue;
+						}
 
-					// Open File
-					auto ANSIPath = StringCast<ANSICHAR>(*Request.FullPath);
-					CComBSTR COMStrFileName(ANSIPath.Get());
-					CComBSTR COMStrKind(EnvDTE::vsViewKindTextView);
-					CComPtr<EnvDTE::Window> Window;
-					if ( SUCCEEDED(ItemOperations->OpenFile(COMStrFileName, COMStrKind, &Window)) )
-					{
-						// If we've made it this far - we've opened the file.  it doesn't matter if
-						// we successfully get to the line number.  Everything else is gravy.
-						bSuccess |= true;
-
-						// Scroll to Line Number
-						CComPtr<EnvDTE::Document> Document;
-						CComPtr<IDispatch> SelectionDispatch;
-						CComPtr<EnvDTE::TextSelection> Selection;
-						if ( SUCCEEDED(DTE->get_ActiveDocument(&Document)) &&
-							 SUCCEEDED(Document->get_Selection(&SelectionDispatch)) &&
-							 SUCCEEDED(SelectionDispatch->QueryInterface(&Selection)) &&
-							 SUCCEEDED(Selection->GotoLine(Request.LineNumber, true)) )
+						// Open File
+						auto ANSIPath = StringCast<ANSICHAR>(*Request.FullPath);
+						CComBSTR COMStrFileName(ANSIPath.Get());
+						CComBSTR COMStrKind(EnvDTE::vsViewKindTextView);
+						CComPtr<EnvDTE::Window> Window;
+						if ( SUCCEEDED(ItemOperations->OpenFile(COMStrFileName, COMStrKind, &Window)) )
 						{
-							if ( !SUCCEEDED(Selection->MoveToLineAndOffset(Request.LineNumber, Request.ColumnNumber, false)) )
+							// If we've made it this far - we've opened the file.  it doesn't matter if
+							// we successfully get to the line number.  Everything else is gravy.
+							bSuccess |= true;
+
+							// Scroll to Line Number
+							CComPtr<EnvDTE::Document> Document;
+							CComPtr<IDispatch> SelectionDispatch;
+							CComPtr<EnvDTE::TextSelection> Selection;
+							if ( SUCCEEDED(DTE->get_ActiveDocument(&Document)) &&
+								 SUCCEEDED(Document->get_Selection(&SelectionDispatch)) &&
+								 SUCCEEDED(SelectionDispatch->QueryInterface(&Selection)) &&
+								 SUCCEEDED(Selection->GotoLine(Request.LineNumber, true)) )
 							{
-								UE_LOG(LogVSAccessor, Warning, TEXT("Couldn't goto column number '%i' of line '%i' in '%s'"), Request.ColumnNumber, Request.LineNumber, *Request.FullPath);
+								if ( !SUCCEEDED(Selection->MoveToLineAndOffset(Request.LineNumber, Request.ColumnNumber, false)) )
+								{
+									UE_LOG(LogVSAccessor, Warning, TEXT("Couldn't goto column number '%i' of line '%i' in '%s'"), Request.ColumnNumber, Request.LineNumber, *Request.FullPath);
+								}
+							}
+							else
+							{
+								UE_LOG(LogVSAccessor, Warning, TEXT("Couldn't goto line number '%i' in '%s'"), Request.LineNumber, *Request.FullPath);
 							}
 						}
 						else
 						{
-							UE_LOG(LogVSAccessor, Warning, TEXT("Couldn't goto line number '%i' in '%s'"), Request.LineNumber, *Request.FullPath);
+							UE_LOG(LogVSAccessor, Warning, TEXT("Couldn't open file '%s'."), *Request.FullPath);
 						}
 					}
-					else
-					{
-						UE_LOG(LogVSAccessor, Warning, TEXT("Couldn't open file '%s'."), *Request.FullPath);
-					}
-				}
 
-				VSLaunchFinished( true );
+					VSLaunchFinished( true );
+				}
+				else
+				{
+					UE_LOG(LogVSAccessor, Log, TEXT("Couldn't get item operations. Visual Studio may still be initializing."));
+					bDefer = true;
+				}
 			}
 			else
 			{
-				UE_LOG(LogVSAccessor, Log, TEXT("Couldn't get item operations. Visual Studio may still be initializing."));
-				bDefer = true;
+				UE_LOG(LogVSAccessor, Warning, TEXT("Couldn't set focus on Visual Studio."));
 			}
 		}
-		else
-		{
-			UE_LOG(LogVSAccessor, Warning, TEXT("Couldn't set focus on Visual Studio."));
-		}
-	}
-	else
-	{
-		bDefer = true;
+		break;
 
-		// We can't process until we're in the main thread, if we aren't initially defer until we are
-		if ( IsInGameThread() )
+	case EAccessVisualStudioResult::VSInstanceIsNotOpen:
 		{
-			// If we haven't already attempted to launch VS do so now
-			if ( !IsVSLaunchInProgress() )
+			bDefer = true;
+
+			// We can't process until we're in the main thread, if we aren't initially defer until we are
+			if ( IsInGameThread() )
 			{
-				// If there's no valid instance of VS running, run one if we have it installed
-				if ( !RunVisualStudioAndOpenSolution() )
-				{
-					bDefer = false;
-				}
-				else
+				// If we haven't already attempted to launch VS do so now
+				if ( !IsVSLaunchInProgress() )
 				{
-					VSLaunchStarted();
+					// If there's no valid instance of VS running, run one if we have it installed
+					if ( !RunVisualStudioAndOpenSolution() )
+					{
+						bDefer = false;
+					}
+					else
+					{
+						VSLaunchStarted();
+					}
 				}
 			}
 		}
+		break;
+
+	default:
+		// Do nothing if we failed the VS detection, otherwise we could get stuck in a loop of constantly 
+		// trying to open a VS instance since we can't detect that one is already running
+		bDefer = false;
+		break;
 	}
 	
 	if ( !bSuccess )
@@ -369,6 +378,7 @@ bool FVSAccessorModule::OpenVisualStudioFilesInternal(const TArray<FileOpenReque
 	// Uninitialize the com library, if we initialized it above (don't call if S_FALSE)
 	FWindowsPlatformMisc::CoUninitialize();
 
+	bWasDeferred = bDefer;
 	return bSuccess;
 }
 
@@ -384,7 +394,7 @@ bool FVSAccessorModule::SaveAllOpenVisualStudioDocuments()
 	}
 	
 	CComPtr<EnvDTE::_DTE> DTE;
-	if (AccessVisualStudio(DTE, SolutionPath, Locations))
+	if (AccessVisualStudioViaDTE(DTE, SolutionPath, Locations) == EAccessVisualStudioResult::VSInstanceIsOpen)
 	{
 		// Save all documents
 		CComPtr<EnvDTE::Documents> Documents;
@@ -455,6 +465,14 @@ void FVSAccessorModule::VSLaunchFinished( bool bSuccess )
 
 #else
 
+// VS Express-only dummy versions
+bool FVSAccessorModule::SaveAllOpenVisualStudioDocuments() { return false; }
+bool FVSAccessorModule::Exec( UWorld* InWorld, const TCHAR* Cmd, FOutputDevice& Ar ) { return false; }
+void FVSAccessorModule::VSLaunchStarted() {}
+void FVSAccessorModule::VSLaunchFinished( bool bSuccess ) {}
+
+#endif
+
 bool GetProcessCommandLine(const ::DWORD InProcessID, FString& OutCommandLine)
 {
 	check(InProcessID);
@@ -560,10 +578,12 @@ bool GetProcessCommandLine(const ::DWORD InProcessID, FString& OutCommandLine)
 	return Data.OutHwnd;
 }
 
-bool AccessVisualStudio(::DWORD& OutProcessID, FString& OutExecutablePath, const FString& InSolutionPath, const TArray<FVSAccessorModule::VisualStudioLocation>& InLocations)
+EAccessVisualStudioResult AccessVisualStudioViaProcess(::DWORD& OutProcessID, FString& OutExecutablePath, const FString& InSolutionPath, const TArray<FVSAccessorModule::VisualStudioLocation>& InLocations)
 {
 	OutProcessID = 0;
 
+	EAccessVisualStudioResult AccessResult = EAccessVisualStudioResult::VSInstanceIsNotOpen;
+
 	::HANDLE hProcessSnap = ::CreateToolhelp32Snapshot(TH32CS_SNAPPROCESS, 0);
 	if (hProcessSnap != INVALID_HANDLE_VALUE)
 	{
@@ -601,87 +621,126 @@ bool AccessVisualStudio(::DWORD& OutProcessID, FString& OutExecutablePath, const
 								{
 									OutProcessID = ProcEntry.th32ProcessID;
 									OutExecutablePath = Location.ExecutablePath;
+									AccessResult = EAccessVisualStudioResult::VSInstanceIsOpen;
 									break;
 								}
 							}
+							else
+							{
+								UE_LOG(LogVSAccessor, Warning, TEXT("Couldn't access module information"));
+								AccessResult = EAccessVisualStudioResult::VSInstanceUnknown;
+							}
 						}
 					}
+					else
+					{
+						UE_LOG(LogVSAccessor, Warning, TEXT("Couldn't access module table"));
+						AccessResult = EAccessVisualStudioResult::VSInstanceUnknown;
+					}
 
 					::CloseHandle(hModuleSnap);
 				}
+				else
+				{
+					UE_LOG(LogVSAccessor, Warning, TEXT("Couldn't access module table"));
+					AccessResult = EAccessVisualStudioResult::VSInstanceUnknown;
+				}
 			}
 		}
 
 		::CloseHandle(hProcessSnap);
 	}
+	else
+	{
+		UE_LOG(LogVSAccessor, Warning, TEXT("Couldn't access process table"));
+		AccessResult = EAccessVisualStudioResult::VSInstanceUnknown;
+	}
 
-	return OutProcessID != 0;
+	return AccessResult;
 }
 
 bool FVSAccessorModule::OpenVisualStudioSolution()
 {
-	::DWORD ProcessID = 0;
-	FString Path;
-	if (AccessVisualStudio(ProcessID, Path, SolutionPath, Locations))
+#if VSACCESSOR_HAS_DTE
+	if(OpenVisualStudioSolutionViaDTE())
 	{
-		// Try to bring Visual Studio to the foreground
-		::HWND VisualStudioHwnd = GetTopWindowForProcess(ProcessID);
-		if (VisualStudioHwnd)
-		{
-			// SwitchToThisWindow isn't really intended for general use, however it can switch to 
-			// the VS window, where SetForegroundWindow will fail due to process permissions
-			::SwitchToThisWindow(VisualStudioHwnd, 0);
-		}
 		return true;
 	}
-
-	return RunVisualStudioAndOpenSolution();
+#endif
+	return OpenVisualStudioSolutionViaProcess();
 }
 
-/** Opens multiple files in the correct running instance of Visual Studio. */
 bool FVSAccessorModule::OpenVisualStudioFilesInternal(const TArray<FileOpenRequest>& Requests)
 {
-	bool bSuccess = false;
-	for ( const FileOpenRequest& Request : Requests )
+#if VSACCESSOR_HAS_DTE
 	{
-		bSuccess |= OpenVisualStudioFileAtLineInternal(Request.FullPath, Request.LineNumber, Request.ColumnNumber);
+		bool bWasDeferred = false;
+		if(OpenVisualStudioFilesInternalViaDTE(Requests, bWasDeferred) || bWasDeferred)
+		{
+			return true;
+		}
 	}
-
-	return bSuccess;
+#endif
+	return OpenVisualStudioFilesInternalViaProcess(Requests);
 }
 
-bool FVSAccessorModule::OpenVisualStudioFileAtLineInternal(const FString& FullPath, int32 LineNumber, int32 ColumnNumber)
+bool FVSAccessorModule::OpenVisualStudioSolutionViaProcess()
 {
-	// Check that the file actually exists first
-	if (!FPaths::FileExists(FullPath))
+	::DWORD ProcessID = 0;
+	FString Path;
+	switch (AccessVisualStudioViaProcess(ProcessID, Path, SolutionPath, Locations))
 	{
-		OpenFileFailed.Broadcast(FullPath);
-		return false;
+	case EAccessVisualStudioResult::VSInstanceIsOpen:
+		{
+			// Try to bring Visual Studio to the foreground
+			::HWND VisualStudioHwnd = GetTopWindowForProcess(ProcessID);
+			if (VisualStudioHwnd)
+			{
+				// SwitchToThisWindow isn't really intended for general use, however it can switch to 
+				// the VS window, where SetForegroundWindow will fail due to process permissions
+				::SwitchToThisWindow(VisualStudioHwnd, 0);
+			}
+			return true;
+		}
+		break;
+
+	case EAccessVisualStudioResult::VSInstanceIsNotOpen:
+		return RunVisualStudioAndOpenSolution();
+		
+	default:
+		// Do nothing if we failed the VS detection, otherwise we could get stuck in a loop of constantly 
+		// trying to open a VS instance since we can't detect that one is already running
+		break;
 	}
 
+	return false;
+}
+
+bool FVSAccessorModule::OpenVisualStudioFilesInternalViaProcess(const TArray<FileOpenRequest>& Requests)
+{
 	::DWORD ProcessID = 0;
 	FString Path;
-	if (AccessVisualStudio(ProcessID, Path, SolutionPath, Locations))
-	{
-		return RunVisualStudioAndOpenSolutionAndFile(Path, "", FullPath, LineNumber);
-	}
-	
-	if (CanRunVisualStudio(Path))
+	switch (AccessVisualStudioViaProcess(ProcessID, Path, SolutionPath, Locations))
 	{
-		return RunVisualStudioAndOpenSolutionAndFile(Path, SolutionPath, FullPath, LineNumber);
+	case EAccessVisualStudioResult::VSInstanceIsOpen:
+		return RunVisualStudioAndOpenSolutionAndFiles(Path, "", &Requests);
+
+	case EAccessVisualStudioResult::VSInstanceIsNotOpen:
+		if (CanRunVisualStudio(Path))
+		{
+			return RunVisualStudioAndOpenSolutionAndFiles(Path, SolutionPath, &Requests);
+		}
+		break;
+
+	default:
+		// Do nothing if we failed the VS detection, otherwise we could get stuck in a loop of constantly 
+		// trying to open a VS instance since we can't detect that one is already running
+		break;
 	}
 
 	return false;
 }
 
-// VS Express-only dummy versions
-bool FVSAccessorModule::SaveAllOpenVisualStudioDocuments() { return false; }
-bool FVSAccessorModule::Exec( UWorld* InWorld, const TCHAR* Cmd, FOutputDevice& Ar ) { return false; }
-void FVSAccessorModule::VSLaunchStarted() {}
-void FVSAccessorModule::VSLaunchFinished( bool bSuccess ) {}
-
-#endif
-
 bool FVSAccessorModule::CanRunVisualStudio(FString& OutPath) const
 {
 	// AddVisualStudioVersion only adds paths to Locations if they're valid, so we can just use the first entry in the array
@@ -700,7 +759,7 @@ bool FVSAccessorModule::RunVisualStudioAndOpenSolution() const
 	FString Path;
 	if (CanRunVisualStudio(Path))
 	{
-		return RunVisualStudioAndOpenSolutionAndFile(Path, SolutionPath, "", 0);
+		return RunVisualStudioAndOpenSolutionAndFiles(Path, SolutionPath, nullptr);
 	}
 
 	return false;
@@ -734,6 +793,14 @@ bool FVSAccessorModule::OpenVisualStudioFileAtLine(const FString& FullPath, int3
 	return false;
 }
 
+bool FVSAccessorModule::OpenVisualStudioFileAtLineInternal(const FString& FullPath, int32 LineNumber, int32 ColumnNumber)
+{
+	TArray<FileOpenRequest> Requests;
+	Requests.Add(FileOpenRequest(FullPath, LineNumber, ColumnNumber));
+
+	return OpenVisualStudioFilesInternal(Requests);
+}
+
 void FVSAccessorModule::AddVisualStudioVersion(const int MajorVersion, const bool bAllowExpress)
 {
 	// We can use the common tools macro to work out if this version of Visual Studio is installed
@@ -777,3 +844,65 @@ void FVSAccessorModule::AddVisualStudioVersion(const int MajorVersion, const boo
 		}
 	}
 }
+
+bool FVSAccessorModule::RunVisualStudioAndOpenSolutionAndFiles(const FString& ExecutablePath, const FString& SolutionPath, const TArray<FileOpenRequest>* const Requests) const
+{
+	FString Params;
+
+	// Only open the solution if it exists
+	if (!SolutionPath.IsEmpty())
+	{
+		if (FPaths::FileExists(SolutionPath))
+		{
+			Params += TEXT("\"");
+			Params += SolutionPath;
+			Params += TEXT("\"");
+		}
+		else
+		{
+			OpenFileFailed.Broadcast(SolutionPath);
+			return false;
+		}
+	}
+
+	if (Requests)
+	{
+		int32 GoToLine = 0;
+		for (const FileOpenRequest& Request : *Requests)
+		{
+			// Only open the file if it exists
+			if (FPaths::FileExists(Request.FullPath))
+			{
+				Params += TEXT(" \"");
+				Params += Request.FullPath;
+				Params += TEXT("\"");
+
+				GoToLine = Request.LineNumber;
+			}
+			else
+			{
+				OpenFileFailed.Broadcast(Request.FullPath);
+				return false;
+			}
+		}
+
+		if (GoToLine > 0)
+		{
+			Params += FString::Printf(TEXT(" /command \"edit.goto %d\""), GoToLine);
+		}
+	}
+
+	bool bSuccess = false;
+
+	if (!Params.IsEmpty())
+	{
+		FProcHandle WorkerHandle = FPlatformProcess::CreateProc(*ExecutablePath, *Params, true, false, false, nullptr, 0, nullptr, nullptr);
+		if (WorkerHandle.IsValid())
+		{
+			bSuccess = true;
+		}
+		WorkerHandle.Close();
+	}
+
+	return bSuccess;
+}
diff --git a/Engine/Source/Developer/Windows/VSAccessor/Public/VSAccessorModule.h b/Engine/Source/Developer/Windows/VSAccessor/Public/VSAccessorModule.h
--- a/Engine/Source/Developer/Windows/VSAccessor/Public/VSAccessorModule.h
+++ b/Engine/Source/Developer/Windows/VSAccessor/Public/VSAccessorModule.h
@@ -133,6 +133,24 @@ private:
 	 */
 	void AddVisualStudioVersion(const int MajorVersion, const bool bAllowExpress = true);
 
+	/** 
+	 * Run a new instance Visual Studio, optionally opening the provided solution and list of files
+	 * 
+	 * @param	ExecutablePath	Path to the Visual Studio executable to open
+	 * @param	SolutionPath	Path to the solution to open, or an empty string to open no solution
+	 * @param	Requests		Array of files to open, or null to open no files
+	 */
+	bool RunVisualStudioAndOpenSolutionAndFiles(const FString& ExecutablePath, const FString& SolutionPath, const TArray<FileOpenRequest>* const Requests) const;
+
+#if VSACCESSOR_HAS_DTE
+	/** DTE specific implementations */
+	bool OpenVisualStudioSolutionViaDTE();
+	bool OpenVisualStudioFilesInternalViaDTE(const TArray<FileOpenRequest>& Requests, bool& bWasDeferred);
+#endif
+	/** Fallback (non-DTE) implementations */
+	bool OpenVisualStudioSolutionViaProcess();
+	bool OpenVisualStudioFilesInternalViaProcess(const TArray<FileOpenRequest>& Requests);
+
 private:
 
 	/** The versions of VS we support, in preference order */
