Removed unused GetBuildArtifacts() from ITargetPlatform (this is now being done in UAT).

diff --git a/Engine/Source/Developer/Android/AndroidTargetPlatform/Private/AndroidTargetPlatform.h b/Engine/Source/Developer/Android/AndroidTargetPlatform/Private/AndroidTargetPlatform.h
--- a/Engine/Source/Developer/Android/AndroidTargetPlatform/Private/AndroidTargetPlatform.h
+++ b/Engine/Source/Developer/Android/AndroidTargetPlatform/Private/AndroidTargetPlatform.h
@@ -84,8 +84,6 @@ public:
 
 	virtual ECompressionFlags GetBaseCompressionMethod( ) const OVERRIDE;
 
-	virtual bool GetBuildArtifacts( const FString& ProjectPath, EBuildTargets::Type BuildTarget, EBuildConfigurations::Type BuildConfiguration, ETargetPlatformBuildArtifacts::Type Artifacts, TMap<FString, FString>& OutFiles, TArray<FString>& OutMissingFiles ) const OVERRIDE;
-
 	virtual bool GenerateStreamingInstallManifest(const TMultiMap<FString, int32>& ChunkMap, const TSet<int32>& ChunkIDsInUse) const OVERRIDE
 	{
 		return true;
diff --git a/Engine/Source/Developer/Android/AndroidTargetPlatform/Private/AndroidTargetPlatform.inl b/Engine/Source/Developer/Android/AndroidTargetPlatform/Private/AndroidTargetPlatform.inl
--- a/Engine/Source/Developer/Android/AndroidTargetPlatform/Private/AndroidTargetPlatform.inl
+++ b/Engine/Source/Developer/Android/AndroidTargetPlatform/Private/AndroidTargetPlatform.inl
@@ -64,13 +64,6 @@ inline ITargetDevicePtr FAndroidTargetPlatform::GetDefaultDevice( ) const
 }
 
 
-inline bool FAndroidTargetPlatform::GetBuildArtifacts( const FString& ProjectPath, EBuildTargets::Type BuildTarget, EBuildConfigurations::Type BuildConfiguration, ETargetPlatformBuildArtifacts::Type Artifacts, TMap<FString, FString>& OutFiles, TArray<FString>& OutMissingFiles ) const
-{
-	//@todo android: implement build artifacts gathering
-	return false;
-}
-
-
 inline ITargetDevicePtr FAndroidTargetPlatform::GetDevice( const FTargetDeviceId& DeviceId )
 {
 	if (DeviceId.GetPlatformName() == PlatformName())
diff --git a/Engine/Source/Developer/HTML5/HTML5TargetPlatform/Private/HTML5TargetPlatform.cpp b/Engine/Source/Developer/HTML5/HTML5TargetPlatform/Private/HTML5TargetPlatform.cpp
--- a/Engine/Source/Developer/HTML5/HTML5TargetPlatform/Private/HTML5TargetPlatform.cpp
+++ b/Engine/Source/Developer/HTML5/HTML5TargetPlatform/Private/HTML5TargetPlatform.cpp
@@ -44,95 +44,6 @@ ECompressionFlags FHTML5TargetPlatform::GetBaseCompressionMethod( ) const
 }
 
 
-bool FHTML5TargetPlatform::GetBuildArtifacts( const FString& ProjectPath, EBuildTargets::Type BuildTarget, EBuildConfigurations::Type BuildConfiguration, ETargetPlatformBuildArtifacts::Type Artifacts, TMap<FString, FString>& OutFiles, TArray<FString>& OutMissingFiles ) const
-{
-	if ((BuildTarget != EBuildTargets::Game) || (BuildConfiguration == EBuildConfigurations::Unknown))
-	{
-		return false;
-	}
-
-	const FString ProjectName = FPaths::GetBaseFilename(ProjectPath);
-	const FString ProjectRoot = FPaths::GetPath(ProjectPath);
-
-	if (!ProjectName.IsEmpty() || !ProjectRoot.IsEmpty())
-	{
-		return false;
-	}
-
-	// append binaries
-	if (Artifacts & ETargetPlatformBuildArtifacts::Engine)
-	{
-		const FString EngineBinariesDir = FPaths::EngineDir() / TEXT("Binaries") / PlatformName();
-		const FString DeployedBinariesDir = FString(TEXT("Engine/Binaries")) / PlatformName();
-
-		// executable
-		{
-			FString ExecutableName;
-			FString ProjectBinariesDir = ProjectRoot / TEXT("Binaries");
-
-			// content-only projects use game agnostic binaries
-			if (IFileManager::Get().DirectoryExists(*ProjectBinariesDir))
-			{
-				ExecutableName = ProjectName;
-			}
-			else
-			{
-				ExecutableName = TEXT("UE4");
-				ProjectBinariesDir = EngineBinariesDir;
-			}
-
-			if (BuildConfiguration != EBuildConfigurations::Development)
-			{
-				ExecutableName += FString::Printf(TEXT("-HTML5-%s"), EBuildConfigurations::ToString(BuildConfiguration));
-			}
-
-			ExecutableName += TEXT(".js");
-			ProjectBinariesDir /= PlatformName();
-
-			// add executable
-			{
-				const FString ExecutablePath = ProjectBinariesDir / ExecutableName;
-
-				if (FPaths::FileExists(*ExecutablePath))
-				{
-					OutFiles.Add(ExecutablePath, DeployedBinariesDir / ExecutableName);
-				}
-				else
-				{
-					OutMissingFiles.Add(ExecutablePath);
-				}
-
-				const FString MemFileName = ExecutableName + TEXT(".mem");
-				const FString MemFilePath = ProjectBinariesDir / MemFileName;
-
-				if (FPaths::FileExists(*MemFilePath))
-				{
-					OutFiles.Add(MemFilePath, DeployedBinariesDir / MemFileName);
-				}
-				else
-				{
-					OutMissingFiles.Add(MemFilePath);
-				}
-			}
-		}
-	}
-
-	// append engine content
-	if (Artifacts & ETargetPlatformBuildArtifacts::Content)
-	{
-		// @todo: handle Engine configuration files and content
-	}
-
-	// append game content
-	if (Artifacts & ETargetPlatformBuildArtifacts::Content)
-	{
-		// @todo: handle game configuration files and content
-	}
-
-	return (OutMissingFiles.Num() == 0);
-}
-
-
 ITargetDevicePtr FHTML5TargetPlatform::GetDefaultDevice( ) const
 {
 	return LocalDevice;
diff --git a/Engine/Source/Developer/HTML5/HTML5TargetPlatform/Private/HTML5TargetPlatform.h b/Engine/Source/Developer/HTML5/HTML5TargetPlatform/Private/HTML5TargetPlatform.h
--- a/Engine/Source/Developer/HTML5/HTML5TargetPlatform/Private/HTML5TargetPlatform.h
+++ b/Engine/Source/Developer/HTML5/HTML5TargetPlatform/Private/HTML5TargetPlatform.h
@@ -30,8 +30,6 @@ public:
 
 	virtual ECompressionFlags GetBaseCompressionMethod( ) const OVERRIDE;
 
-	virtual bool GetBuildArtifacts( const FString& ProjectPath, EBuildTargets::Type BuildTarget, EBuildConfigurations::Type BuildConfiguration, ETargetPlatformBuildArtifacts::Type Artifacts, TMap<FString, FString>& OutFile, TArray<FString>& OutMissingFiless ) const;
-
 	virtual bool GenerateStreamingInstallManifest(const TMultiMap<FString, int32>& ChunkMap, const TSet<int32>& ChunkIDsInUse) const OVERRIDE
 	{
 		return true;
diff --git a/Engine/Source/Developer/Linux/LinuxTargetPlatform/Private/LinuxTargetPlatform.h b/Engine/Source/Developer/Linux/LinuxTargetPlatform/Private/LinuxTargetPlatform.h
--- a/Engine/Source/Developer/Linux/LinuxTargetPlatform/Private/LinuxTargetPlatform.h
+++ b/Engine/Source/Developer/Linux/LinuxTargetPlatform/Private/LinuxTargetPlatform.h
@@ -51,122 +51,6 @@ public:
 		return COMPRESS_ZLIB;
 	}
 
-	virtual bool GetBuildArtifacts( const FString& ProjectPath, EBuildTargets::Type BuildTarget, EBuildConfigurations::Type BuildConfiguration, ETargetPlatformBuildArtifacts::Type Artifacts, TMap<FString, FString>& OutFiles, TArray<FString>& OutMissingFiles ) const OVERRIDE
-	{
-		if ((BuildTarget == EBuildTargets::Unknown) || (BuildConfiguration == EBuildConfigurations::Unknown))
-		{
-			return false;
-		}
-
-		const FString ProjectName = FPaths::GetBaseFilename(ProjectPath);
-		const FString ProjectRoot = FPaths::GetPath(ProjectPath);
-
-		if (!ProjectName.IsEmpty() || !ProjectRoot.IsEmpty())
-		{
-			return false;
-		}
-
-		// append binaries
-		if (Artifacts & ETargetPlatformBuildArtifacts::Engine)
-		{
-			const FString EngineBinariesDir = FPaths::EngineDir() / TEXT("Binaries") / TEXT("Linux");
-			const FString DeployedBinariesDir = FString(TEXT("Engine")) / TEXT("Binaries") / TEXT("Linux");
-
-			// executable
-			{
-				FString ExecutableBaseName;
-				FString ProjectBinariesDir = ProjectRoot / TEXT("Binaries");
-
-				// content-only projects use game agnostic binaries
-				if (IFileManager::Get().DirectoryExists(*ProjectBinariesDir))
-				{
-					ExecutableBaseName = ProjectName;
-				}
-				else
-				{
-					ExecutableBaseName = TEXT("UE4");
-					ProjectBinariesDir = EngineBinariesDir;
-				}
-	
-				if (BuildConfiguration != EBuildConfigurations::Development)
-				{
-					ExecutableBaseName += FString::Printf(TEXT("-%s-%s"), *PlatformName(), EBuildConfigurations::ToString(BuildConfiguration));
-				}
-
-				ProjectBinariesDir /= TEXT("Linux");
-
-				// add executable
-				{
-					FString ExecutablePath = ProjectBinariesDir / ExecutableBaseName;
-
-					if (FPaths::FileExists(*ExecutablePath))
-					{
-						OutFiles.Add(ExecutablePath, DeployedBinariesDir / ExecutableBaseName);
-					}
-					else
-					{
-						OutMissingFiles.Add(ExecutablePath);
-					}
-				}
-			}
-
-			// dynamic link libraries
-			{
-				// @todo platforms: this needs some major improvement, so that only required files are included
-				TArray<FString> FileNames;
-
-				IFileManager::Get().FindFiles(FileNames, *(EngineBinariesDir / TEXT("*.so")), true, false);
-
-				for (int32 FileIndex = 0; FileIndex < FileNames.Num(); ++FileIndex)
-				{
-					// add library
-					const FString& FileName = FileNames[FileIndex];
-					OutFiles.Add(EngineBinariesDir / FileName, DeployedBinariesDir / FileName);
-				}
-			}
-
-			// append third party binaries
-			{
-				FString ThirdPartyBinariesDir = FPaths::EngineDir() / TEXT("Binaries") / TEXT("ThirdParty");
-				TArray<FString> FileNames;
-
-				IFileManager::Get().FindFilesRecursive(FileNames, *ThirdPartyBinariesDir, TEXT("*.*"), true, false);
-
-				for (int32 FileIndex = 0; FileIndex < FileNames.Num(); ++FileIndex)
-				{
-					FString FileName = FileNames[FileIndex];
-
-					if (FileName.EndsWith(TEXT(".so")) && !FileName.Contains(TEXT("Mcpp")) && !FileName.Contains(TEXT("NoRedist")) )
-					{
-						OutFiles.Add(FileName, FileName.RightChop(9));
-					}
-				}
-			}
-		}
-
-		// append engine content
-		if (Artifacts & ETargetPlatformBuildArtifacts::Content)
-		{
-			// @todo platforms: handle Engine configuration files and content
-		}
-
-		// append game content
-		if (Artifacts & ETargetPlatformBuildArtifacts::Content)
-		{
-			// @todo platforms: handle game configuration files and content
-		}
-
-		// append tools
-		if (Artifacts & ETargetPlatformBuildArtifacts::Tools)
-		{
-			// @todo platforms: add tools deployment support
-			return false;
-		}
-
-		return (OutMissingFiles.Num() == 0);
-	}
-
-
 	virtual bool GenerateStreamingInstallManifest(const TMultiMap<FString, int32>& ChunkMap, const TSet<int32>& ChunkIDsInUse) const OVERRIDE
 	{
 		return true;
diff --git a/Engine/Source/Developer/Mac/MacTargetPlatform/Private/GenericMacTargetPlatform.h b/Engine/Source/Developer/Mac/MacTargetPlatform/Private/GenericMacTargetPlatform.h
--- a/Engine/Source/Developer/Mac/MacTargetPlatform/Private/GenericMacTargetPlatform.h
+++ b/Engine/Source/Developer/Mac/MacTargetPlatform/Private/GenericMacTargetPlatform.h
@@ -47,152 +47,6 @@ public:
 		return COMPRESS_ZLIB;
 	}
 
-	virtual bool GetBuildArtifacts( const FString& ProjectPath, EBuildTargets::Type BuildTarget, EBuildConfigurations::Type BuildConfiguration, ETargetPlatformBuildArtifacts::Type Artifacts, TMap<FString, FString>& OutFiles, TArray<FString>& OutMissingFiles ) const OVERRIDE
-	{
-		if ((BuildTarget == EBuildTargets::Unknown) || (BuildConfiguration == EBuildConfigurations::Unknown))
-		{
-			return false;
-		}
-
-		const FString ProjectName = FPaths::GetBaseFilename(ProjectPath);
-		const FString ProjectRoot = FPaths::GetPath(ProjectPath);
-
-		if (!ProjectName.IsEmpty() || !ProjectRoot.IsEmpty())
-		{
-			return false;
-		}
-
-		// append binaries
-		if (Artifacts & ETargetPlatformBuildArtifacts::Engine)
-		{
-			const FString EngineBinariesDir = FPaths::EngineDir() / TEXT("Binaries") / this->PlatformName();
-			const FString DeployedBinariesDir = FString(TEXT("Engine")) / TEXT("Binaries") / this->PlatformName();
-
-			// executable
-			{
-				FString ExecutableBaseName;
-				FString ProjectBinariesDir = ProjectRoot / TEXT("Binaries");
-
-				// content-only projects use game agnostic binaries
-				if (IFileManager::Get().DirectoryExists(*ProjectBinariesDir))
-				{
-					ExecutableBaseName = ProjectName;
-				}
-				else
-				{
-					ExecutableBaseName = TEXT("UE4");
-					ProjectBinariesDir = EngineBinariesDir;
-				}
-		
-				if (BuildConfiguration != EBuildConfigurations::Development)
-				{
-					ExecutableBaseName += FString::Printf(TEXT("-%s-%s.exe"), *this->PlatformName(), EBuildConfigurations::ToString(BuildConfiguration));
-				}
-
-				ProjectBinariesDir /= this->PlatformName();
-
-				// add executable
-				{
-					FString ExecutablePath = ProjectBinariesDir / ExecutableBaseName + TEXT(".exe");
-
-					if (FPaths::FileExists(*ExecutablePath))
-					{
-						OutFiles.Add(ExecutablePath, DeployedBinariesDir / ExecutableBaseName + TEXT(".exe"));
-					}
-					else
-					{
-						OutMissingFiles.Add(ExecutablePath);
-					}
-				}
-
-				// add debug symbols
-				if (Artifacts & ETargetPlatformBuildArtifacts::DebugSymbols)
-				{
-					FString DebugSymbolsPath = ProjectBinariesDir / ExecutableBaseName + TEXT(".pdb");
-
-					if (FPaths::FileExists(*DebugSymbolsPath))
-					{
-						OutFiles.Add(DebugSymbolsPath, DeployedBinariesDir / ExecutableBaseName + TEXT(".pdb"));
-					}
-					else
-					{
-						OutMissingFiles.Add(DebugSymbolsPath);
-					}
-				}
-			}
-
-			// dynamic link libraries
-			{
-				// @todo platforms: this needs some major improvement, so that only required files are included
-				TArray<FString> FileNames;
-
-				IFileManager::Get().FindFiles(FileNames, *(EngineBinariesDir / TEXT("*.dylib")), true, false);
-
-				for (int32 FileIndex = 0; FileIndex < FileNames.Num(); ++FileIndex)
-				{
-					// add library
-					const FString& FileName = FileNames[FileIndex];
-					OutFiles.Add(EngineBinariesDir / FileName, DeployedBinariesDir / FileName);
-
-					// add debug symbols
-					if (Artifacts & ETargetPlatformBuildArtifacts::DebugSymbols)
-					{
-						FString DebugSymbolFile = FPaths::GetBaseFilename(FileName) + TEXT(".pdb");
-						FString DebugySymbolPath = EngineBinariesDir / DebugSymbolFile;
-
-						if (FPaths::FileExists(DebugySymbolPath))
-						{
-							OutFiles.Add(DebugySymbolPath, DeployedBinariesDir / DebugSymbolFile);
-						}
-						else
-						{
-							OutMissingFiles.Add(DebugySymbolPath);
-						}
-					}
-				}
-			}
-
-			// append third party binaries
-			{
-				FString ThirdPartyBinariesDir = FPaths::EngineDir() / TEXT("Binaries") / TEXT("ThirdParty");
-				TArray<FString> FileNames;
-
-				IFileManager::Get().FindFilesRecursive(FileNames, *ThirdPartyBinariesDir, TEXT("*.*"), true, false);
-
-				for (int32 FileIndex = 0; FileIndex < FileNames.Num(); ++FileIndex)
-				{
-					FString FileName = FileNames[FileIndex];
-
-					if (FileName.EndsWith(TEXT(".dylib")) && !FileName.Contains(TEXT("Mcpp")) && !FileName.Contains(TEXT("NoRedist")) )
-					{
-						OutFiles.Add(FileName, FileName.RightChop(9));
-					}
-				}
-			}
-		}
-
-		// append engine content
-		if (Artifacts & ETargetPlatformBuildArtifacts::Content)
-		{
-			// @todo platforms: handle Engine configuration files and content
-		}
-
-		// append game content
-		if (Artifacts & ETargetPlatformBuildArtifacts::Content)
-		{
-			// @todo platforms: handle game configuration files and content
-		}
-
-		// append tools
-		if (Artifacts & ETargetPlatformBuildArtifacts::Tools)
-		{
-			// @todo platforms: add tools deployment support
-			return false;
-		}
-
-		return (OutMissingFiles.Num() == 0);
-	}
-
 	virtual bool GenerateStreamingInstallManifest(const TMultiMap<FString, int32>& ChunkMap, const TSet<int32>& ChunkIDsInUse) const OVERRIDE
 	{
 		return true;
diff --git a/Engine/Source/Developer/TargetPlatform/Public/Interfaces/ITargetPlatform.h b/Engine/Source/Developer/TargetPlatform/Public/Interfaces/ITargetPlatform.h
--- a/Engine/Source/Developer/TargetPlatform/Public/Interfaces/ITargetPlatform.h
+++ b/Engine/Source/Developer/TargetPlatform/Public/Interfaces/ITargetPlatform.h
@@ -179,23 +179,6 @@ public:
 	 */
 	virtual ECompressionFlags GetBaseCompressionMethod( ) const = 0;
 
-	/**
-	 * Gets the paths of the binaries that need to be deployed for the specified game and build configuration.
-	 *
-	 * The returned collection maps the path of local build artifacts (including the artifact's file name)
-	 * to a deployment folder path on a target device. Both paths are relative to the Root.
-	 *
-	 * @param ProjectPath - The path to the Unreal project file.
-	 * @param BuildTarget - The build target to get the artifacts for, i.e Game or Server.
-	 * @param BuildConfiguration - The build configuration to get the artifacts for, i.e. Debug or Shipping.
-	 * @param Artifacts - The type of build artifacts to include.
-	 * @param OutFiles - Will hold the collection of artifacts.
-	 * @param OutMissingFiles - Will hold the paths to missing artifacts.
-	 *
-	 * @return true if all required artifacts were gathered, false otherwise.
-	 */
-	virtual bool GetBuildArtifacts( const FString& ProjectPath, EBuildTargets::Type BuildTarget, EBuildConfigurations::Type BuildConfiguration, ETargetPlatformBuildArtifacts::Type Artifacts, TMap<FString, FString>& OutFiles, TArray<FString>& OutMissingFiles ) const = 0;
-
 	/** 
 	 * Generates a platform specific asset manifest given an array of FAssetData.
 	 *
diff --git a/Engine/Source/Developer/Windows/WindowsTargetPlatform/Private/GenericWindowsTargetPlatform.h b/Engine/Source/Developer/Windows/WindowsTargetPlatform/Private/GenericWindowsTargetPlatform.h
--- a/Engine/Source/Developer/Windows/WindowsTargetPlatform/Private/GenericWindowsTargetPlatform.h
+++ b/Engine/Source/Developer/Windows/WindowsTargetPlatform/Private/GenericWindowsTargetPlatform.h
@@ -49,161 +49,6 @@ public:
 		return COMPRESS_ZLIB;
 	}
 
-	virtual bool GetBuildArtifacts( const FString& ProjectPath, EBuildTargets::Type BuildTarget, EBuildConfigurations::Type BuildConfiguration, ETargetPlatformBuildArtifacts::Type Artifacts, TMap<FString, FString>& OutFiles, TArray<FString>& OutMissingFiles ) const OVERRIDE
-	{
-		if ((BuildTarget == EBuildTargets::Unknown) || (BuildConfiguration == EBuildConfigurations::Unknown))
-		{
-			return false;
-		}
-
-		const FString ProjectName = FPaths::GetBaseFilename(ProjectPath);
-		const FString ProjectRoot = FPaths::GetPath(ProjectPath);
-
-		if (ProjectName.IsEmpty() || ProjectRoot.IsEmpty())
-		{
-			return false;
-		}
-
-		// append binaries
-		if (Artifacts & ETargetPlatformBuildArtifacts::Engine)
-		{
-			const FString EngineBinariesDir = FPaths::EngineDir() / TEXT("Binaries") / GetBinariesSubDir();
-			const FString DeployedBinariesDir = FString(TEXT("Engine")) / TEXT("Binaries") / GetBinariesSubDir();
-
-			// executable
-			{
-				FString ExecutableBaseName;
-				FString ProjectBinariesDir = ProjectRoot / TEXT("Binaries") / GetBinariesSubDir();
-
-				// content-only projects use game agnostic binaries
-				if (IFileManager::Get().DirectoryExists(*ProjectBinariesDir))
-				{
-					ExecutableBaseName = ProjectName;
-				}
-				else
-				{
-					if (BuildTarget == EBuildTargets::Game)
-					{
-						ExecutableBaseName = TEXT("UE4Game");
-					}
-					else if (BuildTarget == EBuildTargets::Server)
-					{
-						ExecutableBaseName = TEXT("UE4Game");
-					}
-					else if (BuildTarget == EBuildTargets::Editor)
-					{
-						ExecutableBaseName = TEXT("UE4Editor");
-					}
-					ProjectBinariesDir = EngineBinariesDir;
-				}
-		
-				if (BuildConfiguration != EBuildConfigurations::Development)
-				{
-					ExecutableBaseName += FString::Printf(TEXT("-%s-%s.exe"), *PlatformName(), EBuildConfigurations::ToString(BuildConfiguration));
-				}
-
-				// add executable
-				{
-					FString ExecutablePath = ProjectBinariesDir / ExecutableBaseName + TEXT(".exe");
-
-					if (FPaths::FileExists(*ExecutablePath))
-					{
-						OutFiles.Add(ExecutablePath, DeployedBinariesDir / ExecutableBaseName + TEXT(".exe"));
-					}
-					else
-					{
-						OutMissingFiles.Add(ExecutablePath);
-					}
-				}
-
-				// add debug symbols
-				if (Artifacts & ETargetPlatformBuildArtifacts::DebugSymbols)
-				{
-					FString DebugSymbolsPath = ProjectBinariesDir / ExecutableBaseName + TEXT(".pdb");
-
-					if (FPaths::FileExists(*DebugSymbolsPath))
-					{
-						OutFiles.Add(DebugSymbolsPath, DeployedBinariesDir / ExecutableBaseName + TEXT(".pdb"));
-					}
-					else
-					{
-						OutMissingFiles.Add(DebugSymbolsPath);
-					}
-				}
-			}
-
-			// dynamic link libraries
-			{
-				// @todo platforms: this needs some major improvement, so that only required files are included
-				TArray<FString> FileNames;
-
-				IFileManager::Get().FindFiles(FileNames, *(EngineBinariesDir / TEXT("*.dll")), true, false);
-
-				for (int32 FileIndex = 0; FileIndex < FileNames.Num(); ++FileIndex)
-				{
-					// add library
-					const FString& FileName = FileNames[FileIndex];
-					OutFiles.Add(EngineBinariesDir / FileName, DeployedBinariesDir / FileName);
-
-					// add debug symbols
-					if (Artifacts & ETargetPlatformBuildArtifacts::DebugSymbols)
-					{
-						FString DebugSymbolFile = FPaths::GetBaseFilename(FileName) + TEXT(".pdb");
-						FString DebugySymbolPath = EngineBinariesDir / DebugSymbolFile;
-
-						if (FPaths::FileExists(DebugySymbolPath))
-						{
-							OutFiles.Add(DebugySymbolPath, DeployedBinariesDir / DebugSymbolFile);
-						}
-						else
-						{
-							OutMissingFiles.Add(DebugySymbolPath);
-						}
-					}
-				}
-			}
-
-			// append third party binaries
-			{
-				FString ThirdPartyBinariesDir = FPaths::EngineDir() / TEXT("Binaries") / TEXT("ThirdParty");
-				TArray<FString> FileNames;
-
-				IFileManager::Get().FindFilesRecursive(FileNames, *ThirdPartyBinariesDir, TEXT("*.*"), true, false);
-
-				for (int32 FileIndex = 0; FileIndex < FileNames.Num(); ++FileIndex)
-				{
-					FString FileName = FileNames[FileIndex];
-
-					if (FileName.EndsWith(TEXT(".dll")) && !FileName.Contains(TEXT("Mcpp")) && !FileName.Contains(TEXT("NoRedist")) )
-					{
-						OutFiles.Add(FileName, FileName.RightChop(9));
-					}
-				}
-			}
-		}
-
-		// append engine content
-		if (Artifacts & ETargetPlatformBuildArtifacts::Content)
-		{
-			// @todo platforms: handle Engine configuration files and content
-		}
-
-		// append game content
-		if (Artifacts & ETargetPlatformBuildArtifacts::Content)
-		{
-			// @todo platforms: handle game configuration files and content
-		}
-
-		// append tools
-		if (Artifacts & ETargetPlatformBuildArtifacts::Tools)
-		{
-			// @todo platforms: add tools deployment support
-			return false;
-		}
-
-		return (OutMissingFiles.Num() == 0);
-	}
-
 	virtual bool GenerateStreamingInstallManifest(const TMultiMap<FString, int32>& ChunkMap, const TSet<int32>& ChunkIDsInUse) const OVERRIDE
 	{
 		return true;
diff --git a/Engine/Source/Developer/iOS/IOSTargetPlatform/Private/IOSTargetPlatform.cpp b/Engine/Source/Developer/iOS/IOSTargetPlatform/Private/IOSTargetPlatform.cpp
--- a/Engine/Source/Developer/iOS/IOSTargetPlatform/Private/IOSTargetPlatform.cpp
+++ b/Engine/Source/Developer/iOS/IOSTargetPlatform/Private/IOSTargetPlatform.cpp
@@ -54,30 +54,6 @@ void FIOSTargetPlatform::GetAllDevices( TArray<ITargetDevicePtr>& OutDevices ) c
 }
 
 
-bool FIOSTargetPlatform::GetBuildArtifacts( const FString& ProjectPath, EBuildTargets::Type BuildTarget, EBuildConfigurations::Type Config, ETargetPlatformBuildArtifacts::Type Artifacts, TMap<FString, FString>& OutFiles, TArray<FString>& OutMissingFiles ) const
-{
-	/*
-	//Using this as the signal that deployment will happen.
-	//Gather all necessary information and store it to the Target Device, such that, the device will save the state for deployment.
-	ITargetDevicePtr CurrentDevice = NULL;
-	TSharedPtr<FIOSTargetDevice> CurrentIOSDevice = NULL;
-
-	auto Iter = Devices.CreateConstIterator();
-
-	while(Iter)
-	{
-		CurrentDevice = Iter.Value();
-		CurrentIOSDevice = StaticCastSharedPtr<FIOSTargetDevice, ITargetDevice, ESPMode::Fast>(CurrentDevice);
-		CurrentIOSDevice->SetAppId(GameName);
-		CurrentIOSDevice->SetAppConfiguration(Config);
-		++Iter;
-	}
-	*/
-
-	return false;
-}
-
-
 ITargetDevicePtr FIOSTargetPlatform::GetDefaultDevice() const
 {
 	if (Devices.Num() > 0)
diff --git a/Engine/Source/Developer/iOS/IOSTargetPlatform/Private/IOSTargetPlatform.h b/Engine/Source/Developer/iOS/IOSTargetPlatform/Private/IOSTargetPlatform.h
--- a/Engine/Source/Developer/iOS/IOSTargetPlatform/Private/IOSTargetPlatform.h
+++ b/Engine/Source/Developer/iOS/IOSTargetPlatform/Private/IOSTargetPlatform.h
@@ -51,8 +51,6 @@ public:
 		return COMPRESS_ZLIB;
 	}
 
-	virtual bool GetBuildArtifacts( const FString& ProjectPath, EBuildTargets::Type BuildTarget, EBuildConfigurations::Type BuildConfiguration, ETargetPlatformBuildArtifacts::Type Artifacts, TMap<FString, FString>& OutFiles, TArray<FString>& OutMissingFiles ) const OVERRIDE;
-
 	virtual bool GenerateStreamingInstallManifest(const TMultiMap<FString, int32>& ChunkMap, const TSet<int32>& ChunkIDsInUse) const OVERRIDE
 	{
 		return true;
