Disabling glsl shader format for windows for now. there are a few issues with cooking to resolve first. A proper fix is incoming.

diff --git a/Engine/Source/Developer/Windows/WindowsTargetPlatform/Private/GenericWindowsTargetPlatform.h b/Engine/Source/Developer/Windows/WindowsTargetPlatform/Private/GenericWindowsTargetPlatform.h
--- a/Engine/Source/Developer/Windows/WindowsTargetPlatform/Private/GenericWindowsTargetPlatform.h
+++ b/Engine/Source/Developer/Windows/WindowsTargetPlatform/Private/GenericWindowsTargetPlatform.h
@@ -139,11 +139,9 @@ public:
 			//       may not wish to support SM4 or OpenGL.
 			static FName NAME_PCD3D_SM5(TEXT("PCD3D_SM5"));
 			static FName NAME_PCD3D_SM4(TEXT("PCD3D_SM4"));
-			static FName NAME_GLSL_150(TEXT("GLSL_150"));
 
 			OutFormats.AddUnique(NAME_PCD3D_SM5);
 			OutFormats.AddUnique(NAME_PCD3D_SM4);
-			OutFormats.AddUnique(NAME_GLSL_150);
 		}
 	}
 
