Fix CIS

diff --git a/Engine/Source/Developer/Windows/WindowsTargetPlatform/Private/GenericWindowsTargetPlatform.h b/Engine/Source/Developer/Windows/WindowsTargetPlatform/Private/GenericWindowsTargetPlatform.h
--- a/Engine/Source/Developer/Windows/WindowsTargetPlatform/Private/GenericWindowsTargetPlatform.h
+++ b/Engine/Source/Developer/Windows/WindowsTargetPlatform/Private/GenericWindowsTargetPlatform.h
@@ -6,7 +6,9 @@
 
 #pragma once
 
+#if WITH_ENGINE
 #include "StaticMeshResources.h"
+#endif // WITH_ENGINE
 
 /**
  * Template for Windows target platforms
diff --git a/Engine/Source/Developer/iOS/IOSTargetPlatform/Private/IOSTargetPlatform.h b/Engine/Source/Developer/iOS/IOSTargetPlatform/Private/IOSTargetPlatform.h
--- a/Engine/Source/Developer/iOS/IOSTargetPlatform/Private/IOSTargetPlatform.h
+++ b/Engine/Source/Developer/iOS/IOSTargetPlatform/Private/IOSTargetPlatform.h
@@ -7,7 +7,10 @@
 #pragma once
 
 #include "Ticker.h"
+
+#if WITH_ENGINE
 #include "StaticMeshResources.h"
+#endif // WITH_ENGINE
 
 /**
  * FIOSTargetPlatform, abstraction for cooking iOS platforms
diff --git a/Engine/Source/Runtime/Engine/Public/SceneManagement.h b/Engine/Source/Runtime/Engine/Public/SceneManagement.h
--- a/Engine/Source/Runtime/Engine/Public/SceneManagement.h
+++ b/Engine/Source/Runtime/Engine/Public/SceneManagement.h
@@ -15,7 +15,7 @@
 #include "Engine/TextureLightProfile.h"
 #include "Engine/World.h"
 #include "RendererInterface.h"
-#include "GameViewportClient.h"
+#include "Engine/GameViewportClient.h"
 
 // Forward declarations.
 class FSceneViewFamily;
