r.DumpShaderDebugInfo=1 now outputs a batch file to run AMD CodeXL (if installed), GPUShaderAnalyzer commandline removed

diff --git a/Engine/Source/Developer/Windows/ShaderFormatD3D/Private/D3D11ShaderCompiler.cpp b/Engine/Source/Developer/Windows/ShaderFormatD3D/Private/D3D11ShaderCompiler.cpp
--- a/Engine/Source/Developer/Windows/ShaderFormatD3D/Private/D3D11ShaderCompiler.cpp
+++ b/Engine/Source/Developer/Windows/ShaderFormatD3D/Private/D3D11ShaderCompiler.cpp
@@ -204,29 +204,22 @@ static FString D3D11CreateShaderCompileCommandLine(
 }
 
 /** Creates a batch file string to call the AMD shader analyzer. */
-static FString CreateAMDShaderAnalyzerCommandLine(
+static FString CreateAMDCodeXLCommandLine(
 	const FString& ShaderPath, 
 	const TCHAR* EntryFunction, 
-	const TCHAR* ShaderProfile
+	const TCHAR* ShaderProfile,
+	uint32 DXFlags
 	)
 {
 	// Hardcoded to the default install path since there's no Env variable or addition to PATH
-	FString CommandlineBegin = FString(TEXT("\"C:\\Program Files (x86)\\AMD\\GPU ShaderAnalyzer\\GPUShaderAnalyzer.exe\" ")) 
-		+ ShaderPath
-		+ TEXT(" -Function ")
-		+ EntryFunction;
-
-	FString CommandlineEnd = FString(TEXT(" -Profile "))
-		+ ShaderProfile
-		// Latest version of GPUShaderAnalyzer (1.59) doesn't support 7970 yet =(
-		+ TEXT(" -ASIC HD6970\r\n");
-
-	FString Commandline;
-
-	// Call GPUShaderAnalyzer to generate microcode
-	Commandline += CommandlineBegin + TEXT(" -Disassemble") + CommandlineEnd;
-	// Call GPUShaderAnalyzer to analyze 
-	Commandline += CommandlineBegin + TEXT(" -Analyze") + CommandlineEnd;
+	FString Commandline = FString(TEXT("\"C:\\Program Files (x86)\\AMD\\CodeXL\\CodeXLAnalyzer.exe\" -c Pitcairn")) 
+		+ TEXT(" -f ") + EntryFunction
+		+ TEXT(" -s HLSL")
+		+ TEXT(" -p ") + ShaderProfile
+		+ TEXT(" -a AnalyzerStats.csv")
+		+ TEXT(" --isa ISA.txt")
+		+ *FString::Printf(TEXT(" --DXFlags %u "), DXFlags)
+		+ ShaderPath;
 
 	// add a pause on a newline
 	Commandline += FString(TEXT(" \r\n pause"));
@@ -462,7 +455,7 @@ void CompileD3D11Shader(const FShaderCompilerInput& Input,FShaderCompilerOutput&
 			const FString BatchFileContents = D3D11CreateShaderCompileCommandLine((Input.SourceFilename + TEXT(".usf")), *Input.EntryPointName, ShaderProfile, CompileFlags);
 			FFileHelper::SaveStringToFile(BatchFileContents, *(Input.DumpDebugInfoPath / TEXT("CompileD3D.bat")));
 
-			const FString BatchFileContents2 = CreateAMDShaderAnalyzerCommandLine((Input.SourceFilename + TEXT(".usf")), *Input.EntryPointName, ShaderProfile);
+			const FString BatchFileContents2 = CreateAMDCodeXLCommandLine((Input.SourceFilename + TEXT(".usf")), *Input.EntryPointName, ShaderProfile, CompileFlags);
 			FFileHelper::SaveStringToFile(BatchFileContents2, *(Input.DumpDebugInfoPath / TEXT("CompileAMD.bat")));
 		}
 
