UE4 - Improve D3D compile times

diff --git a/Engine/Shaders/ShaderVersion.usf b/Engine/Shaders/ShaderVersion.usf
--- a/Engine/Shaders/ShaderVersion.usf
+++ b/Engine/Shaders/ShaderVersion.usf
@@ -2,4 +2,4 @@
 // Each time the console command is executed it generates a new GUID. As this file is included
 // in common.usf (which should be included in any shader) it allows to invalidate the shader DDC.
 // 
-// GUID = A7702884495089DC221506B8C4961E6
\ No newline at end of file
+// GUID = {61298DFD-DB18-4EED-AA6F-4B5C802866AC}
diff --git a/Engine/Source/Developer/Windows/ShaderFormatD3D/Private/D3D11ShaderCompiler.cpp b/Engine/Source/Developer/Windows/ShaderFormatD3D/Private/D3D11ShaderCompiler.cpp
--- a/Engine/Source/Developer/Windows/ShaderFormatD3D/Private/D3D11ShaderCompiler.cpp
+++ b/Engine/Source/Developer/Windows/ShaderFormatD3D/Private/D3D11ShaderCompiler.cpp
@@ -329,6 +329,11 @@ void CompileD3D11Shader(const FShaderCompilerInput& Input,FShaderCompilerOutput&
 		return;
 	}
 
+	if (!RemoveUniformBuffersFromSource(PreprocessedShaderSource))
+	{
+		return;
+	}
+
 	// Search definitions for a custom D3D compiler path.
 	for(TMap<FString,FString>::TConstIterator DefinitionIt(Input.Environment.GetDefinitions());DefinitionIt;++DefinitionIt)
 	{
