TargetPlatform: Removed file header comments; changed NULL to nullptr

diff --git a/Engine/Source/Developer/Windows/WindowsClientTargetPlatform/Private/WindowsClientTargetPlatformModule.cpp b/Engine/Source/Developer/Windows/WindowsClientTargetPlatform/Private/WindowsClientTargetPlatformModule.cpp
--- a/Engine/Source/Developer/Windows/WindowsClientTargetPlatform/Private/WindowsClientTargetPlatformModule.cpp
+++ b/Engine/Source/Developer/Windows/WindowsClientTargetPlatform/Private/WindowsClientTargetPlatformModule.cpp
@@ -1,16 +1,10 @@
 // Copyright 1998-2014 Epic Games, Inc. All Rights Reserved.
 
-/*=============================================================================
-	WindowsClientTargetPlatformModule.cpp: Implements the FWindowsClientTargetPlatformModule class.
-=============================================================================*/
-
 #include "WindowsClientTargetPlatformPrivatePCH.h"
 
 
-/**
- * Holds the target platform singleton.
- */
-static ITargetPlatform* Singleton = NULL;
+/** Holds the target platform singleton. */
+static ITargetPlatform* Singleton = nullptr;
 
 
 /**
@@ -23,12 +17,12 @@ public:
 
 	virtual ~FWindowsClientTargetPlatformModule( )
 	{
-		Singleton = NULL;
+		Singleton = nullptr;
 	}
 
 	virtual ITargetPlatform* GetTargetPlatform( )
 	{
-		if (Singleton == NULL)
+		if (Singleton == nullptr)
 		{
 			Singleton = new TGenericWindowsTargetPlatform<false, false, true>();
 		}
@@ -38,4 +32,4 @@ public:
 };
 
 
-IMPLEMENT_MODULE(FWindowsClientTargetPlatformModule, WindowsClientTargetPlatform);
\ No newline at end of file
+IMPLEMENT_MODULE(FWindowsClientTargetPlatformModule, WindowsClientTargetPlatform);
diff --git a/Engine/Source/Developer/Windows/WindowsNoEditorTargetPlatform/Private/WindowsNoEditorTargetPlatformModule.cpp b/Engine/Source/Developer/Windows/WindowsNoEditorTargetPlatform/Private/WindowsNoEditorTargetPlatformModule.cpp
--- a/Engine/Source/Developer/Windows/WindowsNoEditorTargetPlatform/Private/WindowsNoEditorTargetPlatformModule.cpp
+++ b/Engine/Source/Developer/Windows/WindowsNoEditorTargetPlatform/Private/WindowsNoEditorTargetPlatformModule.cpp
@@ -1,16 +1,10 @@
 // Copyright 1998-2014 Epic Games, Inc. All Rights Reserved.
 
-/*=============================================================================
-	WindowsNoEditorTargetPlatformModule.cpp: Implements the FWindowsNoEditorTargetPlatformModule class.
-=============================================================================*/
-
 #include "WindowsNoEditorTargetPlatformPrivatePCH.h"
 
 
-/**
- * Holds the target platform singleton.
- */
-static ITargetPlatform* Singleton = NULL;
+/** Holds the target platform singleton. */
+static ITargetPlatform* Singleton = nullptr;
 
 
 /**
@@ -23,12 +17,12 @@ public:
 
 	virtual ~FWindowsNoEditorTargetPlatformModule( )
 	{
-		Singleton = NULL;
+		Singleton = nullptr;
 	}
 
 	virtual ITargetPlatform* GetTargetPlatform( )
 	{
-		if (Singleton == NULL)
+		if (Singleton == nullptr)
 		{
 			Singleton = new TGenericWindowsTargetPlatform<false, false, false>();
 		}
@@ -38,4 +32,4 @@ public:
 };
 
 
-IMPLEMENT_MODULE(FWindowsNoEditorTargetPlatformModule, WindowsNoEditorTargetPlatform);
\ No newline at end of file
+IMPLEMENT_MODULE(FWindowsNoEditorTargetPlatformModule, WindowsNoEditorTargetPlatform);
diff --git a/Engine/Source/Developer/Windows/WindowsServerTargetPlatform/Private/WindowsServerTargetPlatformModule.cpp b/Engine/Source/Developer/Windows/WindowsServerTargetPlatform/Private/WindowsServerTargetPlatformModule.cpp
--- a/Engine/Source/Developer/Windows/WindowsServerTargetPlatform/Private/WindowsServerTargetPlatformModule.cpp
+++ b/Engine/Source/Developer/Windows/WindowsServerTargetPlatform/Private/WindowsServerTargetPlatformModule.cpp
@@ -1,16 +1,10 @@
 // Copyright 1998-2014 Epic Games, Inc. All Rights Reserved.
 
-/*=============================================================================
-	WindowsServerTargetPlatformModule.cpp: Implements the FWindowsServerTargetPlatformModule class.
-=============================================================================*/
-
 #include "WindowsServerTargetPlatformPrivatePCH.h"
 
 
-/**
- * Holds the target platform singleton.
- */
-static ITargetPlatform* Singleton = NULL;
+/** Holds the target platform singleton. */
+static ITargetPlatform* Singleton = nullptr;
 
 
 /**
@@ -23,12 +17,12 @@ public:
 
 	virtual ~FWindowsServerTargetPlatformModule( )
 	{
-		Singleton = NULL;
+		Singleton = nullptr;
 	}
 
 	virtual ITargetPlatform* GetTargetPlatform( )
 	{
-		if (Singleton == NULL)
+		if (Singleton == nullptr)
 		{
 			Singleton = new TGenericWindowsTargetPlatform<false, true, false>();
 		}
@@ -38,4 +32,4 @@ public:
 };
 
 
-IMPLEMENT_MODULE(FWindowsServerTargetPlatformModule, WindowsServerTargetPlatform);
\ No newline at end of file
+IMPLEMENT_MODULE(FWindowsServerTargetPlatformModule, WindowsServerTargetPlatform);
diff --git a/Engine/Source/Developer/Windows/WindowsServerTargetPlatform/Private/WindowsServerTargetPlatformPrivatePCH.h b/Engine/Source/Developer/Windows/WindowsServerTargetPlatform/Private/WindowsServerTargetPlatformPrivatePCH.h
--- a/Engine/Source/Developer/Windows/WindowsServerTargetPlatform/Private/WindowsServerTargetPlatformPrivatePCH.h
+++ b/Engine/Source/Developer/Windows/WindowsServerTargetPlatform/Private/WindowsServerTargetPlatformPrivatePCH.h
@@ -1,13 +1,9 @@
 // Copyright 1998-2014 Epic Games, Inc. All Rights Reserved.
 
-/*=============================================================================
-	WindowsServerTargetPlatformPrivatePCH.h: Pre-compiled header file for the WindowsServerargetPlatform module.
-=============================================================================*/
-
 #pragma once
 
 
-/* Dependencies
+/* Private dependencies
  *****************************************************************************/
 
 #include "Core.h"
diff --git a/Engine/Source/Developer/Windows/WindowsTargetPlatform/Private/GenericWindowsTargetPlatform.h b/Engine/Source/Developer/Windows/WindowsTargetPlatform/Private/GenericWindowsTargetPlatform.h
--- a/Engine/Source/Developer/Windows/WindowsTargetPlatform/Private/GenericWindowsTargetPlatform.h
+++ b/Engine/Source/Developer/Windows/WindowsTargetPlatform/Private/GenericWindowsTargetPlatform.h
@@ -1,15 +1,13 @@
 // Copyright 1998-2014 Epic Games, Inc. All Rights Reserved.
 
-/*=============================================================================
-	GenericWindowsTargetPlatform.h: Declares the TGenericWindowsTargetPlatform class template.
-=============================================================================*/
-
 #pragma once
 
+
 #if WITH_ENGINE
-#include "StaticMeshResources.h"
+	#include "StaticMeshResources.h"
 #endif // WITH_ENGINE
 
+
 /**
  * Template for Windows target platforms
  */
@@ -85,7 +83,7 @@ public:
 			return LocalDevice;
 		}
 
-		return NULL;
+		return nullptr;
 	}
 
 	virtual bool IsRunningPlatform( ) const override
diff --git a/Engine/Source/Developer/Windows/WindowsTargetPlatform/Private/LocalPcTargetDevice.h b/Engine/Source/Developer/Windows/WindowsTargetPlatform/Private/LocalPcTargetDevice.h
--- a/Engine/Source/Developer/Windows/WindowsTargetPlatform/Private/LocalPcTargetDevice.h
+++ b/Engine/Source/Developer/Windows/WindowsTargetPlatform/Private/LocalPcTargetDevice.h
@@ -1,9 +1,5 @@
 // Copyright 1998-2014 Epic Games, Inc. All Rights Reserved.
 
-/*=============================================================================
-	LocalPcTargetDevice.h: Declares the TLocalPcTargetDevice class template.
-=============================================================================*/
-
 #pragma once
 
 
@@ -29,6 +25,8 @@ public:
 
 public:
 
+	// ITargetDevice interface
+
 	virtual bool Connect( ) override
 	{
 		return true;
@@ -258,7 +256,7 @@ public:
 		ExecutablePath += TEXT(".exe");
 
 		// launch the game
-		FProcHandle ProcessHandle = FPlatformProcess::CreateProc(*ExecutablePath, *Params, true, false, false, OutProcessId, 0, NULL, NULL);
+		FProcHandle ProcessHandle = FPlatformProcess::CreateProc(*ExecutablePath, *Params, true, false, false, OutProcessId, 0, nullptr, nullptr);
 		return ProcessHandle.Close();
 	}
 
@@ -289,7 +287,7 @@ public:
 
 	virtual bool Run( const FString& ExecutablePath, const FString& Params, uint32* OutProcessId )
 	{
-		FProcHandle ProcessHandle = FPlatformProcess::CreateProc(*ExecutablePath, *Params, true, false, false, OutProcessId, 0, NULL, NULL);
+		FProcHandle ProcessHandle = FPlatformProcess::CreateProc(*ExecutablePath, *Params, true, false, false, OutProcessId, 0, nullptr, nullptr);
 		return ProcessHandle.Close();
 	}
 
diff --git a/Engine/Source/Developer/Windows/WindowsTargetPlatform/Private/WindowsTargetPlatformClasses.cpp b/Engine/Source/Developer/Windows/WindowsTargetPlatform/Private/WindowsTargetPlatformClasses.cpp
--- a/Engine/Source/Developer/Windows/WindowsTargetPlatform/Private/WindowsTargetPlatformClasses.cpp
+++ b/Engine/Source/Developer/Windows/WindowsTargetPlatform/Private/WindowsTargetPlatformClasses.cpp
@@ -1,16 +1,11 @@
 // Copyright 1998-2014 Epic Games, Inc. All Rights Reserved.
 
-/*=============================================================================
-	WindowsTargetPlatformClasses.cpp: Implements the module's UClasses.
-=============================================================================*/
-
 #include "WindowsTargetPlatformPrivatePCH.h"
 
 
-
 /* UWindowsTargetSettings structors
  *****************************************************************************/
 
 UWindowsTargetSettings::UWindowsTargetSettings( const class FPostConstructInitializeProperties& PCIP )
 	: Super(PCIP)
-{ }
\ No newline at end of file
+{ }
diff --git a/Engine/Source/Developer/Windows/WindowsTargetPlatform/Private/WindowsTargetPlatformModule.cpp b/Engine/Source/Developer/Windows/WindowsTargetPlatform/Private/WindowsTargetPlatformModule.cpp
--- a/Engine/Source/Developer/Windows/WindowsTargetPlatform/Private/WindowsTargetPlatformModule.cpp
+++ b/Engine/Source/Developer/Windows/WindowsTargetPlatform/Private/WindowsTargetPlatformModule.cpp
@@ -1,9 +1,5 @@
 // Copyright 1998-2014 Epic Games, Inc. All Rights Reserved.
 
-/*=============================================================================
-	WindowsTargetPlatformModule.cpp: Implements the FWindowsTargetPlatformModule class.
-=============================================================================*/
-
 #include "WindowsTargetPlatformPrivatePCH.h"
 
 
@@ -13,7 +9,7 @@
 /**
  * Holds the target platform singleton.
  */
-static ITargetPlatform* Singleton = NULL;
+static ITargetPlatform* Singleton = nullptr;
 
 
 /**
@@ -24,22 +20,33 @@ class FWindowsTargetPlatformModule
 {
 public:
 
-	/**
-	 * Destructor.
-	 */
+	/** Destructor. */
 	~FWindowsTargetPlatformModule( )
 	{
-		Singleton = NULL;
+		Singleton = nullptr;
 	}
 
+public:
+
+	// this is an example of a hotfix, declared here for no particular reason. Once we have other examples, it can be deleted.
+#if 0
+	void HotfixTest( void *InPayload, int PayloadSize )
+	{
+		check(sizeof(FTestHotFixPayload) == PayloadSize);
+		
+		FTestHotFixPayload* Payload = (FTestHotFixPayload*)InPayload;
+		UE_LOG(LogTemp, Log, TEXT("Hotfix Test %s"), *Payload->Message);
+		Payload->Result = Payload->ValueToReturn;
+	}
+#endif
 
 public:
 
-	// Begin ITargetPlatformModule interface
+	// ITargetPlatformModule interface
 
 	virtual ITargetPlatform* GetTargetPlatform( ) override
 	{
-		if (Singleton == NULL)
+		if (Singleton == nullptr)
 		{
 			Singleton = new TGenericWindowsTargetPlatform<true, false, false>();
 		}
@@ -47,23 +54,9 @@ public:
 		return Singleton;
 	}
 
-	// End ITargetPlatformModule interface
-
-
 public:
 
-	// this is an example of a hotfix, declared here for no particular reason. Once we have other examples, it can be deleted.
-#if 0
-	void HotfixTest(void *InPayload, int PayloadSize)
-	{
-		check(sizeof(FTestHotFixPayload) == PayloadSize);
-		
-		FTestHotFixPayload* Payload = (FTestHotFixPayload*)InPayload;
-		UE_LOG(LogTemp, Log, TEXT("Hotfix Test %s"), *Payload->Message);
-		Payload->Result = Payload->ValueToReturn;
-	}
-#endif
-	// Begin IModuleInterface interface
+	// IModuleInterface interface
 
 	virtual void StartupModule() override
 	{
@@ -111,13 +104,10 @@ public:
 		}
 		else
 		{
-			TargetSettings = NULL;
+			TargetSettings = nullptr;
 		}
 	}
 
-	// End IModuleInterface interface
-
-
 private:
 
 	// Holds the target settings.
diff --git a/Engine/Source/Developer/Windows/WindowsTargetPlatform/Private/WindowsTargetPlatformPrivatePCH.h b/Engine/Source/Developer/Windows/WindowsTargetPlatform/Private/WindowsTargetPlatformPrivatePCH.h
--- a/Engine/Source/Developer/Windows/WindowsTargetPlatform/Private/WindowsTargetPlatformPrivatePCH.h
+++ b/Engine/Source/Developer/Windows/WindowsTargetPlatform/Private/WindowsTargetPlatformPrivatePCH.h
@@ -1,13 +1,9 @@
 // Copyright 1998-2014 Epic Games, Inc. All Rights Reserved.
 
-/*=============================================================================
-	WindowsTargetPlatformPrivatePCH.h: Pre-compiled header file for the WindowsTargetPlatform module.
-=============================================================================*/
-
 #pragma once
 
 
-/* Dependencies
+/* Private dependencies
  *****************************************************************************/
 
 #include "Core.h"
@@ -27,10 +23,10 @@
 	#include <TlHelp32.h>
 #include "HideWindowsPlatformTypes.h"
 
+
 /* Private includes
  *****************************************************************************/
 
 #include "WindowsTargetSettings.h"
 #include "LocalPcTargetDevice.h"
 #include "GenericWindowsTargetPlatform.h"
-
