Added r.StripShaderDebugData and r.OptimizeShaders cvars, useful to disable when running

diff --git a/Engine/Source/Developer/Windows/ShaderFormatD3D/Private/D3D11ShaderCompiler.cpp b/Engine/Source/Developer/Windows/ShaderFormatD3D/Private/D3D11ShaderCompiler.cpp
--- a/Engine/Source/Developer/Windows/ShaderFormatD3D/Private/D3D11ShaderCompiler.cpp
+++ b/Engine/Source/Developer/Windows/ShaderFormatD3D/Private/D3D11ShaderCompiler.cpp
@@ -38,7 +38,6 @@ static uint32 TranslateCompilerFlagD3D11(ECompilerFlags CompilerFlag)
 	switch(CompilerFlag)
 	{
 	case CFLAG_PreferFlowControl: return D3D10_SHADER_PREFER_FLOW_CONTROL;
-	case CFLAG_Debug: return D3D10_SHADER_DEBUG | D3D10_SHADER_SKIP_OPTIMIZATION;
 	case CFLAG_AvoidFlowControl: return D3D10_SHADER_AVOID_FLOW_CONTROL;
 	default: return 0;
 	};
@@ -352,7 +351,7 @@ void CompileD3D11Shader(const FShaderCompilerInput& Input,FShaderCompilerOutput&
 		// Unpack uniform matrices as row-major to match the CPU layout.
 		| D3D10_SHADER_PACK_MATRIX_ROW_MAJOR;
 
-	if (DEBUG_SHADERS) 
+	if (DEBUG_SHADERS || Input.Environment.CompilerFlags.Contains(CFLAG_Debug)) 
 	{
 		//add the debug flags
 		CompileFlags |= D3D10_SHADER_DEBUG | D3D10_SHADER_SKIP_OPTIMIZATION;
@@ -632,21 +631,29 @@ void CompileD3D11Shader(const FShaderCompilerInput& Input,FShaderCompilerOutput&
 				}
 			}
 
-			// Strip shader reflection
-			D3D_SHADER_DATA ShaderData;
-			ShaderData.pBytecode = Shader->GetBufferPointer();
-			ShaderData.BytecodeLength = Shader->GetBufferSize();
 			TRefCountPtr<ID3DBlob> CompressedData;
-			Result = D3DStripShader(Shader->GetBufferPointer(),
-				Shader->GetBufferSize(),
-				D3DCOMPILER_STRIP_REFLECTION_DATA | D3DCOMPILER_STRIP_DEBUG_INFO | D3DCOMPILER_STRIP_TEST_BLOBS,
-				CompressedData.GetInitReference());
 
-			if (FAILED(Result))
+			if (Input.Environment.CompilerFlags.Contains(CFLAG_DontStripDebugInfo))
 			{
-				UE_LOG(LogD3D11ShaderCompiler, Fatal,TEXT("D3DStripShader failed: Result=%08x"),Result);
+				CompressedData = Shader;
 			}
-
+			else
+			{
+				// Strip shader reflection and debug info
+				D3D_SHADER_DATA ShaderData;
+				ShaderData.pBytecode = Shader->GetBufferPointer();
+				ShaderData.BytecodeLength = Shader->GetBufferSize();
+				Result = D3DStripShader(Shader->GetBufferPointer(),
+					Shader->GetBufferSize(),
+					D3DCOMPILER_STRIP_REFLECTION_DATA | D3DCOMPILER_STRIP_DEBUG_INFO | D3DCOMPILER_STRIP_TEST_BLOBS,
+					CompressedData.GetInitReference());
+
+				if (FAILED(Result))
+				{
+					UE_LOG(LogD3D11ShaderCompiler, Fatal,TEXT("D3DStripShader failed: Result=%08x"),Result);
+				}
+			}
+			
 			// Build the SRT for this shader.
 			FD3D11ShaderResourceTable SRT;
 			{
diff --git a/Engine/Source/Runtime/Engine/Private/ShaderCompiler/ShaderCompiler.cpp b/Engine/Source/Runtime/Engine/Private/ShaderCompiler/ShaderCompiler.cpp
--- a/Engine/Source/Runtime/Engine/Private/ShaderCompiler/ShaderCompiler.cpp
+++ b/Engine/Source/Runtime/Engine/Private/ShaderCompiler/ShaderCompiler.cpp
@@ -32,6 +32,18 @@ static FAutoConsoleVariableRef CVarDumpShaderDebugInfo(
 	TEXT("On iOS, if the PowerVR graphics SDK is installed to the default path, the PowerVR shader compiler will be called and errors will be reported during the cook.")
 	);
 
+static TAutoConsoleVariable<int32> CVarStripShaderDebugData(
+	TEXT("r.StripShaderDebugData"),
+	1,
+	TEXT("Whether to strip shader reflection and debug data from shader bytecode.  When using graphical debuggers like Nsight it can be useful to disable this on startup."),
+	ECVF_ReadOnly);
+
+static TAutoConsoleVariable<int32> CVarOptimizeShaders(
+	TEXT("r.OptimizeShaders"),
+	1,
+	TEXT("Whether to optimize shaders.  When using graphical debuggers like Nsight it can be useful to disable this on startup."),
+	ECVF_ReadOnly);
+
 #if !(UE_BUILD_SHIPPING || UE_BUILD_TEST)
 static TAutoConsoleVariable<FString> CVarD3DCompilerPath(TEXT("r.D3DCompilerPath"),
 	TEXT(""),	// default
@@ -2062,6 +2074,24 @@ void GlobalBeginCompileShader(
 	}
 #endif
 
+	{
+		static const auto CVar = IConsoleManager::Get().FindConsoleVariable(TEXT("r.OptimizeShaders"));
+
+		if (CVar->GetInt() == 0)
+		{
+			Input.Environment.CompilerFlags.Add(CFLAG_Debug);
+		}
+	}
+	
+	{
+		static const auto CVar = IConsoleManager::Get().FindConsoleVariable(TEXT("r.StripShaderDebugData"));
+
+		if (CVar->GetInt() == 0)
+		{
+			Input.Environment.CompilerFlags.Add(CFLAG_DontStripDebugInfo);
+		}
+	}
+
 	{
 		FString ShaderPDBRoot;
 		GConfig->GetString(TEXT("DevOptions.Shaders"), TEXT("ShaderPDBRoot"), ShaderPDBRoot, GEngineIni);
diff --git a/Engine/Source/Runtime/ShaderCore/Private/Shader.cpp b/Engine/Source/Runtime/ShaderCore/Private/Shader.cpp
--- a/Engine/Source/Runtime/ShaderCore/Private/Shader.cpp
+++ b/Engine/Source/Runtime/ShaderCore/Private/Shader.cpp
@@ -1158,6 +1158,16 @@ void ShaderMapAppendKeyString(EShaderPlatform Platform, FString& KeyString)
 		KeyString += (CVar && CVar->GetInt() != 0) ? TEXT("_DBuf") : TEXT("_NoDBuf");
 	}
 
+	{
+		static const auto CVar = IConsoleManager::Get().FindConsoleVariable(TEXT("r.StripShaderDebugData"));
+		KeyString += (CVar && CVar->GetInt() != 0) ? TEXT("") : TEXT("_NoStrip");
+	}
+
+	{
+		static const auto CVar = IConsoleManager::Get().FindConsoleVariable(TEXT("r.OptimizeShaders"));
+		KeyString += (CVar && CVar->GetInt() != 0) ? TEXT("") : TEXT("_NoOpt");
+	}
+
 	if( Platform == SP_PS4 )
 	{
 		{
diff --git a/Engine/Source/Runtime/ShaderCore/Public/ShaderCore.h b/Engine/Source/Runtime/ShaderCore/Public/ShaderCore.h
--- a/Engine/Source/Runtime/ShaderCore/Public/ShaderCore.h
+++ b/Engine/Source/Runtime/ShaderCore/Public/ShaderCore.h
@@ -110,7 +110,8 @@ enum ECompilerFlags
 	/** Only allows standard optimizations, not the longest compile times. */
 	CFLAG_StandardOptimization,
 	/** Shader should use on chip memory instead of main memory ring buffer memory. */
-	CFLAG_OnChip
+	CFLAG_OnChip,
+	CFLAG_DontStripDebugInfo
 };
 
 /**
