Fix for UFE compiling under Rocket.

diff --git a/Engine/Source/Developer/Windows/WindowsTargetPlatform/Private/WindowsTargetPlatformPrivatePCH.h b/Engine/Source/Developer/Windows/WindowsTargetPlatform/Private/WindowsTargetPlatformPrivatePCH.h
--- a/Engine/Source/Developer/Windows/WindowsTargetPlatform/Private/WindowsTargetPlatformPrivatePCH.h
+++ b/Engine/Source/Developer/Windows/WindowsTargetPlatform/Private/WindowsTargetPlatformPrivatePCH.h
@@ -7,6 +7,7 @@
  *****************************************************************************/
 
 #include "Core.h"
+#include "CoreUObject.h"
 #include "WindowsPlatformProperties.h"
 
 #if WITH_ENGINE
