Quick workaround for B6H and BC7 on OpenGL. This looks at the project settings (.ini) and allows D3D11 texture formats if D3D11 SM5 is the only targeted shader format, otherwise it will fall back to DXT/etc. Note that it does not check for command-line overrides, like running with -opengl.

diff --git a/Engine/Source/Developer/Android/AndroidTargetPlatform/Private/AndroidTargetPlatform.inl b/Engine/Source/Developer/Android/AndroidTargetPlatform/Private/AndroidTargetPlatform.inl
--- a/Engine/Source/Developer/Android/AndroidTargetPlatform/Private/AndroidTargetPlatform.inl
+++ b/Engine/Source/Developer/Android/AndroidTargetPlatform/Private/AndroidTargetPlatform.inl
@@ -139,7 +139,7 @@ inline void FAndroidTargetPlatform<TPlatformProperties>::GetTextureFormats( cons
 	// and supported by the device, the first supported format listed will be used. 
 	// eg, ETC1/uncompressed should always be last
 
-	bool bNoCompression = InTexture->CompressionNone								// Code wants the texture uncompressed.
+	bool bNoCompression = InTexture->CompressionNone				// Code wants the texture uncompressed.
 		|| (InTexture->LODGroup == TEXTUREGROUP_ColorLookupTable)	// Textures in certain LOD groups should remain uncompressed.
 		|| (InTexture->LODGroup == TEXTUREGROUP_Bokeh)
 		|| (InTexture->CompressionSettings == TC_EditorIcon)
@@ -157,7 +157,8 @@ inline void FAndroidTargetPlatform<TPlatformProperties>::GetTextureFormats( cons
 	{
 		OutFormats.Add(AndroidTexFormat::NameBGRA8);
 	}
-	else if (InTexture->CompressionSettings == TC_HDR)
+	else if (InTexture->CompressionSettings == TC_HDR
+		|| InTexture->CompressionSettings == TC_HDR_Compressed)
 	{
 		OutFormats.Add(AndroidTexFormat::NameRGBA16F);
 	}
@@ -189,7 +190,8 @@ inline void FAndroidTargetPlatform<TPlatformProperties>::GetTextureFormats( cons
 	{
 		OutFormats.Add(AndroidTexFormat::NameG8);
 	}
-	else if (InTexture->bForcePVRTC4)
+	else if (InTexture->bForcePVRTC4
+		|| InTexture->CompressionSettings == TC_BC7)
 	{
 		AddTextureFormatIfSupports(AndroidTexFormat::NamePVRTC4, OutFormats);
 		AddTextureFormatIfSupports(AndroidTexFormat::NameDXT5, OutFormats);
diff --git a/Engine/Source/Developer/HTML5/HTML5TargetPlatform/Private/HTML5TargetPlatform.cpp b/Engine/Source/Developer/HTML5/HTML5TargetPlatform/Private/HTML5TargetPlatform.cpp
--- a/Engine/Source/Developer/HTML5/HTML5TargetPlatform/Private/HTML5TargetPlatform.cpp
+++ b/Engine/Source/Developer/HTML5/HTML5TargetPlatform/Private/HTML5TargetPlatform.cpp
@@ -208,7 +208,8 @@ void FHTML5TargetPlatform::GetTextureFormats( const UTexture* Texture, TArray<FN
 				TextureFormatName = NameBGRA8;
 			}
 	}
-	else if (Texture->CompressionSettings == TC_HDR)
+	else if (Texture->CompressionSettings == TC_HDR
+		|| Texture->CompressionSettings == TC_HDR_Compressed)
 	{
 		TextureFormatName = NameRGBA16F;
 	}
diff --git a/Engine/Source/Developer/Linux/LinuxTargetPlatform/Private/LinuxTargetPlatform.h b/Engine/Source/Developer/Linux/LinuxTargetPlatform/Private/LinuxTargetPlatform.h
--- a/Engine/Source/Developer/Linux/LinuxTargetPlatform/Private/LinuxTargetPlatform.h
+++ b/Engine/Source/Developer/Linux/LinuxTargetPlatform/Private/LinuxTargetPlatform.h
@@ -158,12 +158,26 @@ public:
 		return StaticMeshLODSettings;
 	}
 
+
 	virtual void GetTextureFormats( const UTexture* InTexture, TArray<FName>& OutFormats ) const override
 	{
+		static FName NameBC6H(TEXT("BC6H"));
+		static FName NameBC7(TEXT("BC7"));
+		static FName NameRGBA16F(TEXT("RGBA16F"));
+		static FName NameAutoDXT(TEXT("AutoDXT"));
 		if (!IS_DEDICATED_SERVER)
 		{
 			// just use the standard texture format name for this texture
-			OutFormats.Add(this->GetDefaultTextureFormatName(InTexture, EngineSettings));
+			FName TextureFormatName = this->GetDefaultTextureFormatName(InTexture, EngineSettings);
+			if (TextureFormatName == NameBC6H)
+			{
+				TextureFormatName = NameRGBA16F;
+			}
+			else if (TextureFormatName == NameBC7)
+			{
+				TextureFormatName = NameAutoDXT;
+			}
+			OutFormats.Add(TextureFormatName);
 		}
 	}
 
diff --git a/Engine/Source/Developer/Mac/MacTargetPlatform/Private/GenericMacTargetPlatform.h b/Engine/Source/Developer/Mac/MacTargetPlatform/Private/GenericMacTargetPlatform.h
--- a/Engine/Source/Developer/Mac/MacTargetPlatform/Private/GenericMacTargetPlatform.h
+++ b/Engine/Source/Developer/Mac/MacTargetPlatform/Private/GenericMacTargetPlatform.h
@@ -114,9 +114,13 @@ return TSuper::SupportsFeature(Feature);
 		return StaticMeshLODSettings;
 	}
 
-
 	virtual void GetTextureFormats( const UTexture* InTexture, TArray<FName>& OutFormats ) const override
 	{
+		static FName NameBC6H(TEXT("BC6H"));
+		static FName NameBC7(TEXT("BC7"));
+		static FName NameRGBA16F(TEXT("RGBA16F"));
+		static FName NameAutoDXT(TEXT("AutoDXT"));
+
 		if (IS_DEDICATED_SERVER)
 		{
 			// no textures needed for dedicated server target
@@ -125,7 +129,16 @@ return TSuper::SupportsFeature(Feature);
 		else
 		{
 			// just use the standard texture format name for this texture
-			OutFormats.Add(this->GetDefaultTextureFormatName(InTexture, EngineSettings));
+			FName TextureFormatName = this->GetDefaultTextureFormatName(InTexture, EngineSettings);
+			if (TextureFormatName == NameBC6H)
+			{
+				TextureFormatName = NameRGBA16F;
+			}
+			else if (TextureFormatName == NameBC7)
+			{
+				TextureFormatName = NameAutoDXT;
+			}
+			OutFormats.Add(TextureFormatName);
 		}
 	}
 
diff --git a/Engine/Source/Developer/Windows/WindowsTargetPlatform/Private/GenericWindowsTargetPlatform.h b/Engine/Source/Developer/Windows/WindowsTargetPlatform/Private/GenericWindowsTargetPlatform.h
--- a/Engine/Source/Developer/Windows/WindowsTargetPlatform/Private/GenericWindowsTargetPlatform.h
+++ b/Engine/Source/Developer/Windows/WindowsTargetPlatform/Private/GenericWindowsTargetPlatform.h
@@ -47,6 +47,11 @@ public:
 					TargetedShaderFormats.Remove(ShaderFormat);
 				}
 			}
+
+			// If we're targeting only DX11 we can use DX11 texture formats. Otherwise we'd have to compress fallbacks and increase the size of cooked content significantly.
+			static FString NAME_PCD3D_SM5(TEXT("PCD3D_SM5"));
+			bSupportDX11TextureFormats = TargetedShaderFormats.Num() == 1
+				&& TargetedShaderFormats[0] == NAME_PCD3D_SM5;
 		#endif
 	}
 
@@ -110,9 +115,6 @@ public:
 		// no shaders needed for dedicated server target
 		if (!IS_DEDICATED_SERVER)
 		{
-			// Only shader formats needed for cooked builds should be added here.
-			// TODO: this should be configurable per-project, e.g. you may or
-			//       may not wish to support SM4 or OpenGL.
 			static FName NAME_PCD3D_SM5(TEXT("PCD3D_SM5"));
 			static FName NAME_PCD3D_SM4(TEXT("PCD3D_SM4"));
 			static FName NAME_GLSL_150(TEXT("GLSL_150"));
@@ -140,10 +142,26 @@ public:
 
 	virtual void GetTextureFormats( const UTexture* InTexture, TArray<FName>& OutFormats ) const override
 	{
+		static FName NameBC6H(TEXT("BC6H"));
+		static FName NameBC7(TEXT("BC7"));
+		static FName NameRGBA16F(TEXT("RGBA16F"));
+		static FName NameAutoDXT(TEXT("AutoDXT"));
+
 		if (!IS_DEDICATED_SERVER)
 		{
-			// just use the standard texture format name for this texture
-			OutFormats.Add(GetDefaultTextureFormatName(InTexture, EngineSettings));
+			FName TextureFormatName = GetDefaultTextureFormatName(InTexture, EngineSettings);
+			if (!bSupportDX11TextureFormats)
+			{
+				if (TextureFormatName == NameBC6H)
+				{
+					TextureFormatName = NameRGBA16F;
+				}
+				else if (TextureFormatName == NameBC7)
+				{
+					TextureFormatName = NameAutoDXT;
+				}
+			}
+			OutFormats.Add(TextureFormatName);
 		}
 	}
 
@@ -249,6 +267,9 @@ private:
 
 	// List of shader formats specified as targets
 	TArray<FString> TargetedShaderFormats;
+
+	// True if the project supports non-DX11 texture formats.
+	bool bSupportDX11TextureFormats;
 #endif // WITH_ENGINE
 
 private:
diff --git a/Engine/Source/Developer/iOS/IOSTargetPlatform/Private/IOSTargetPlatform.cpp b/Engine/Source/Developer/iOS/IOSTargetPlatform/Private/IOSTargetPlatform.cpp
--- a/Engine/Source/Developer/iOS/IOSTargetPlatform/Private/IOSTargetPlatform.cpp
+++ b/Engine/Source/Developer/iOS/IOSTargetPlatform/Private/IOSTargetPlatform.cpp
@@ -375,6 +375,8 @@ void FIOSTargetPlatform::GetTextureFormats( const UTexture* Texture, TArray<FNam
 		FName(TEXT("DXT5n")),	FName(TEXT("PVRTCN")),		FName(TEXT("ASTC_NormalAG")),
 		FName(TEXT("BC5")),		FName(TEXT("PVRTCN")),		FName(TEXT("ASTC_NormalRG")),
 		FName(TEXT("AutoDXT")),	FName(TEXT("AutoPVRTC")),	FName(TEXT("ASTC_RGBAuto")),
+		FName(TEXT("BC6H")),	FName(TEXT("RGBA16F")),		FName(TEXT("RGBA16F")),
+		FName(TEXT("BC7")),		FName(TEXT("AutoPVRTC")),	FName(TEXT("ASTC_RGBAuto")),
 	};
 
 	FName TextureFormatName = NAME_None;
diff --git a/Engine/Source/Runtime/Engine/Private/Texture2D.cpp b/Engine/Source/Runtime/Engine/Private/Texture2D.cpp
--- a/Engine/Source/Runtime/Engine/Private/Texture2D.cpp
+++ b/Engine/Source/Runtime/Engine/Private/Texture2D.cpp
@@ -787,18 +787,24 @@ FTextureResource* UTexture2D::CreateResource()
 		}
 	}
 
+	const EPixelFormat PixelFormat = GetPixelFormat();
 	const bool bIncompatibleTexture = (NumMips == 0);
 	const bool bTextureTooLarge = FMath::Max(GetSizeX(), GetSizeY()) > (int32)GetMax2DTextureDimension();
 	// Too large textures with full mip chains are OK as we load up to max supported mip.
 	const bool bNotSupportedByRHI = NumMips == 1 && bTextureTooLarge;
+	const bool bFormatNotSupported = !(GPixelFormats[PixelFormat].Supported);
 
-	if (bIncompatibleTexture || bNotSupportedByRHI)
+	if (bIncompatibleTexture || bNotSupportedByRHI || bFormatNotSupported)
 	{
 		// Handle unsupported/corrupted/incompatible textures :(
 		ResidentMips	= 0;
 		RequestedMips	= 0;
 
-		if (bNotSupportedByRHI)
+		if (bFormatNotSupported)
+		{
+			UE_LOG(LogTexture, Error, TEXT("%s is %s which is not supported."), *GetFullName(), GPixelFormats[PixelFormat].Name);
+		}
+		else if (bNotSupportedByRHI)
 		{
 			UE_LOG(LogTexture, Warning, TEXT("%s cannot be created, exceeds this rhi's maximum dimension (%d) and has no mip chain to fall back on."), *GetFullName(), GetMax2DTextureDimension());
 		}
