UE-13389 - Project Settings typo: "OSVersion" instead of "OS Version"

I've updated the DisplayName for this property in it's uproperty meta data. Category, also, now has a relevant space.

diff --git a/Engine/Source/Developer/Windows/WindowsTargetPlatform/Classes/WindowsTargetSettings.h b/Engine/Source/Developer/Windows/WindowsTargetPlatform/Classes/WindowsTargetSettings.h
--- a/Engine/Source/Developer/Windows/WindowsTargetPlatform/Classes/WindowsTargetSettings.h
+++ b/Engine/Source/Developer/Windows/WindowsTargetPlatform/Classes/WindowsTargetSettings.h
@@ -36,6 +36,6 @@ public:
 	/**
 	 * Determine the minimum supported 
 	 */
-	UPROPERTY(EditAnywhere, config, Category=OSInfo)
+	UPROPERTY(EditAnywhere, config, Category="OS Info", Meta=(DisplayName = "Minimum OS Version"))
 	TEnumAsByte<EMinimumSupportedOS> MinimumOSVersion;
 };
