UE-13383 - Minimum OS project setting doesn't persist between editor sessions

UWindowsTargetSettings object is loaded before the actual UWindowsTargetSettings class, so no config value is loaded correctly thorugh the uobject system

I now load the value through the GConfig system

diff --git a/Engine/Source/Developer/Windows/WindowsTargetPlatform/Private/WindowsTargetPlatformModule.cpp b/Engine/Source/Developer/Windows/WindowsTargetPlatform/Private/WindowsTargetPlatformModule.cpp
--- a/Engine/Source/Developer/Windows/WindowsTargetPlatform/Private/WindowsTargetPlatformModule.cpp
+++ b/Engine/Source/Developer/Windows/WindowsTargetPlatform/Private/WindowsTargetPlatformModule.cpp
@@ -69,6 +69,16 @@ public:
 		
 		// We need to manually load the config properties here, as this module is loaded before the UObject system is setup to do this
 		GConfig->GetArray(TEXT("/Script/WindowsTargetPlatform.WindowsTargetSettings"), TEXT("TargetedRHIs"), TargetSettings->TargetedRHIs, GEngineIni);
+
+		// When this is initialized the UEnum for EMinimumSupportedOS hasn't been registered. 
+		FString MinOSString;
+		if (GConfig->GetString(TEXT("/Script/WindowsTargetPlatform.WindowsTargetSettings"), TEXT("MinimumOSVersion"), MinOSString, GEngineIni))
+		{
+			// We need to parse the string and compare manually.
+			TargetSettings->MinimumOSVersion = MinOSString == TEXT("MSOS_XP") ? EMinimumSupportedOS::MSOS_XP : EMinimumSupportedOS::MSOS_Vista;
+		}
+
+
 		TargetSettings->AddToRoot();
 
 		ISettingsModule* SettingsModule = FModuleManager::GetModulePtr<ISettingsModule>("Settings");
