Fixed assertion at startup when same shader format is specified twice in a project's engine config file

Somehow one of my projects ended up with many of the same key.  The code that was handling loading of these had an assumption that each key would only appear once in the list (fixed.)

Example of duplicate keys that would cause this to crash:

[/Script/WindowsTargetPlatform.WindowsTargetSettings]
TargetedRHIs=PCD3D_SM5
TargetedRHIs=PCD3D_SM4
TargetedRHIs=PCD3D_SM5
TargetedRHIs=PCD3D_SM4
TargetedRHIs=PCD3D_SM5
TargetedRHIs=PCD3D_SM4
TargetedRHIs=PCD3D_SM5
TargetedRHIs=PCD3D_SM4
TargetedRHIs=PCD3D_SM5
TargetedRHIs=PCD3D_SM4

diff --git a/Engine/Source/Developer/Linux/LinuxTargetPlatform/Private/LinuxTargetPlatform.h b/Engine/Source/Developer/Linux/LinuxTargetPlatform/Private/LinuxTargetPlatform.h
--- a/Engine/Source/Developer/Linux/LinuxTargetPlatform/Private/LinuxTargetPlatform.h
+++ b/Engine/Source/Developer/Linux/LinuxTargetPlatform/Private/LinuxTargetPlatform.h
@@ -52,7 +52,7 @@ public:
 			FString ShaderFormat = TargetedShaderFormats[ShaderFormatIdx];
 			if(PossibleShaderFormats.Contains(FName(*ShaderFormat)) == false)
 			{
-				TargetedShaderFormats.Remove(ShaderFormat);
+				TargetedShaderFormats.RemoveAt(ShaderFormatIdx);
 			}
 		}
 #endif // WITH_ENGINE
diff --git a/Engine/Source/Developer/Windows/WindowsTargetPlatform/Private/GenericWindowsTargetPlatform.h b/Engine/Source/Developer/Windows/WindowsTargetPlatform/Private/GenericWindowsTargetPlatform.h
--- a/Engine/Source/Developer/Windows/WindowsTargetPlatform/Private/GenericWindowsTargetPlatform.h
+++ b/Engine/Source/Developer/Windows/WindowsTargetPlatform/Private/GenericWindowsTargetPlatform.h
@@ -44,7 +44,7 @@ public:
 				FString ShaderFormat = TargetedShaderFormats[ShaderFormatIdx];
 				if(PossibleShaderFormats.Contains(FName(*ShaderFormat)) == false)
 				{
-					TargetedShaderFormats.Remove(ShaderFormat);
+					TargetedShaderFormats.RemoveAt(ShaderFormatIdx);
 				}
 			}
 
