Copying //UE4/Dev-Rendering to Dev-Main (//UE4/Dev-Main)

==========================
MAJOR FEATURES + CHANGES
==========================

Change 2762805 on 2015/11/11 by [REDACTED_PERSON]

	#UE4 Material editor optimization. FindShaderTypeByName now uses the precreated map to look up types by name. Also changed GetRepresentativeShaderTypesAndDescriptions to return a map.

Change 2763372 on 2015/11/11 by [REDACTED_PERSON]

	Auto events around ComputeFenceWaits
	No transition ensure in parallel mode.

Change 2764204 on 2015/11/12 by [REDACTED_PERSON]

	DevRendering - Removed reference counting on pipelines as it's not needed
	- Shadermap of shader types no longer owns the unique shaders per pipeline, the pipelines do
	- Added shader pipeline dependencies and shortened ddc entry
	- Dumped stats now show # shared pipelines
	- Added IMPLEMENT_SHADERPIPELINE_TYPE_* macros

Change 2764409 on 2015/11/12 by [REDACTED_PERSON]

	Event based GPU frame waiting.

Change 2764474 on 2015/11/12 by [REDACTED_PERSON]

	UE4 - Optimize speed tree stuff in FLocalVertexFactory::SetMesh

Change 2765900 on 2015/11/13 by [REDACTED_PERSON]

	Move PS4 system event checking onto a task thread. Saves ~.1ms on gamethread

Change 2768478 on 2015/11/16 by [REDACTED_PERSON]

	Remove flushes when locking single-mip 2d textures for writeonly on PS4

Change 2769532 on 2015/11/17 by [REDACTED_PERSON]

	Discard outdated vertex paint data when loading static mesh components. UE-22603

Change 2770002 on 2015/11/17 by [REDACTED_PERSON]

	DevRendering - Prep for reentrant calls for D3D & PS4 shader compilers
	- Added r.D3DDumpAMDCodeXLFile to enable writing out a batch file for AMD cards (it was always generating this, now it's optional) - defaults to 1, enabled
	- Added r.D3DCheckShadersForDouble to enable checking for doubles when compiling shaders, speeds up a tiny bit when disabling - defaults to 1, enabled
	- Fixed PS4's InitFastMathDetection(), was always testing for this

Change 2770103 on 2015/11/17 by [REDACTED_PERSON]

	Added SourceCubemapAngle to sky lights, so the incoming lighting can be rotated

diff --git a/Engine/Config/BaseEngine.ini b/Engine/Config/BaseEngine.ini
--- a/Engine/Config/BaseEngine.ini
+++ b/Engine/Config/BaseEngine.ini
@@ -870,14 +870,7 @@ bUseStreamingPause=false
 +K2FieldRedirects=(OldFieldName="PrimitiveComponent.CreateAndSetMaterialInstanceDynamic",NewFieldName="PrimitiveComponent.CreateDynamicMaterialInstance")
 +K2FieldRedirects=(OldFieldName="PrimitiveComponent.CreateAndSetMaterialInstanceDynamicFromMaterial",NewFieldName="PrimitiveComponent.CreateDynamicMaterialInstance")
 +K2ParamRedirects=(NodeName="/Script/BlueprintGraph.K2Node_CallFunction", OldParamName="PrimitiveComponent.CreateDynamicMaterialInstance.Parent", NewParamName="SourceMaterial")
-
-;MID rename 
-+K2FieldRedirects=(OldFieldName="DecalComponent.CreateMIDForDecal", NewFieldName="DecalComponent.CreateDynamicMaterialInstance")
-+K2FieldRedirects=(OldFieldName="DecalActor.CreateMIDForDecal", NewFieldName="DecalActor.CreateDynamicMaterialInstance")
-+K2FieldRedirects=(OldFieldName="KismetMaterialLibrary.CreateMaterialInstanceDynamic",NewFieldName="KismetMaterialLibrary.CreateDynamicMaterialInstance")
-+K2FieldRedirects=(OldFieldName="PrimitiveComponent.CreateAndSetMaterialInstanceDynamic",NewFieldName="PrimitiveComponent.CreateDynamicMaterialInstance")
-+K2FieldRedirects=(OldFieldName="PrimitiveComponent.CreateAndSetMaterialInstanceDynamicFromMaterial",NewFieldName="PrimitiveComponent.CreateDynamicMaterialInstance")
-+K2ParamRedirects=(NodeName="/Script/BlueprintGraph.K2Node_CallFunction", OldParamName="PrimitiveComponent.CreateDynamicMaterialInstance.Parent", NewParamName="SourceMaterial")
++K2ParamRedirects=(NodeName="/Script/BlueprintGraph.K2Node_CallFunction", OldParamName="MaterialInstanceDynamic.K2_CopyMaterialInstanceParameters.SourceMaterialToCopyFrom", NewParamName="Source")
 
 ; Kismet property redirects
 ; Must specify class and a property name for old property, and just the property name for the new property (class name is optional)
diff --git a/Engine/Config/BaseInput.ini b/Engine/Config/BaseInput.ini
--- a/Engine/Config/BaseInput.ini
+++ b/Engine/Config/BaseInput.ini
@@ -52,8 +52,10 @@ MaxScrollbackSize=1024
 +ManualAutoCompleteList=(Command="DebugCreatePlayer 1",Desc=)
 +ManualAutoCompleteList=(Command="ToggleDrawEvents",Desc="Toggles annotations for shader debugging with Pix, Razor or similar GPU capture tools")
 +ManualAutoCompleteList=(Command="Shot",Desc="Make a screenshot")
-+ManualAutoCompleteList=(Command="RecompileShaders changed",Desc=)
-+ManualAutoCompleteList=(Command="RecompileShaders global",Desc=)
++ManualAutoCompleteList=(Command="RecompileShaders changed",Desc="Recompile shaders that have any changes on their source files")
++ManualAutoCompleteList=(Command="RecompileShaders global",Desc="Recompile all the global shaders")
++ManualAutoCompleteList=(Command="RecompileShaders material ",Desc="Recompile shaders for a specific material")
++ManualAutoCompleteList=(Command="RecompileShaders all",Desc="Force a full rebuild of all shaders")
 +ManualAutoCompleteList=(Command="StartFPSChart",Desc="after that use StopFPSChart")
 +ManualAutoCompleteList=(Command="StopFPSChart",Desc="after that look for the output in Saved/Profiling/FPSChartStats")
 +ManualAutoCompleteList=(Command="FreezeAt",Desc="Locks the player view and rendering time.")
diff --git a/Engine/Shaders/BasePassCommon.usf b/Engine/Shaders/BasePassCommon.usf
--- a/Engine/Shaders/BasePassCommon.usf
+++ b/Engine/Shaders/BasePassCommon.usf
@@ -16,7 +16,7 @@
 #define USES_GBUFFER						(FEATURE_LEVEL >= FEATURE_LEVEL_SM4 && (MATERIALBLENDING_SOLID || MATERIALBLENDING_MASKED) && !NO_GBUFFER)
 
 // Only some shader models actually need custom data.
-#define WRITES_CUSTOMDATA_TO_GBUFFER		(USES_GBUFFER && (!SELECTIVE_BASEPASS_OUTPUTS || MATERIAL_SHADINGMODEL_SUBSURFACE || MATERIAL_SHADINGMODEL_PREINTEGRATED_SKIN || MATERIAL_SHADINGMODEL_SUBSURFACE_PROFILE || MATERIAL_SHADINGMODEL_CLEAR_COAT))
+#define WRITES_CUSTOMDATA_TO_GBUFFER		(USES_GBUFFER && (!SELECTIVE_BASEPASS_OUTPUTS || MATERIAL_SHADINGMODEL_SUBSURFACE || MATERIAL_SHADINGMODEL_PREINTEGRATED_SKIN || MATERIAL_SHADINGMODEL_SUBSURFACE_PROFILE || MATERIAL_SHADINGMODEL_CLEAR_COAT || MATERIAL_SHADINGMODEL_TWOSIDED_FOLIAGE || MATERIAL_SHADINGMODEL_HAIR || MATERIAL_SHADINGMODEL_CLOTH || MATERIAL_SHADINGMODEL_EYE))
 
 // Based on GetPrecomputedShadowMasks()
 #define WRITES_PRECSHADOWFACTOR_TO_GBUFFER	(USES_GBUFFER && ALLOW_STATIC_LIGHTING && (!SELECTIVE_BASEPASS_OUTPUTS || (STATICLIGHTING_TEXTUREMASK && STATICLIGHTING_SIGNEDDISTANCEFIELD)))
diff --git a/Engine/Shaders/Common.usf b/Engine/Shaders/Common.usf
--- a/Engine/Shaders/Common.usf
+++ b/Engine/Shaders/Common.usf
@@ -923,15 +923,20 @@ float3 DecodeSceneColorForMaterialNode(float2 ScreenUV)
 #endif
 }
 
-// Converts depth buffer Z / W into W
+// also see ConvertToDeviceZ()
+// @param DeviceZ value that is stored in the depth buffer (Z/W)
+// @return SceneDepth (linear in world units, W)
 float ConvertFromDeviceZ(float DeviceZ)
 {
 	return 1.f / (DeviceZ * View.InvDeviceZToWorldZTransform[2] - View.InvDeviceZToWorldZTransform[3]);	
 }
 
-// inverse opteration of ConvertFromDeviceZ
+// inverse operation of ConvertFromDeviceZ()
+// @param SceneDepth (linear in world units, W)
+// @return DeviceZ (Z/W)
 float ConvertToDeviceZ(float SceneDepth)
 {
+	// todo: can change the ADD+MUL into a MAD for precision and performance
 	return 1.0f / ((SceneDepth + View.InvDeviceZToWorldZTransform[3]) * View.InvDeviceZToWorldZTransform[2]);
 }
 
diff --git a/Engine/Shaders/PostProcessMotionBlur.usf b/Engine/Shaders/PostProcessMotionBlur.usf
--- a/Engine/Shaders/PostProcessMotionBlur.usf
+++ b/Engine/Shaders/PostProcessMotionBlur.usf
@@ -651,30 +651,26 @@ void VelocityScatterVS(
 	// triangle B: 3:right bottom, 4:left bottom, 5: right top
 	float2 CornerOffset = float2(VId % 2, VId / 2) * 2 - 1;
 
-	uint2 TilePos = uint2( IId % ViewportRect.z, IId / ViewportRect.z );
+	uint2 PixelPos = uint2( IId % ViewportRect.z, IId / ViewportRect.z );
 
 	BRANCH
-	if( TilePos.y >= ViewportRect.w )
+	if( PixelPos.y >= ViewportRect.w )
 	{
 		OutColor = 0;
 		return;
 	}
 
-	float4 MinMaxVelocity = PostprocessInput0[ TilePos ];
+	float4 MinMaxVelocity = PostprocessInput0[ PixelPos ];
 	OutColor = MinMaxVelocity;
 
-	//float2 VelocityPixels = MinMaxVelocity.zw * ScreenPosToPixel.x;
-	float2 VelocityPixels = MinMaxVelocity.zw * 0.5 * 1920.0 / 16.0;
-	
-	// TODO move to velocity flatten shader
-	//VelocityPixels *= VelocityScale.xy;
-	VelocityPixels *= 0.5;
+	float4 MinMaxVelocityPixels = MinMaxVelocity * MotionBlurParameters.z;
+	float2 VelocityPixels = MinMaxVelocityPixels.zw;
 	
 #if 1
 	BRANCH
 	if( dot( VelocityPixels, VelocityPixels ) * 16*16 <= 0.25 )
 	{
-		OutPosition.xy = ( TilePos + CornerOffset * 0.5 - ScreenPosToPixel.zw ) / ScreenPosToPixel.xy;
+		OutPosition.xy = ( PixelPos + CornerOffset * 0.5 - ScreenPosToPixel.zw ) / ScreenPosToPixel.xy;
 		OutPosition.z = 0.0002;		// zero clips
 		return;
 	}
@@ -700,21 +696,11 @@ void VelocityScatterVS(
 	float2 AxisY = float2( -VelocityDir.y, VelocityDir.x );
 	CornerOffset = AxisX * CornerOffset.x + AxisY * CornerOffset.y;
 		
-	OutPosition.xy = ( TilePos + CornerOffset - ScreenPosToPixel.zw ) / ScreenPosToPixel.xy;
-
-	if( bDrawMax )
-	{
-		// Depth ordered by velocity length
-		OutPosition.z = saturate( VelocityLengthPixels / ScreenPosToPixel.x * 0.5 );
-	}
-	else
-	{
-		//float2 MinVelocityPixels = MinMaxVelocity.xy * ScreenPosToPixel.xy * VelocityScale.xy;
-		float2 MinVelocityPixels = MinMaxVelocity.xy * 0.5 * 1920.0 / 16.0 * 0.5;
-		OutPosition.z = saturate( length( MinVelocityPixels ) / ScreenPosToPixel.x * 0.5 );
-	}
+	OutPosition.xy = ( PixelPos + CornerOffset - ScreenPosToPixel.zw ) / ScreenPosToPixel.xy;
 
-	OutPosition.z = clamp( OutPosition.z, 0.0002, 0.999 );
+	// Depth ordered by velocity length
+	OutPosition.z = bDrawMax ? VelocityLengthPixels : length( MinMaxVelocityPixels.xy );
+	OutPosition.z = clamp( OutPosition.z / ScreenPosToPixel.x * 0.5, 0.0002, 0.999 );
 }
 
 void VelocityScatterPS(
@@ -732,14 +718,11 @@ void VelocityScatterPS(
 float2 DepthCmp( float CenterDepth, float SampleDepth, float DepthScale )
 {
 	return saturate( 0.5 + float2( DepthScale, -DepthScale ) * (SampleDepth - CenterDepth) );
-	//return SampleDepth > CenterDepth ? float2(1,0) : float2(0,1);
 }
 
 float2 SpreadCmp( float OffsetLength, float2 SpreadLength, float PixelToSampleScale )
 {
 	return saturate( PixelToSampleScale * SpreadLength - max( OffsetLength - 1, 0 ) );
-	//return PixelToSampleScale * SpreadLength > OffsetLength ? 1 : 0;
-	//return SpreadLength > SearchLengthPixels * OffsetLength / ( StepCount - 0.5 ) ? 1 : 0;
 }
 
 float SampleWeight( float CenterDepth, float SampleDepth, float OffsetLength, float CenterSpreadLength, float SampleSpreadLength, float PixelToSampleScale, float DepthScale )
@@ -781,24 +764,21 @@ float2 DecodeVelocity( float2 Velocity )
 	return Velocity;
 }
 
-static const float LimitVelocity = 100;
-
 float GetVelocityLengthPixels( float2 EncodedVelocity )
 {
 #if 1
 	// 11:11:10  (VelocityLength, VelocityAngle, Depth)
 	float VelocityLength = EncodedVelocity.x;
-	VelocityLength *= ScreenPosToPixel.x;
-	return min( VelocityLength, LimitVelocity );
+	VelocityLength *= MotionBlurParameters.y;
+	return min( VelocityLength, MotionBlurParameters.w );
 #else
 	float2 Velocity = DecodeVelocity( EncodedVelocity ) * ScreenPosToPixel.xy;	// in pixels
 	float  VelocityLengthSqr = dot( Velocity, Velocity );
 	float  VelocityLength = sqrtFast( VelocityLengthSqr );
-	return min( VelocityLength, LimitVelocity );
+	return min( VelocityLength, MotionBlurParameters.w );
 #endif
 }
 
-
 void MainNewPS(
 	in float4 UVAndScreenPos : TEXCOORD0,
 	in float4 SvPosition : SV_Position,		// after all interpolators
@@ -812,8 +792,6 @@ void MainNewPS(
 
 	OutColor = 0;
 
-	float VelocityScale = 0.5;//MotionBlurParameters.x;
-
 	const uint StepCount = GetStepCountFromQuality() / 2;
 
 	uint2 PixelPos = SvPosition.xy;
@@ -824,8 +802,6 @@ void MainNewPS(
 	float Random2 = RandFast( PixelPos, 5521 );
 #else
 	float Random = InterleavedGradientNoise( PixelPos );
-	//float Random2 = InterleavedGradientNoise( PixelPos.yx );
-
 	float2 PosMod = float2( PixelPos & 1 );
 	float Random2 = ( PosMod.x * 0.5 - 0.25 ) * ( PosMod.y * 2 - 1 );
 #endif
@@ -836,9 +812,8 @@ void MainNewPS(
 	float4 MinMaxVelocity = PostprocessInput3.SampleLevel( PostprocessInput3Sampler, TileUV, 0 );
 	//float4 MinMaxVelocity = PostprocessInput3[ ( PixelPos - ViewportRect.xy ) / 16 ];
 
-	float4 MinMaxVelocityScreen = MinMaxVelocity * VelocityScale;
-	float2 MinVelocityPixels = MinMaxVelocityScreen.xy * ScreenPosToPixel.x;
-	float2 MaxVelocityPixels = MinMaxVelocityScreen.zw * ScreenPosToPixel.x;
+	float2 MinVelocityPixels = MinMaxVelocity.xy * MotionBlurParameters.y;
+	float2 MaxVelocityPixels = MinMaxVelocity.zw * MotionBlurParameters.y;
 
 	float MinVelocityLengthSqrPixels = dot( MinVelocityPixels, MinVelocityPixels );
 	float MaxVelocityLengthSqrPixels = dot( MaxVelocityPixels, MaxVelocityPixels );
@@ -856,13 +831,12 @@ void MainNewPS(
 		return;
 	}
 
-#if 1
+#if 0
 	float4 SearchVectorPixels = float4( MaxVelocityPixels, -MaxVelocityPixels );
-	float2 SearchLengthPixels = length( MaxVelocityPixels ).xx;
 	float4 SearchVector = SearchVectorPixels * PostprocessInput0Size.zwzw;
 #else
 	// Clip MaxVelocity to screen rect
-	float2 InvVelocityScreen = rcp( MinMaxVelocityScreen.zw );
+	float2 InvVelocityScreen = ScreenPosToPixel.x * rcp( MaxVelocityPixels );
 	float2 MinIntersect = -InvVelocityScreen - ScreenSpacePos * InvVelocityScreen;
 	float2 MaxIntersect =  InvVelocityScreen - ScreenSpacePos * InvVelocityScreen;
 	float4 FarIntersect = float4( max( MinIntersect, MaxIntersect ), max( -MinIntersect, -MaxIntersect ) );
@@ -870,13 +844,11 @@ void MainNewPS(
 
 	// +/-
 	float4 SearchVectorPixels = MaxVelocityPixels.xyxy * float4( Intersect.xx, -Intersect.yy );
-	float2 SearchLengthPixels = { length( SearchVectorPixels.xy ), length( SearchVectorPixels.zw ) };
 	float4 SearchVector = SearchVectorPixels * PostprocessInput0Size.zwzw;
 #endif
 
 	// converts pixel length to sample steps
-	//float2 PixelToSampleScale = StepCount / SearchLengthPixels;
-	float2 PixelToSampleScale = StepCount / length( MaxVelocityPixels );
+	float PixelToSampleScale = StepCount * rsqrtFast( dot( MaxVelocityPixels, MaxVelocityPixels ) );
 
 	BRANCH
 	if( bFastPath )
@@ -905,7 +877,7 @@ void MainNewPS(
 	{
 		float3 CenterVelocityDepth = PostprocessInput2.SampleLevel( PostprocessInput2Sampler, UV, 0 ).xyz;
 		float  CenterDepth = CenterVelocityDepth.z;
-		float  CenterVelocityLength = GetVelocityLengthPixels( CenterVelocityDepth.xy ) * VelocityScale;
+		float  CenterVelocityLength = GetVelocityLengthPixels( CenterVelocityDepth.xy );
 
 		float4 ColorAccum = 0;
 
@@ -932,9 +904,9 @@ void MainNewPS(
 				SampleDepth[j] = SampleVelocityDepth.z;
 
 				// in pixels
-				SampleVelocityLength[j] = GetVelocityLengthPixels( SampleVelocityDepth.xy ) * VelocityScale;
+				SampleVelocityLength[j] = GetVelocityLengthPixels( SampleVelocityDepth.xy );
 
-				Weight[j] = SampleWeight( CenterDepth, SampleDepth[j], OffsetLength, CenterVelocityLength, SampleVelocityLength[j], PixelToSampleScale[j], SOFT_Z_EXTENT );
+				Weight[j] = SampleWeight( CenterDepth, SampleDepth[j], OffsetLength, CenterVelocityLength, SampleVelocityLength[j], PixelToSampleScale, SOFT_Z_EXTENT );
 			}
 		
 			bool2 Mirror = bool2( SampleDepth[0] > SampleDepth[1], SampleVelocityLength[1] > SampleVelocityLength[0] );
diff --git a/Engine/Shaders/PostProcessVelocityFlatten.usf b/Engine/Shaders/PostProcessVelocityFlatten.usf
--- a/Engine/Shaders/PostProcessVelocityFlatten.usf
+++ b/Engine/Shaders/PostProcessVelocityFlatten.usf
@@ -13,6 +13,7 @@ groupshared float4 Shared[ THREADGROUP_TOTALSIZE ];
 RWTexture2D<float4>	OutVelocityFlat;
 RWTexture2D<float4>	OutMaxTileVelocity;
 
+float4 MotionBlurParameters;
 
 float4 MinMaxLength( float4 v0, float4 v1 )
 {
@@ -50,14 +51,6 @@ void VelocityFlattenMain(
 	uint3 GroupThreadId : SV_GroupThreadID,
 	uint GroupIndex : SV_GroupIndex ) 
 {
-	if(0)
-	{
-		// Swizzle threads
-		uint TileIndex = (DispatchThreadId.x & 3) + (DispatchThreadId.y & 3) * 4;
-		DispatchThreadId.x = (DispatchThreadId.x & ~3) | ((0xebfa5014 >> (TileIndex*2) ) & 3);
-		DispatchThreadId.y = (DispatchThreadId.y & ~3) | ((0xfa4141af >> (TileIndex*2) ) & 3);
-	}
-
 	uint2 PixelPos = DispatchThreadId.xy + ViewportRect.xy;
 
 	float2 Velocity = PostprocessInput0[ PixelPos ].xy;
@@ -75,19 +68,19 @@ void VelocityFlattenMain(
 		Velocity.x = -2.0 * ((PosN.x * ((CameraMotion.Value[1].x * PosN.y) + (CameraMotion.Value[1].y * PosN.z) + CameraMotion.Value[1].z)) + (CameraMotion.Value[1].w * PosN.y) + (CameraMotion.Value[2].x * PosN.x * PosN.x) + (CameraMotion.Value[2].y * PosN.z) + CameraMotion.Value[2].z) * ScaleM;
 		Velocity.y =  2.0 * ((PosN.y * ((CameraMotion.Value[3].x * PosN.x) + (CameraMotion.Value[3].y * PosN.z) + CameraMotion.Value[3].z)) + (CameraMotion.Value[3].w * PosN.x) + (CameraMotion.Value[4].x * PosN.y * PosN.y) + (CameraMotion.Value[4].y * PosN.z) + CameraMotion.Value[4].z) * ScaleM;
 
-		//Velocity[i].x += 0.3;
-		//Velocity[i].y += 0.1;
+		//Velocity.x += 0.3;
+		//Velocity.y += 0.1;
 	}
 
-	const float AspectRatio = ScreenPosToPixel.y / ScreenPosToPixel.x;
-	Velocity.y *= AspectRatio;
+	const float AspectRatio = MotionBlurParameters.x;
+	Velocity.y *= -AspectRatio;
 
 	float2 VelocityPolar = CartesianToPolar( Velocity );
 	
 	bool bInsideViewport = all( PixelPos.xy < ViewportRect.zw );
 	
 #if 1
-	// 10:10:10:2  (VelocityLength, VelocityAngle, Depth)
+	// 11:11:10  (VelocityLength, VelocityAngle, Depth)
 	float2 EncodedVelocity;
 	EncodedVelocity.x = VelocityPolar.x;
 	EncodedVelocity.y = VelocityPolar.y * (0.5 / PI) + 0.5;
@@ -109,7 +102,7 @@ void VelocityFlattenMain(
 
 #if 1
 	// Limit velocity
-	VelocityPolar.x = min( VelocityPolar.x, 100 / ScreenPosToPixel.x );	// TODO input limit
+	VelocityPolar.x = min( VelocityPolar.x, MotionBlurParameters.w / MotionBlurParameters.y );
 #endif
 
 	float4 VelocityMinMax = VelocityPolar.xyxy;
@@ -185,13 +178,9 @@ void VelocityGatherCS(
 				continue;
 
 			float4 ScatterMinMax = PostprocessInput0[ SampleIndex ];
-			//float4 ScatterMinMax = PostprocessInput0.SampleLevel( PostprocessInput0Sampler, UV, 0, Offset );
 			float2 MaxVelocity = ScatterMinMax.zw;
 
-			// TODO move to velocity flatten shader
-			//float2 VelocityPixels = MaxVelocity * ScreenPosToPixel.x;// * VelocityScale.xy;
-			float2 VelocityPixels = MaxVelocity * 0.5 * 1920.0 / 16.0;
-			VelocityPixels *= 0.5;
+			float2 VelocityPixels = MaxVelocity * MotionBlurParameters.z;
 
 			float  VelocityLengthPixelsSqr = dot( VelocityPixels, VelocityPixels );
 			float  VelocityLengthPixelsInv = rsqrtFast( VelocityLengthPixelsSqr + 1e-8 );
diff --git a/Engine/Shaders/ReflectionEnvironmentShaders.usf b/Engine/Shaders/ReflectionEnvironmentShaders.usf
--- a/Engine/Shaders/ReflectionEnvironmentShaders.usf
+++ b/Engine/Shaders/ReflectionEnvironmentShaders.usf
@@ -111,6 +111,7 @@ float3 GetCubemapVector(float2 ScaledUVs)
 
 TextureCube SourceTexture;
 SamplerState SourceTextureSampler;
+float2 SinCosSourceCubemapRotation;
 
 void CopyCubemapToCubeFaceColorPS(
 	FScreenVertexOutput Input,
@@ -119,6 +120,8 @@ void CopyCubemapToCubeFaceColorPS(
 { 
 	float2 ScaledUVs = Input.UV * 2 - 1;
 	float3 CubeCoordinates = GetCubemapVector(ScaledUVs);
+	// Rotate around Z axis
+	CubeCoordinates.xy = float2(dot(CubeCoordinates.xy, float2(SinCosSourceCubemapRotation.y, -SinCosSourceCubemapRotation.x)), dot(CubeCoordinates.xy, SinCosSourceCubemapRotation));
 	OutColor = TextureCubeSampleLevel(SourceTexture, SourceTextureSampler, CubeCoordinates, 0);
 
 	if (SkyLightCaptureParameters.x > 0)
diff --git a/Engine/Source/Developer/ShaderCompilerCommon/Private/HlslAST.cpp b/Engine/Source/Developer/ShaderCompilerCommon/Private/HlslAST.cpp
--- a/Engine/Source/Developer/ShaderCompilerCommon/Private/HlslAST.cpp
+++ b/Engine/Source/Developer/ShaderCompilerCommon/Private/HlslAST.cpp
@@ -23,7 +23,10 @@ namespace CrossCompiler
 				for (const auto* Dimension : ArraySize)
 				{
 					Writer << (TCHAR)'[';
-					Dimension->Write(Writer);
+					if (Dimension)
+					{
+						Dimension->Write(Writer);
+					}
 					Writer << (TCHAR)']';
 				}
 			}
@@ -87,7 +90,7 @@ namespace CrossCompiler
 				Writer << TEXT("+");
 				break;
 
-			case EOperators::Neg:
+			case EOperators::Minus:
 				Writer << TEXT("-");
 				break;
 
@@ -163,6 +166,10 @@ namespace CrossCompiler
 				Writer << TEXT("&");
 				break;
 
+			case EOperators::BitNeg:
+				Writer << TEXT("~");
+				break;
+
 			case EOperators::Equal:
 				Writer << TEXT("==");
 				break;
@@ -256,7 +263,7 @@ namespace CrossCompiler
 				Writer << TEXT("*MISSING_");
 				Writer << (uint32)Operator;
 				Writer << (TCHAR)'*';
-				check(0);
+				checkf(0, TEXT("Unhandled AST Operator %d!"), Operator);
 				break;
 			}
 		}
@@ -279,7 +286,7 @@ namespace CrossCompiler
 				Writer << TEXT("*MISSING_");
 				Writer << (uint32)Operator;
 				Writer << (TCHAR)'*';
-				check(0);
+				checkf(0, TEXT("Unhandled AST Operator %d!"), Operator);
 				break;
 			}
 		}
@@ -604,7 +611,7 @@ namespace CrossCompiler
 				Writer << TEXT("*MISSING_");
 				Writer << (uint32)Type;
 				Writer << (TCHAR)'*';
-				check(0);
+				checkf(0, TEXT("Unhandled AST jump type %d!"), Type);
 				break;
 			}
 
@@ -950,6 +957,8 @@ namespace CrossCompiler
 		FParameterDeclarator::FParameterDeclarator(FLinearAllocator* InAllocator, const FSourceInfo& InInfo) :
 			FNode(InAllocator, InInfo),
 			Type(nullptr),
+			Identifier(nullptr),
+			Semantic(nullptr),
 			bIsArray(false),
 			ArraySize(InAllocator),
 			DefaultValue(nullptr)
@@ -1094,7 +1103,7 @@ namespace CrossCompiler
 				break;
 
 			default:
-				check(0);
+				checkf(0, TEXT("Unhandled AST iteration type %d!"), Type);
 				break;
 			}
 		}
diff --git a/Engine/Source/Developer/ShaderCompilerCommon/Private/HlslAST.h b/Engine/Source/Developer/ShaderCompilerCommon/Private/HlslAST.h
--- a/Engine/Source/Developer/ShaderCompilerCommon/Private/HlslAST.h
+++ b/Engine/Source/Developer/ShaderCompilerCommon/Private/HlslAST.h
@@ -19,6 +19,7 @@ namespace CrossCompiler
 		struct FParameterDeclarator;
 		struct FUnaryExpression;
 		struct FAttribute;
+		struct FJumpStatement;
 
 		struct FASTWriter
 		{
@@ -107,6 +108,7 @@ namespace CrossCompiler
 			virtual FFunctionDefinition* AsFunctionDefinition() { return nullptr; }
 			virtual FParameterDeclarator* AsParameterDeclarator() { return nullptr; }
 			virtual FUnaryExpression* AsUnaryExpression() { return nullptr; }
+			virtual FJumpStatement* AsJumpStatement() { return nullptr; }
 
 			// Returns true if the expression can be evaluated to a constant int
 			virtual bool GetConstantIntValue(int32& OutValue) const { return false; }
@@ -150,7 +152,7 @@ namespace CrossCompiler
 		{
 			Assign,
 			Plus,        /**< Unary + operator. */
-			Neg,
+			Minus,
 			Add,
 			Sub,
 			Mul,
@@ -167,11 +169,11 @@ namespace CrossCompiler
 			BitAnd,
 			BitXor,
 			BitOr,
-			BitNot,
+			BitNeg,		// ~
 			LogicAnd,
 			LogicXor,
 			LogicOr,
-			LogicNot,
+			LogicNot,	// !
 
 			MulAssign,
 			DivAssign,
@@ -687,6 +689,7 @@ namespace CrossCompiler
 			FExpression* OptionalExpression;
 
 			virtual void Write(FASTWriter& Writer) const override;
+			virtual FJumpStatement* AsJumpStatement() override { return this; }
 		};
 
 		struct FFunctionDefinition : public FNode
diff --git a/Engine/Source/Developer/ShaderCompilerCommon/Private/HlslExpressionParser.inl b/Engine/Source/Developer/ShaderCompilerCommon/Private/HlslExpressionParser.inl
--- a/Engine/Source/Developer/ShaderCompilerCommon/Private/HlslExpressionParser.inl
+++ b/Engine/Source/Developer/ShaderCompilerCommon/Private/HlslExpressionParser.inl
@@ -502,7 +502,7 @@ namespace CrossCompiler
 			case EHlslToken::Minus:
 				Scanner.Advance();
 				bFoundAny = true;
-				Operator = AST::EOperators::Neg;
+				Operator = AST::EOperators::Minus;
 				break;
 
 			case EHlslToken::Not:
@@ -514,7 +514,7 @@ namespace CrossCompiler
 			case EHlslToken::Neg:
 				Scanner.Advance();
 				bFoundAny = true;
-				Operator = AST::EOperators::BitNot;
+				Operator = AST::EOperators::BitNeg;
 				break;
 
 			case EHlslToken::LeftParenthesis:
diff --git a/Engine/Source/Developer/ShaderCompilerCommon/Private/HlslUtils.cpp b/Engine/Source/Developer/ShaderCompilerCommon/Private/HlslUtils.cpp
--- a/Engine/Source/Developer/ShaderCompilerCommon/Private/HlslUtils.cpp
+++ b/Engine/Source/Developer/ShaderCompilerCommon/Private/HlslUtils.cpp
@@ -276,11 +276,12 @@ struct FRemoveUnusedOutputs
 		return ReturnType;
 	};
 
+	// Case-insensitive when working with Semantics
 	bool IsStringInArray(const TArray<FString>& Array, const TCHAR* Semantic)
 	{
 		for (const FString& String : Array)
 		{
-			if (FCString::Strcmp(*String, Semantic) == 0)
+			if (FCString::Stricmp(*String, Semantic) == 0)
 			{
 				return true;
 			}
@@ -341,14 +342,26 @@ struct FRemoveUnusedOutputs
 	void WriteGeneratedCode(CrossCompiler::AST::FFunctionDefinition* NewEntryFunction, TArray<CrossCompiler::AST::FStructSpecifier*>& NewStructs, FString& OutGeneratedCode)
 	{
 		CrossCompiler::AST::FASTWriter Writer(OutGeneratedCode);
-		GeneratedCode = TEXT("#line 1 \"RemoveUnusedOutputs.usf\"\n// Generated Entry Point: \n");
+		GeneratedCode = TEXT("#line 1 \"RemoveUnusedOutputs.usf\"\n// Generated Entry Point: ");
+		GeneratedCode += NewEntryFunction->Prototype->Identifier;
+		GeneratedCode += TEXT("\n");
+		if (UsedOutputs.Num() > 0)
+		{
+			GeneratedCode += TEXT("// Requested UsedOutputs:");
+			for (int32 Index = 0; Index < UsedOutputs.Num(); ++Index)
+			{
+				GeneratedCode += (Index == 0) ? TEXT(" ") : TEXT(", ");
+				GeneratedCode += UsedOutputs[Index];
+			}
+			GeneratedCode += TEXT("\n");
+		}
 		if (RemovedSemantics.Num() > 0)
 		{
-			GeneratedCode = TEXT("// Removed Outputs: ");
+			GeneratedCode += TEXT("// Removed Outputs:");
 			for (int32 Index = 0; Index < RemovedSemantics.Num(); ++Index)
 			{
+				GeneratedCode += (Index == 0) ? TEXT(" ") : TEXT(", ");
 				GeneratedCode += RemovedSemantics[Index];
-				GeneratedCode += TEXT(" ");
 			}
 			GeneratedCode += TEXT("\n");
 		}
diff --git a/Engine/Source/Developer/Windows/ShaderFormatD3D/Private/D3D11ShaderCompiler.cpp b/Engine/Source/Developer/Windows/ShaderFormatD3D/Private/D3D11ShaderCompiler.cpp
--- a/Engine/Source/Developer/Windows/ShaderFormatD3D/Private/D3D11ShaderCompiler.cpp
+++ b/Engine/Source/Developer/Windows/ShaderFormatD3D/Private/D3D11ShaderCompiler.cpp
@@ -32,11 +32,31 @@
 	TEXT("r.D3DRemoveUnusedInterpolators"),
 	GD3DAllowRemoveUnused,
 	TEXT("Enables removing unused interpolators mode when compiling pipelines for D3D.\n")
+	TEXT(" -1: Do not actually remove, but make the app think it did (for debugging)\n")
 	TEXT(" 0: Disable (default)\n")
-	TEXT(" 1: Enable)"),
+	TEXT(" 1: Enable removing unused"),
 	ECVF_Default
 	);
 
+static int32 GD3DCheckForDoubles = 1;
+static FAutoConsoleVariableRef CVarD3DCheckForDoubles(
+	TEXT("r.D3DCheckShadersForDouble"),
+	GD3DCheckForDoubles,
+	TEXT("Enables checking the D3D microcode for uses of double. This is not allowed on all D3D11 cards.\n")
+	TEXT(" 0: Do not check for faster compilation\n")
+	TEXT(" 1: Enable checking and error if found (default)"),
+	ECVF_Default
+	);
 
+static int32 GD3DDumpAMDCodeXLFile = 1;
+static FAutoConsoleVariableRef CVarD3DDumpAMDCodeXLFile(
+	TEXT("r.D3DDumpAMDCodeXLFile"),
+	GD3DDumpAMDCodeXLFile,
+	TEXT("When r.DumpShaderDebugInfo is enabled, this will generate a batch file for running CodeXL.\n")
+	TEXT(" 0: Do not generate extra batch file (default)\n")
+	TEXT(" 1: Enable generating extra batch file"),
+	ECVF_Default
+	);
+
 /**
  * TranslateCompilerFlag - translates the platform-independent compiler flags into D3DX defines
@@ -164,6 +184,12 @@
 		FXCCommandline += FString(TEXT(" /Od"));
 	}
 
+	if (CompileFlags & D3D10_SHADER_SKIP_VALIDATION)
+	{
+		CompileFlags &= ~D3D10_SHADER_SKIP_VALIDATION;
+		FXCCommandline += FString(TEXT(" /Vd"));
+	}
+
 	if(CompileFlags & D3D10_SHADER_AVOID_FLOW_CONTROL)
 	{
 		CompileFlags &= ~D3D10_SHADER_AVOID_FLOW_CONTROL;
@@ -182,17 +208,21 @@
 		FXCCommandline += FString(TEXT(" /Gec"));
 	}
 
-	if(CompileFlags & D3D10_SHADER_OPTIMIZATION_LEVEL1)
+	if ((CompileFlags & D3D10_SHADER_OPTIMIZATION_LEVEL2) == D3D10_SHADER_OPTIMIZATION_LEVEL2)
 	{
-		CompileFlags &= ~D3D10_SHADER_OPTIMIZATION_LEVEL1;
-		FXCCommandline += FString(TEXT(" /O1"));
+		CompileFlags &= ~D3D10_SHADER_OPTIMIZATION_LEVEL2;
+		FXCCommandline += FString(TEXT(" /O2"));
 	}
-
-	if(CompileFlags & D3D10_SHADER_OPTIMIZATION_LEVEL3)
+	else if (CompileFlags & D3D10_SHADER_OPTIMIZATION_LEVEL3)
 	{
 		CompileFlags &= ~D3D10_SHADER_OPTIMIZATION_LEVEL3;
 		FXCCommandline += FString(TEXT(" /O3"));
 	}
+	else if (CompileFlags & D3D10_SHADER_OPTIMIZATION_LEVEL1)
+	{
+		CompileFlags &= ~D3D10_SHADER_OPTIMIZATION_LEVEL1;
+		FXCCommandline += FString(TEXT(" /O1"));
+	}
 
 	checkf(CompileFlags == 0, TEXT("Unhandled d3d11 shader compiler flag!"));
 
@@ -317,218 +347,105 @@
 #endif
 }
 
-void CompileD3D11Shader(const FShaderCompilerInput& Input,FShaderCompilerOutput& Output, FShaderCompilerDefinitions& AdditionalDefines, const FString& WorkingDirectory)
+// Utility variable so we can place a breakpoint while debugging
+static int32 GBreakpoint = 0;
+
+// Generate the dumped usf file; call the D3D compiler, gather reflection information and generate the output data
+static bool CompileAndProcessD3DShader(FString& PreprocessedShaderSource, const FString& CompilerPath,
+	uint32 CompileFlags, const FShaderCompilerInput& Input, FString& EntryPointName,
+	const TCHAR* ShaderProfile, TArray<FString>& FilteredErrors, FShaderCompilerOutput& Output)
 {
-	FString PreprocessedShaderSource;
-	FString CompilerPath;
-	const TCHAR* ShaderProfile = GetShaderProfileName(Input.Target);
-
-	if(!ShaderProfile)
-	{
-		Output.Errors.Add(FShaderCompilerError(TEXT("Unrecognized shader frequency")));
-		return;
-	}
-
-	// Set additional defines.
-	AdditionalDefines.SetDefine(TEXT("COMPILER_HLSL"), 1);
-
-	// Preprocess the shader.
-	if (PreprocessShader(PreprocessedShaderSource, Output, Input, AdditionalDefines) != true)
-	{
-		// The preprocessing stage will add any relevant errors.
-		return;
-	}
-
-	FString EntryPointName = Input.EntryPointName;
-
-	Output.bFailedRemovingUnused = false;
-	if (GD3DAllowRemoveUnused && Input.Target.Frequency == SF_Vertex && Input.bCompilingForShaderPipeline)
-	{
-		static TArray<FString> VertexSystemOutputs;
-		if (VertexSystemOutputs.Num() == 0)
-		{
-			VertexSystemOutputs.Add(TEXT("SV_POSITION"));
-			VertexSystemOutputs.Add(TEXT("SV_Position"));
-			VertexSystemOutputs.Add(TEXT("SV_ClipDistance"));
-			VertexSystemOutputs.Add(TEXT("SV_ClipDistance0"));
-			VertexSystemOutputs.Add(TEXT("SV_ClipDistance1"));
-			VertexSystemOutputs.Add(TEXT("SV_ClipDistance2"));
-			VertexSystemOutputs.Add(TEXT("SV_ClipDistance3"));
-			VertexSystemOutputs.Add(TEXT("SV_ClipDistance4"));
-			VertexSystemOutputs.Add(TEXT("SV_ClipDistance5"));
-			VertexSystemOutputs.Add(TEXT("SV_CullDistance0"));
-			VertexSystemOutputs.Add(TEXT("SV_CullDistance1"));
-			VertexSystemOutputs.Add(TEXT("SV_CullDistance2"));
-			VertexSystemOutputs.Add(TEXT("SV_CullDistance3"));
-			VertexSystemOutputs.Add(TEXT("SV_CullDistance4"));
-			VertexSystemOutputs.Add(TEXT("SV_CullDistance5"));
-			VertexSystemOutputs.Add(TEXT("SV_Coverage"));
-			VertexSystemOutputs.Add(TEXT("SV_Depth"));
-			VertexSystemOutputs.Add(TEXT("SV_DomainLocation"));
-			VertexSystemOutputs.Add(TEXT("SV_IsFrontFace"));
-			VertexSystemOutputs.Add(TEXT("SV_OutputControlPointID"));
-			VertexSystemOutputs.Add(TEXT("SV_RenderTargetArrayIndex"));
-			VertexSystemOutputs.Add(TEXT("SV_SampleIndex"));
-			VertexSystemOutputs.Add(TEXT("SV_ViewportArrayIndex"));
-			VertexSystemOutputs.Add(TEXT("SV_TessFactor"));
-			VertexSystemOutputs.Add(TEXT("SV_InstanceID"));
-			VertexSystemOutputs.Add(TEXT("SV_PrimitiveID"));
-			VertexSystemOutputs.Add(TEXT("SV_VertexID"));
-		}
-
-		TArray<FString> Errors;
-		if (RemoveUnusedOutputs(PreprocessedShaderSource, VertexSystemOutputs, Input.UsedOutputs, Input.EntryPointName, Errors))
-		{
-			EntryPointName += FString(TEXT("__OPTIMIZED"));
-		}
-		else
-		{
-			UE_LOG(LogD3D11ShaderCompiler, Warning, TEXT("Failed to Remove unused outputs [%s]!"), *Input.DumpDebugInfoPath);
-			for (int32 Index = 0; Index < Errors.Num(); ++Index)
-			{
-				FShaderCompilerError NewError;
-				NewError.StrippedErrorMessage = Errors[Index];
-				Output.Errors.Add(NewError);
-			}
-			Output.bFailedRemovingUnused = true;
-		}
-	}
-
-	if (!RemoveUniformBuffersFromSource(PreprocessedShaderSource))
-	{
-		return;
-	}
-
-	// Search definitions for a custom D3D compiler path.
-	for(TMap<FString,FString>::TConstIterator DefinitionIt(Input.Environment.GetDefinitions());DefinitionIt;++DefinitionIt)
-	{
-		const FString& Name = DefinitionIt.Key();
-		const FString& Definition = DefinitionIt.Value();
-
-		if(Name == TEXT("D3DCOMPILER_PATH"))
-		{
-			CompilerPath = Definition;
-		}
-	}
-
-	// @TODO - currently d3d11 uses d3d10 shader compiler flags... update when this changes in DXSDK
-	// @TODO - implement different material path to allow us to remove backwards compat flag on sm5 shaders
-	uint32 CompileFlags = D3D10_SHADER_ENABLE_BACKWARDS_COMPATIBILITY
-		// Unpack uniform matrices as row-major to match the CPU layout.
-		| D3D10_SHADER_PACK_MATRIX_ROW_MAJOR;
-
-	if (DEBUG_SHADERS || Input.Environment.CompilerFlags.Contains(CFLAG_Debug)) 
-	{
-		//add the debug flags
-		CompileFlags |= D3D10_SHADER_DEBUG | D3D10_SHADER_SKIP_OPTIMIZATION;
-	}
-	else
-	{
-		if (Input.Environment.CompilerFlags.Contains(CFLAG_StandardOptimization))
-		{
-			CompileFlags |= D3D10_SHADER_OPTIMIZATION_LEVEL1;
-		}
-		else
-		{
-			CompileFlags |= D3D10_SHADER_OPTIMIZATION_LEVEL3;
-		}
-	}
-
-	for(int32 FlagIndex = 0;FlagIndex < Input.Environment.CompilerFlags.Num();FlagIndex++)
-	{
-		//accumulate flags set by the shader
-		CompileFlags |= TranslateCompilerFlagD3D11((ECompilerFlags)Input.Environment.CompilerFlags[FlagIndex]);
-	}
-
 	auto AnsiSourceFile = StringCast<ANSICHAR>(*PreprocessedShaderSource);
-	TArray<FString> FilteredErrors;
+
+	// Write out the preprocessed file and a batch file to compile it if requested (DumpDebugInfoPath is valid)
+	if (Input.DumpDebugInfoPath.Len() > 0 && IFileManager::Get().DirectoryExists(*Input.DumpDebugInfoPath))
 	{
-		TRefCountPtr<ID3DBlob> Shader;
-		TRefCountPtr<ID3DBlob> Errors;
+		FArchive* FileWriter = IFileManager::Get().CreateFileWriter(*(Input.DumpDebugInfoPath / Input.SourceFilename + TEXT(".usf")));
+		if (FileWriter)
+		{
+			FileWriter->Serialize((ANSICHAR*)AnsiSourceFile.Get(), AnsiSourceFile.Length());
+			FileWriter->Close();
+			delete FileWriter;
+		}
 
-		HRESULT Result;
-		pD3DCompile D3DCompileFunc = GetD3DCompileFunc(CompilerPath);
-			
-		if(D3DCompileFunc)
+		const FString BatchFileContents = D3D11CreateShaderCompileCommandLine((Input.SourceFilename + TEXT(".usf")), *EntryPointName, ShaderProfile, CompileFlags);
+		FFileHelper::SaveStringToFile(BatchFileContents, *(Input.DumpDebugInfoPath / TEXT("CompileD3D.bat")));
+
+		if (GD3DDumpAMDCodeXLFile)
 		{
-			bool bException = false;
+			const FString BatchFileContents2 = CreateAMDCodeXLCommandLine((Input.SourceFilename + TEXT(".usf")), *EntryPointName, ShaderProfile, CompileFlags);
+			FFileHelper::SaveStringToFile(BatchFileContents2, *(Input.DumpDebugInfoPath / TEXT("CompileAMD.bat")));
+		}
+	}
 
-			Result = D3DCompileWrapper(
-				D3DCompileFunc,
-				bException,
-				AnsiSourceFile.Get(),
-				AnsiSourceFile.Length(),
-				TCHAR_TO_ANSI(*Input.SourceFilename),
-				/*pDefines=*/ NULL,
-				/*pInclude=*/ NULL,
-				TCHAR_TO_ANSI(*EntryPointName),
-				TCHAR_TO_ANSI(ShaderProfile),
-				CompileFlags,
-				0,
-				Shader.GetInitReference(),
-				Errors.GetInitReference()
+	TRefCountPtr<ID3DBlob> Shader;
+	TRefCountPtr<ID3DBlob> Errors;
+
+	HRESULT Result;
+	pD3DCompile D3DCompileFunc = GetD3DCompileFunc(CompilerPath);
+
+	if (D3DCompileFunc)
+	{
+		bool bException = false;
+
+		Result = D3DCompileWrapper(
+			D3DCompileFunc,
+			bException,
+			AnsiSourceFile.Get(),
+			AnsiSourceFile.Length(),
+			TCHAR_TO_ANSI(*Input.SourceFilename),
+			/*pDefines=*/ NULL,
+			/*pInclude=*/ NULL,
+			TCHAR_TO_ANSI(*EntryPointName),
+			TCHAR_TO_ANSI(ShaderProfile),
+			CompileFlags,
+			0,
+			Shader.GetInitReference(),
+			Errors.GetInitReference()
 			);
 
-			if( bException )
-			{
-				FilteredErrors.Add( TEXT("D3DCompile exception") );
-			}
+		if (bException)
+		{
+			FilteredErrors.Add(TEXT("D3DCompile exception"));
 		}
-		else
+	}
+	else
+	{
+		FilteredErrors.Add(FString::Printf(TEXT("Couldn't find shader compiler: %s"), *CompilerPath));
+		Result = E_FAIL;
+	}
+
+	// Filter any errors.
+	void* ErrorBuffer = Errors ? Errors->GetBufferPointer() : NULL;
+	if (ErrorBuffer)
+	{
+		D3D11FilterShaderCompileWarnings(ANSI_TO_TCHAR(ErrorBuffer), FilteredErrors);
+	}
+
+	// Fail the compilation if double operations are being used, since those are not supported on all D3D11 cards
+	if (SUCCEEDED(Result))
+	{
+		if (GD3DCheckForDoubles)
 		{
-			FilteredErrors.Add(FString::Printf(TEXT("Couldn't find shader compiler: %s"), *CompilerPath));
-			Result = E_FAIL;
-		}
-
-		// Filter any errors.
-		void* ErrorBuffer = Errors ? Errors->GetBufferPointer() : NULL;
-		if (ErrorBuffer)
-		{
-			D3D11FilterShaderCompileWarnings(ANSI_TO_TCHAR(ErrorBuffer), FilteredErrors);
-		}
-
-		// Fail the compilation if double operations are being used, since those are not supported on all D3D11 cards
-		if (SUCCEEDED(Result))
-		{
 			TRefCountPtr<ID3DBlob> Dissasembly;
-
 			if (SUCCEEDED(D3DDisassemble(Shader->GetBufferPointer(), Shader->GetBufferSize(), 0, "", Dissasembly.GetInitReference())))
 			{
 				ANSICHAR* DissasemblyString = new ANSICHAR[Dissasembly->GetBufferSize() + 1];
 				FMemory::Memcpy(DissasemblyString, Dissasembly->GetBufferPointer(), Dissasembly->GetBufferSize());
 				DissasemblyString[Dissasembly->GetBufferSize()] = 0;
 				FString DissasemblyStringW(DissasemblyString);
-				delete [] DissasemblyString;
+				delete[] DissasemblyString;
 
 				// dcl_globalFlags will contain enableDoublePrecisionFloatOps when the shader uses doubles, even though the docs on dcl_globalFlags don't say anything about this
 				if (DissasemblyStringW.Contains(TEXT("enableDoublePrecisionFloatOps")))
 				{
 					FilteredErrors.Add(TEXT("Shader uses double precision floats, which are not supported on all D3D11 hardware!"));
-					Result = E_FAIL;
+					return false;
 				}
 			}
 		}
 
-		// Write out the preprocessed file and a batch file to compile it if requested (DumpDebugInfoPath is valid)
-		if (Input.DumpDebugInfoPath != TEXT("") 
-			&& IFileManager::Get().DirectoryExists(*Input.DumpDebugInfoPath))
-		{
-			FArchive* FileWriter = IFileManager::Get().CreateFileWriter(*(Input.DumpDebugInfoPath / Input.SourceFilename + TEXT(".usf")));
-
-			if (FileWriter)
-			{
-				FileWriter->Serialize((ANSICHAR*)AnsiSourceFile.Get(), AnsiSourceFile.Length());
-				FileWriter->Close();
-				delete FileWriter;
-			}
-
-			const FString BatchFileContents = D3D11CreateShaderCompileCommandLine((Input.SourceFilename + TEXT(".usf")), *EntryPointName, ShaderProfile, CompileFlags);
-			FFileHelper::SaveStringToFile(BatchFileContents, *(Input.DumpDebugInfoPath / TEXT("CompileD3D.bat")));
-
-			const FString BatchFileContents2 = CreateAMDCodeXLCommandLine((Input.SourceFilename + TEXT(".usf")), *EntryPointName, ShaderProfile, CompileFlags);
-			FFileHelper::SaveStringToFile(BatchFileContents2, *(Input.DumpDebugInfoPath / TEXT("CompileAMD.bat")));
-		}
-
+		// Gather reflection information
 		int32 NumInterpolants = 0;
 		TIndirectArray<FString> InterpolantNames;
 		TArray<FString> ShaderInputs;
@@ -537,10 +454,10 @@
 			Output.bSucceeded = true;
 
 			ID3D11ShaderReflection* Reflector = NULL;
-			Result = D3DReflect( Shader->GetBufferPointer(), Shader->GetBufferSize(), IID_ID3D11ShaderReflection, (void**) &Reflector );
-			if(FAILED(Result))
+			Result = D3DReflect(Shader->GetBufferPointer(), Shader->GetBufferSize(), IID_ID3D11ShaderReflection, (void**)&Reflector);
+			if (FAILED(Result))
 			{
-				UE_LOG(LogD3D11ShaderCompiler, Fatal,TEXT("D3DReflect failed: Result=%08x"),Result);
+				UE_LOG(LogD3D11ShaderCompiler, Fatal, TEXT("D3DReflect failed: Result=%08x"), Result);
 			}
 
 			// Read the constant table description.
@@ -556,7 +473,7 @@
 			TArray<FString> ShaderOutputs;
 
 			TBitArray<> UsedUniformBufferSlots;
-			UsedUniformBufferSlots.Init(false,32);
+			UsedUniformBufferSlots.Init(false, 32);
 
 			if (Input.Target.Frequency == SF_Vertex)
 			{
@@ -567,7 +484,7 @@
 					if (ParamDesc.SystemValueType == D3D_NAME_UNDEFINED && ParamDesc.Mask != 0)
 					{
 						++NumInterpolants;
-						new(InterpolantNames) FString(FString::Printf(TEXT("%s%d"), ANSI_TO_TCHAR(ParamDesc.SemanticName),ParamDesc.SemanticIndex));
+						new(InterpolantNames) FString(FString::Printf(TEXT("%s%d"), ANSI_TO_TCHAR(ParamDesc.SemanticName), ParamDesc.SemanticIndex));
 						ShaderOutputs.Add(*InterpolantNames.Last());
 					}
 				}
@@ -618,7 +535,7 @@
 									VariableDesc.StartOffset,
 									VariableDesc.Size
 									);
-									UsedUniformBufferSlots[CBIndex] = true;
+								UsedUniformBufferSlots[CBIndex] = true;
 							}
 						}
 					}
@@ -656,7 +573,7 @@
 						TCHAR *BracketLocation = FCString::Strchr(OfficialName, TEXT('['));
 						if (BracketLocation)
 						{
-							*BracketLocation = 0;	
+							*BracketLocation = 0;
 
 							const int32 NumCharactersBeforeArray = BracketLocation - OfficialName;
 
@@ -681,7 +598,7 @@
 									break;
 								}
 							}
-						}		
+						}
 					}
 
 					if (BindDesc.Type == D3D10_SIT_SAMPLER)
@@ -701,22 +618,22 @@
 						BindCount
 						);
 				}
-				else if (BindDesc.Type == D3D11_SIT_UAV_RWTYPED || BindDesc.Type == D3D11_SIT_UAV_RWSTRUCTURED || 
-						 BindDesc.Type == D3D11_SIT_UAV_RWBYTEADDRESS || BindDesc.Type == D3D11_SIT_UAV_RWSTRUCTURED_WITH_COUNTER ||
-						 BindDesc.Type == D3D11_SIT_UAV_APPEND_STRUCTURED)
+				else if (BindDesc.Type == D3D11_SIT_UAV_RWTYPED || BindDesc.Type == D3D11_SIT_UAV_RWSTRUCTURED ||
+					BindDesc.Type == D3D11_SIT_UAV_RWBYTEADDRESS || BindDesc.Type == D3D11_SIT_UAV_RWSTRUCTURED_WITH_COUNTER ||
+					BindDesc.Type == D3D11_SIT_UAV_APPEND_STRUCTURED)
 				{
-						TCHAR OfficialName[1024];
-						FCString::Strcpy(OfficialName, ANSI_TO_TCHAR(BindDesc.Name));
+					TCHAR OfficialName[1024];
+					FCString::Strcpy(OfficialName, ANSI_TO_TCHAR(BindDesc.Name));
 
-						Output.ParameterMap.AddParameterAllocation(
-							OfficialName,
-							0,
-							BindDesc.BindPoint,
-							1
-							);
+					Output.ParameterMap.AddParameterAllocation(
+						OfficialName,
+						0,
+						BindDesc.BindPoint,
+						1
+						);
 
-						NumUAVs = FMath::Max(NumUAVs, BindDesc.BindPoint + BindDesc.BindCount);
-					}
+					NumUAVs = FMath::Max(NumUAVs, BindDesc.BindPoint + BindDesc.BindCount);
+				}
 				else if (BindDesc.Type == D3D11_SIT_STRUCTURED || BindDesc.Type == D3D11_SIT_BYTEADDRESS)
 				{
 					TCHAR OfficialName[1024];
@@ -752,10 +669,10 @@
 
 				if (FAILED(Result))
 				{
-					UE_LOG(LogD3D11ShaderCompiler, Fatal,TEXT("D3DStripShader failed: Result=%08x"),Result);
+					UE_LOG(LogD3D11ShaderCompiler, Fatal, TEXT("D3DStripShader failed: Result=%08x"), Result);
 				}
 			}
-			
+
 			// Build the SRT for this shader.
 			FD3D11ShaderResourceTable SRT;
 
@@ -798,12 +715,16 @@
 				BuildResourceTableTokenStream(GenericSRT.UnorderedAccessViewMap, GenericSRT.MaxBoundResourceTable, SRT.UnorderedAccessViewMap);
 			}
 
-			if (GD3DAllowRemoveUnused && Input.Target.Frequency == SF_Pixel && Input.bCompilingForShaderPipeline)
+			if (GD3DAllowRemoveUnused != 0 && Input.Target.Frequency == SF_Pixel && Input.bCompilingForShaderPipeline)
 			{
 				Output.bSupportsQueryingUsedAttributes = true;
-				Output.UsedAttributes = ShaderInputs;
+				if (GD3DAllowRemoveUnused == 1)
+				{
+					Output.UsedAttributes = ShaderInputs;
+				}
 			}
 
+			// Generate the final Output
 			FMemoryWriter Ar(Output.ShaderCode.GetWriteAccess(), true);
 			Ar << SRT;
 			Ar.Serialize(CompressedData->GetBufferPointer(), CompressedData->GetBufferSize());
@@ -847,50 +768,186 @@
 				FString InterpolantsStr;
 				for (int32 i = 0; i < InterpolantNames.Num(); ++i)
 				{
-					InterpolantsStr += FString::Printf(TEXT("\n\t%s"),*InterpolantNames[i]);
+					InterpolantsStr += FString::Printf(TEXT("\n\t%s"), *InterpolantNames[i]);
 				}
 				FilteredErrors.Add(FString::Printf(TEXT("Shader uses more than 8 interpolants which is not supported by ES2!  Used: %u%s"), NumInterpolants, *InterpolantsStr));
 				Result = E_FAIL;
 				Output.bSucceeded = false;
 			}
 		}
+	}
+	else
+	{
+		++GBreakpoint;
+	}
 
-		if(FAILED(Result) && !FilteredErrors.Num())
+	return SUCCEEDED(Result);
+}
+
+void CompileD3D11Shader(const FShaderCompilerInput& Input,FShaderCompilerOutput& Output, FShaderCompilerDefinitions& AdditionalDefines, const FString& WorkingDirectory)
+{
+	FString PreprocessedShaderSource;
+	FString CompilerPath;
+	const TCHAR* ShaderProfile = GetShaderProfileName(Input.Target);
+
+	if(!ShaderProfile)
+	{
+		Output.Errors.Add(FShaderCompilerError(TEXT("Unrecognized shader frequency")));
+		return;
+	}
+
+	// Set additional defines.
+	AdditionalDefines.SetDefine(TEXT("COMPILER_HLSL"), 1);
+
+	// Preprocess the shader.
+	if (PreprocessShader(PreprocessedShaderSource, Output, Input, AdditionalDefines) != true)
+	{
+		// The preprocessing stage will add any relevant errors.
+		return;
+	}
+
+	FString EntryPointName = Input.EntryPointName;
+
+	Output.bFailedRemovingUnused = false;
+	if (GD3DAllowRemoveUnused && Input.Target.Frequency == SF_Vertex && Input.bCompilingForShaderPipeline)
+	{
+		static TArray<FString> VertexSystemOutputs;
+		if (VertexSystemOutputs.Num() == 0)
 		{
+			VertexSystemOutputs.Add(TEXT("SV_POSITION"));
+			VertexSystemOutputs.Add(TEXT("SV_Position"));
+			VertexSystemOutputs.Add(TEXT("SV_ClipDistance"));
+			VertexSystemOutputs.Add(TEXT("SV_ClipDistance0"));
+			VertexSystemOutputs.Add(TEXT("SV_ClipDistance1"));
+			VertexSystemOutputs.Add(TEXT("SV_ClipDistance2"));
+			VertexSystemOutputs.Add(TEXT("SV_ClipDistance3"));
+			VertexSystemOutputs.Add(TEXT("SV_ClipDistance4"));
+			VertexSystemOutputs.Add(TEXT("SV_ClipDistance5"));
+			VertexSystemOutputs.Add(TEXT("SV_CullDistance0"));
+			VertexSystemOutputs.Add(TEXT("SV_CullDistance1"));
+			VertexSystemOutputs.Add(TEXT("SV_CullDistance2"));
+			VertexSystemOutputs.Add(TEXT("SV_CullDistance3"));
+			VertexSystemOutputs.Add(TEXT("SV_CullDistance4"));
+			VertexSystemOutputs.Add(TEXT("SV_CullDistance5"));
+			VertexSystemOutputs.Add(TEXT("SV_Coverage"));
+			VertexSystemOutputs.Add(TEXT("SV_Depth"));
+			VertexSystemOutputs.Add(TEXT("SV_DomainLocation"));
+			VertexSystemOutputs.Add(TEXT("SV_IsFrontFace"));
+			VertexSystemOutputs.Add(TEXT("SV_OutputControlPointID"));
+			VertexSystemOutputs.Add(TEXT("SV_RenderTargetArrayIndex"));
+			VertexSystemOutputs.Add(TEXT("SV_SampleIndex"));
+			VertexSystemOutputs.Add(TEXT("SV_ViewportArrayIndex"));
+			VertexSystemOutputs.Add(TEXT("SV_TessFactor"));
+			VertexSystemOutputs.Add(TEXT("SV_InstanceID"));
+			VertexSystemOutputs.Add(TEXT("SV_PrimitiveID"));
+			VertexSystemOutputs.Add(TEXT("SV_VertexID"));
+		}
+
+		TArray<FString> Errors;
+		if (RemoveUnusedOutputs(PreprocessedShaderSource, VertexSystemOutputs, Input.UsedOutputs, Input.EntryPointName, Errors))
+		{
+			EntryPointName += FString(TEXT("__OPTIMIZED"));
+		}
+		else
+		{
+			UE_LOG(LogD3D11ShaderCompiler, Warning, TEXT("Failed to Remove unused outputs [%s]!"), *Input.DumpDebugInfoPath);
+			for (int32 Index = 0; Index < Errors.Num(); ++Index)
+			{
+				FShaderCompilerError NewError;
+				NewError.StrippedErrorMessage = Errors[Index];
+				Output.Errors.Add(NewError);
+			}
+			Output.bFailedRemovingUnused = true;
+		}
+	}
+
+	if (!RemoveUniformBuffersFromSource(PreprocessedShaderSource))
+	{
+		return;
+	}
+
+	// Search definitions for a custom D3D compiler path.
+	for(TMap<FString,FString>::TConstIterator DefinitionIt(Input.Environment.GetDefinitions());DefinitionIt;++DefinitionIt)
+	{
+		const FString& Name = DefinitionIt.Key();
+		const FString& Definition = DefinitionIt.Value();
+
+		if(Name == TEXT("D3DCOMPILER_PATH"))
+		{
+			CompilerPath = Definition;
+			break;
+		}
+	}
+
+	// @TODO - currently d3d11 uses d3d10 shader compiler flags... update when this changes in DXSDK
+	// @TODO - implement different material path to allow us to remove backwards compat flag on sm5 shaders
+	uint32 CompileFlags = D3D10_SHADER_ENABLE_BACKWARDS_COMPATIBILITY
+		// Unpack uniform matrices as row-major to match the CPU layout.
+		| D3D10_SHADER_PACK_MATRIX_ROW_MAJOR;
+
+	if (DEBUG_SHADERS || Input.Environment.CompilerFlags.Contains(CFLAG_Debug)) 
+	{
+		//add the debug flags
+		CompileFlags |= D3D10_SHADER_DEBUG | D3D10_SHADER_SKIP_OPTIMIZATION;
+	}
+	else
+	{
+		if (Input.Environment.CompilerFlags.Contains(CFLAG_StandardOptimization))
+		{
+			CompileFlags |= D3D10_SHADER_OPTIMIZATION_LEVEL1;
+		}
+		else
+		{
+			CompileFlags |= D3D10_SHADER_OPTIMIZATION_LEVEL3;
+		}
+	}
+
+	for (int32 FlagIndex = 0; FlagIndex < Input.Environment.CompilerFlags.Num(); FlagIndex++)
+	{
+		//accumulate flags set by the shader
+		CompileFlags |= TranslateCompilerFlagD3D11((ECompilerFlags)Input.Environment.CompilerFlags[FlagIndex]);
+	}
+
+	TArray<FString> FilteredErrors;
+	if (!CompileAndProcessD3DShader(PreprocessedShaderSource, CompilerPath, CompileFlags, Input, EntryPointName, ShaderProfile, FilteredErrors, Output))
+	{
+		if (!FilteredErrors.Num())
+		{
 			FilteredErrors.Add(TEXT("Compile Failed without errors!"));
 		}
+	}
 
-		for (int32 ErrorIndex = 0; ErrorIndex < FilteredErrors.Num(); ErrorIndex++)
+	// Process errors
+	for (int32 ErrorIndex = 0; ErrorIndex < FilteredErrors.Num(); ErrorIndex++)
+	{
+		const FString& CurrentError = FilteredErrors[ErrorIndex];
+		FShaderCompilerError NewError;
+		// Extract the filename and line number from the shader compiler error message for PC whose format is:
+		// "d:\UE4\Binaries\BasePassPixelShader(30,7): error X3000: invalid target or usage string"
+		int32 FirstParenIndex = CurrentError.Find(TEXT("("));
+		int32 LastParenIndex = CurrentError.Find(TEXT("):"));
+		if (FirstParenIndex != INDEX_NONE 
+			&& LastParenIndex != INDEX_NONE
+			&& LastParenIndex > FirstParenIndex)
 		{
-			const FString& CurrentError = FilteredErrors[ErrorIndex];
-			FShaderCompilerError NewError;
-			// Extract the filename and line number from the shader compiler error message for PC whose format is:
-			// "d:\UnrealEngine3\Binaries\BasePassPixelShader(30,7): error X3000: invalid target or usage string"
-			int32 FirstParenIndex = CurrentError.Find(TEXT("("));
-			int32 LastParenIndex = CurrentError.Find(TEXT("):"));
-			if (FirstParenIndex != INDEX_NONE 
-				&& LastParenIndex != INDEX_NONE
-				&& LastParenIndex > FirstParenIndex)
+			FString ErrorFileAndPath = CurrentError.Left(FirstParenIndex);
+			if (FPaths::GetExtension(ErrorFileAndPath).ToUpper() == TEXT("USF"))
 			{
-				FString ErrorFileAndPath = CurrentError.Left(FirstParenIndex);
-				if (FPaths::GetExtension(ErrorFileAndPath).ToUpper() == TEXT("USF"))
-				{
-					NewError.ErrorFile = FPaths::GetCleanFilename(ErrorFileAndPath);
-				}
-				else
-				{
-					NewError.ErrorFile = FPaths::GetCleanFilename(ErrorFileAndPath) + TEXT(".usf");
-				}
-
-				NewError.ErrorLineString = CurrentError.Mid(FirstParenIndex + 1, LastParenIndex - FirstParenIndex - FCString::Strlen(TEXT("(")));
-				NewError.StrippedErrorMessage = CurrentError.Right(CurrentError.Len() - LastParenIndex - FCString::Strlen(TEXT("):")));
+				NewError.ErrorFile = FPaths::GetCleanFilename(ErrorFileAndPath);
 			}
 			else
 			{
-				NewError.StrippedErrorMessage = CurrentError;
+				NewError.ErrorFile = FPaths::GetCleanFilename(ErrorFileAndPath) + TEXT(".usf");
 			}
-			Output.Errors.Add(NewError);
+
+			NewError.ErrorLineString = CurrentError.Mid(FirstParenIndex + 1, LastParenIndex - FirstParenIndex - FCString::Strlen(TEXT("(")));
+			NewError.StrippedErrorMessage = CurrentError.Right(CurrentError.Len() - LastParenIndex - FCString::Strlen(TEXT("):")));
 		}
+		else
+		{
+			NewError.StrippedErrorMessage = CurrentError;
+		}
+		Output.Errors.Add(NewError);
 	}
 }
 
diff --git a/Engine/Source/Editor/UnrealEd/Private/PreviewMaterial.cpp b/Engine/Source/Editor/UnrealEd/Private/PreviewMaterial.cpp
--- a/Engine/Source/Editor/UnrealEd/Private/PreviewMaterial.cpp
+++ b/Engine/Source/Editor/UnrealEd/Private/PreviewMaterial.cpp
@@ -73,12 +73,11 @@ public:
 
 			if (bEditorStatsMaterial)
 			{
-				TArray<FString> ShaderTypeNames;
-				TArray<FString> ShaderTypeDescriptions;
-				GetRepresentativeShaderTypesAndDescriptions(ShaderTypeNames, ShaderTypeDescriptions);
+				TMap<FName, FString> ShaderTypeNamesAndDescriptions;
+				GetRepresentativeShaderTypesAndDescriptions(ShaderTypeNamesAndDescriptions);
 
 				//Only allow shaders that are used in the stats.
-				return ShaderTypeNames.Find(ShaderType->GetName()) != INDEX_NONE;
+				return ShaderTypeNamesAndDescriptions.Contains(ShaderType->GetFName());
 			}
 
 			// look for any of the needed type
diff --git a/Engine/Source/Editor/UnrealEd/Private/ThumbnailRendering/MaterialInstanceThumbnailRenderer.cpp b/Engine/Source/Editor/UnrealEd/Private/ThumbnailRendering/MaterialInstanceThumbnailRenderer.cpp
--- a/Engine/Source/Editor/UnrealEd/Private/ThumbnailRendering/MaterialInstanceThumbnailRenderer.cpp
+++ b/Engine/Source/Editor/UnrealEd/Private/ThumbnailRendering/MaterialInstanceThumbnailRenderer.cpp
@@ -28,6 +28,7 @@ void UMaterialInstanceThumbnailRenderer::Draw(UObject* Object, int32 X, int32 Y,
 			.SetWorldTimes(FApp::GetCurrentTime() - GStartTime, FApp::GetDeltaTime(), FApp::GetCurrentTime() - GStartTime));
 
 		ViewFamily.EngineShowFlags.DisableAdvancedFeatures();
+		ViewFamily.EngineShowFlags.SetSeparateTranslucency(true);
 		ViewFamily.EngineShowFlags.MotionBlur = 0;
 		ViewFamily.EngineShowFlags.AntiAliasing = 0;
 
diff --git a/Engine/Source/Runtime/Engine/Classes/Components/SkyLightComponent.h b/Engine/Source/Runtime/Engine/Classes/Components/SkyLightComponent.h
--- a/Engine/Source/Runtime/Engine/Classes/Components/SkyLightComponent.h
+++ b/Engine/Source/Runtime/Engine/Classes/Components/SkyLightComponent.h
@@ -97,6 +97,10 @@ class ENGINE_API USkyLightComponent : public ULightComponentBase
 	UPROPERTY(EditAnywhere, BlueprintReadOnly, Category=Light)
 	class UTextureCube* Cubemap;
 
+	/** Angle to rotate the source cubemap when SourceType is set to SLS_SpecifiedCubemap. */
+	UPROPERTY(EditAnywhere, BlueprintReadOnly, Category=Light, meta=(UIMin = "0", UIMax = "360"))
+	float SourceCubemapAngle;
+
 	/** 
 	 * Distance from the sky light at which any geometry should be treated as part of the sky. 
 	 * This is also used by reflection captures, so update reflection captures to see the impact.
diff --git a/Engine/Source/Runtime/Engine/Classes/Engine/Texture.h b/Engine/Source/Runtime/Engine/Classes/Engine/Texture.h
--- a/Engine/Source/Runtime/Engine/Classes/Engine/Texture.h
+++ b/Engine/Source/Runtime/Engine/Classes/Engine/Texture.h
@@ -13,7 +13,7 @@ UENUM()
 enum TextureCompressionSettings
 {
 	TC_Default					UMETA(DisplayName="Default (DXT1/5, BC1/3 on DX11)"),
-	TC_Normalmap				UMETA(DisplayName="Normalmap (DXT5, BC3 on DX11)"),
+	TC_Normalmap				UMETA(DisplayName="Normalmap (DXT5, BC5 on DX11)"),
 	TC_Masks					UMETA(DisplayName="Masks (no sRGB)"),
 	TC_Grayscale				UMETA(DisplayName="Grayscale (R8, RGB8 sRGB)"),
 	TC_Displacementmap			UMETA(DisplayName="Displacementmap (8/16bit)"),
diff --git a/Engine/Source/Runtime/Engine/Classes/Materials/MaterialInstanceDynamic.h b/Engine/Source/Runtime/Engine/Classes/Materials/MaterialInstanceDynamic.h
--- a/Engine/Source/Runtime/Engine/Classes/Materials/MaterialInstanceDynamic.h
+++ b/Engine/Source/Runtime/Engine/Classes/Materials/MaterialInstanceDynamic.h
@@ -47,20 +47,20 @@ class UMaterialInstanceDynamic : public UMaterialInstance
 	void K2_InterpolateMaterialInstanceParams(UMaterialInstance* SourceA, UMaterialInstance* SourceB, float Alpha);
 
 	/**
-	 * Copies over parameters given a material interface (deep, following the hierarchy)
-	 * Very slow implementation, avoid using at runtime. Hopefully we can replace ity later with somthing like CopyInterpParameters()
+	 * Copies over parameters given a material interface (copy each instance following the hierarchy)
+	 * Very slow implementation, avoid using at runtime. Hopefully we can replace ity later with something like CopyInterpParameters()
 	 * The output is the object itself (this).
 	 */
 	UFUNCTION(BlueprintCallable, meta=(DisplayName = "CopyMaterialInstanceParameters"), Category="Rendering|Material")
 	void K2_CopyMaterialInstanceParameters(UMaterialInterface* Source);
 
 	/**
-	 * Copies over parameters given a material instance (shallow, not deep following the hierarchy)
+	 * Copies over parameters given a material instance (only copy from the instance, not following the hierarchy)
 	 * much faster than K2_CopyMaterialInstanceParameters(), 
 	 * The output is the object itself (this).
 	 * @param Source ignores the call if 0
 	 */
-	UFUNCTION(meta=(DisplayName = "CopyMaterialInstanceParameters"), Category="Rendering|Material")
+	UFUNCTION(meta=(DisplayName = "CopyInterpParameters"), Category="Rendering|Material")
 	ENGINE_API void CopyInterpParameters(UMaterialInstance* Source);
 
 	/**
diff --git a/Engine/Source/Runtime/Engine/Private/Components/SkyLightComponent.cpp b/Engine/Source/Runtime/Engine/Private/Components/SkyLightComponent.cpp
--- a/Engine/Source/Runtime/Engine/Private/Components/SkyLightComponent.cpp
+++ b/Engine/Source/Runtime/Engine/Private/Components/SkyLightComponent.cpp
@@ -316,7 +316,8 @@ bool USkyLightComponent::CanEditChange(const UProperty* InProperty) const
 	{
 		FString PropertyName = InProperty->GetName();
 
-		if (FCString::Strcmp(*PropertyName, TEXT("Cubemap")) == 0)
+		if (FCString::Strcmp(*PropertyName, TEXT("Cubemap")) == 0
+			|| FCString::Strcmp(*PropertyName, TEXT("SourceCubemapAngle")) == 0)
 		{
 			return SourceType == SLS_SpecifiedCubemap;
 		}
diff --git a/Engine/Source/Runtime/Engine/Private/Components/StaticMeshComponent.cpp b/Engine/Source/Runtime/Engine/Private/Components/StaticMeshComponent.cpp
--- a/Engine/Source/Runtime/Engine/Private/Components/StaticMeshComponent.cpp
+++ b/Engine/Source/Runtime/Engine/Private/Components/StaticMeshComponent.cpp
@@ -871,8 +871,17 @@ void UStaticMeshComponent::CachePaintedDataIfNecessary()
 					}
 				}
 				else
-				{
-					UE_LOG(LogStaticMesh, Warning, TEXT("Unable to cache painted data for mesh component. Vertex color overrides will be lost if the mesh is modified. %s %s LOD%d."), *GetFullName(), *StaticMesh->GetFullName(), LODIter.GetIndex() );
+				{					
+					// At this point we can't resolve the colors, so just discard any isolated data we still have
+					if (CurCompLODInfo.OverrideVertexColors && CurCompLODInfo.OverrideVertexColors->GetNumVertices() > 0)
+					{
+						UE_LOG(LogStaticMesh, Warning, TEXT("Level requires re-saving! Outdated vertex color overrides have been discarded for %s %s LOD%d. "), *GetFullName(), *StaticMesh->GetFullName(), LODIter.GetIndex());
+						CurCompLODInfo.ReleaseOverrideVertexColorsAndBlock();
+					}
+					else
+					{
+						UE_LOG(LogStaticMesh, Warning, TEXT("Unable to cache painted data for mesh component. Vertex color overrides will be lost if the mesh is modified. %s %s LOD%d."), *GetFullName(), *StaticMesh->GetFullName(), LODIter.GetIndex() );
+					}
 				}
 			}
 		}
diff --git a/Engine/Source/Runtime/Engine/Private/GlobalShader.cpp b/Engine/Source/Runtime/Engine/Private/GlobalShader.cpp
--- a/Engine/Source/Runtime/Engine/Private/GlobalShader.cpp
+++ b/Engine/Source/Runtime/Engine/Private/GlobalShader.cpp
@@ -23,51 +23,46 @@ const int32 GlobalShaderMapId = 0;
 FGlobalShaderMapId::FGlobalShaderMapId(EShaderPlatform Platform)
 {
 	TArray<FShaderType*> ShaderTypes;
-
 	TArray<const FShaderPipelineType*> ShaderPipelineTypes;
+
+	for (TLinkedList<FShaderType*>::TIterator ShaderTypeIt(FShaderType::GetTypeList()); ShaderTypeIt; ShaderTypeIt.Next())
+	{
+		FGlobalShaderType* GlobalShaderType = ShaderTypeIt->GetGlobalShaderType();
+		if (GlobalShaderType && GlobalShaderType->ShouldCache(Platform))
+		{
+			ShaderTypes.Add(GlobalShaderType);
+		}
+	}
+
 	for (TLinkedList<FShaderPipelineType*>::TIterator ShaderPipelineIt(FShaderPipelineType::GetTypeList()); ShaderPipelineIt; ShaderPipelineIt.Next())
 	{
 		const FShaderPipelineType* Pipeline = *ShaderPipelineIt;
 		if (Pipeline->IsGlobalTypePipeline())
 		{
-			// Verify all the ShouldCache are in sync
-			int32 NumShouldCache = 0;
-			const auto& Stages = Pipeline->GetStages();
-			for (int32 Index = 0; Index < Stages.Num(); ++Index)
+			int32 NumStagesNeeded = 0;
+			auto& StageTypes = Pipeline->GetStages();
+			for (const FShaderType* Shader : StageTypes)
 			{
-				const FGlobalShaderType* GlobalShaderType = Stages[Index]->GetGlobalShaderType();
+				const FGlobalShaderType* GlobalShaderType = Shader->GetGlobalShaderType();
 				if (GlobalShaderType->ShouldCache(Platform))
 				{
-					++NumShouldCache;
-				}
-			}
-
-			if (NumShouldCache != 0)
-			{
-				if (NumShouldCache != Stages.Num())
-				{
-					UE_LOG(LogShaders, Error, TEXT("Mismatched ShouldCache() function on the shaders for Global ShaderPipeline %s"), Pipeline->GetName());
+					++NumStagesNeeded;
 				}
 				else
 				{
-					ShaderPipelineTypes.Add(Pipeline);
+					break;
 				}
 			}
-		}
-	}
 
-	for (TLinkedList<FShaderType*>::TIterator ShaderTypeIt(FShaderType::GetTypeList()); ShaderTypeIt; ShaderTypeIt.Next())
-	{
-		FGlobalShaderType* GlobalShaderType = ShaderTypeIt->GetGlobalShaderType();
-		if (GlobalShaderType && GlobalShaderType->ShouldCache(Platform))
-		{
-			ShaderTypes.Add(GlobalShaderType);
+			if (NumStagesNeeded == StageTypes.Num())
+			{
+				ShaderPipelineTypes.Add(Pipeline);
+			}
 		}
 	}
 
+	// Individual shader dependencies
 	ShaderTypes.Sort(FCompareShaderTypes());
-	ShaderPipelineTypes.Sort(FCompareShaderPipelineTypes());
-
 	for (int32 TypeIndex = 0; TypeIndex < ShaderTypes.Num(); TypeIndex++)
 	{
 		FShaderTypeDependency Dependency;
@@ -76,19 +71,15 @@ FGlobalShaderMapId::FGlobalShaderMapId(EShaderPlatform Platform)
 		ShaderTypeDependencies.Add(Dependency);
 	}
 
-	//#todo-rco: Is this currently necessary? Might be redundant until we only have Pipelines...
+	// Shader pipeline dependencies
+	ShaderPipelineTypes.Sort(FCompareShaderPipelineNameTypes());
 	for (int32 TypeIndex = 0; TypeIndex < ShaderPipelineTypes.Num(); TypeIndex++)
 	{
-		auto& Stages = ShaderPipelineTypes[TypeIndex]->GetStages();
-		for (int32 Index = 0; Index < Stages.Num(); ++Index)
-		{
-			FShaderType* GlobalShaderType = (FShaderType*)Stages[Index];
-			FShaderTypeDependency Dependency;
-			Dependency.ShaderPipelineType = ShaderPipelineTypes[TypeIndex];
-			Dependency.ShaderType = (FShaderType*)Stages[Index];
-			Dependency.SourceHash = Stages[Index]->GetSourceHash();
-			ShaderTypeDependencies.Add(Dependency);
-		}
+		const FShaderPipelineType* Pipeline = ShaderPipelineTypes[TypeIndex];
+		FShaderPipelineTypeDependency Dependency;
+		Dependency.ShaderPipelineType = Pipeline;
+		Dependency.StagesSourceHash = Pipeline->GetSourceHash();
+		ShaderPipelineTypeDependencies.Add(Dependency);
 	}
 }
 
@@ -101,12 +92,6 @@ void FGlobalShaderMapId::AppendKeyString(FString& KeyString) const
 		const FShaderTypeDependency& ShaderTypeDependency = ShaderTypeDependencies[ShaderIndex];
 
 		KeyString += TEXT("_");
-		if (ShaderTypeDependency.ShaderPipelineType)
-		{
-			KeyString += ShaderTypeDependency.ShaderPipelineType->GetName();
-			KeyString += TEXT("_");
-		}
-
 		KeyString += ShaderTypeDependency.ShaderType->GetName();
 
 		// Add the type's source hash so that we can invalidate cached shaders when .usf changes are made
@@ -124,6 +109,28 @@ void FGlobalShaderMapId::AppendKeyString(FString& KeyString) const
 		}
 	}
 
+	for (int32 Index = 0; Index < ShaderPipelineTypeDependencies.Num(); ++Index)
+	{
+		const FShaderPipelineTypeDependency& Dependency = ShaderPipelineTypeDependencies[Index];
+
+		KeyString += TEXT("_");
+		KeyString += Dependency.ShaderPipelineType->GetName();
+
+		// Add the type's source hash so that we can invalidate cached shaders when .usf changes are made
+		KeyString += Dependency.StagesSourceHash.ToString();
+
+		for (const FShaderType* ShaderType : Dependency.ShaderPipelineType->GetStages())
+		{
+			const TMap<const TCHAR*, FCachedUniformBufferDeclaration>& ReferencedUniformBufferStructsCache = ShaderType->GetReferencedUniformBufferStructsCache();
+
+			// Gather referenced uniform buffers
+			for (TMap<const TCHAR*, FCachedUniformBufferDeclaration>::TConstIterator It(ReferencedUniformBufferStructsCache); It; ++It)
+			{
+				ReferencedUniformBuffers.Add(It.Key(), It.Value());
+			}
+		}
+	}
+
 	{
 		TArray<uint8> TempData;
 		FSerializationHistory SerializationHistory;
@@ -137,7 +144,7 @@ void FGlobalShaderMapId::AppendKeyString(FString& KeyString) const
 	}
 }
 
-void FGlobalShaderType::BeginCompileShader(EShaderPlatform Platform, const FShaderPipelineType* ShaderPipeline, TArray<FShaderCommonCompileJob*>& NewJobs)
+FShaderCompileJob* FGlobalShaderType::BeginCompileShader(EShaderPlatform Platform, const FShaderPipelineType* ShaderPipeline, TArray<FShaderCommonCompileJob*>& NewJobs)
 {
 	FShaderCompileJob* NewJob = new FShaderCompileJob(GlobalShaderMapId, nullptr, this);
 	FShaderCompilerEnvironment& ShaderEnvironment = NewJob->Input.Environment;
@@ -161,6 +168,8 @@ void FGlobalShaderType::BeginCompileShader(EShaderPlatform Platform, const FShad
 		NewJob,
 		NewJobs
 		);
+
+	return NewJob;
 }
 
 void FGlobalShaderType::BeginCompileShaderPipeline(EShaderPlatform Platform, const FShaderPipelineType* ShaderPipeline, const TArray<FGlobalShaderType*>& ShaderStages, TArray<FShaderCommonCompileJob*>& NewJobs)
@@ -192,6 +201,12 @@ FShader* FGlobalShaderType::FinishCompileShader(const FShaderCompileJob& Current
 		FShaderResource* Resource = FShaderResource::FindOrCreateShaderResource(CurrentJob.Output, SpecificType);
 		check(Resource);
 
+		if (ShaderPipelineType && !ShaderPipelineType->ShouldOptimizeUnusedOutputs())
+		{
+			// If sharing shaders in this pipeline, remove it from the type/id so it uses the one in the shared shadermap list
+			ShaderPipelineType = nullptr;
+		}
+
 		// Find a shader with the same key in memory
 		Shader = CurrentJob.ShaderType->FindShaderById(FShaderId(GGlobalShaderMapHash, ShaderPipelineType, nullptr, CurrentJob.ShaderType, CurrentJob.Input.Target));
 
@@ -287,7 +302,37 @@ void VerifyGlobalShaders(EShaderPlatform Platform, bool bLoadedFromCacheFile)
 		bErrorOnMissing = true;
 	}
 
+	// All jobs, single & pipeline
 	TArray<FShaderCommonCompileJob*> GlobalShaderJobs;
+
+	// Add the single jobs first
+	TMap<FShaderType*, FShaderCompileJob*> SharedShaderJobs;
+	for(TLinkedList<FShaderType*>::TIterator ShaderTypeIt(FShaderType::GetTypeList());ShaderTypeIt;ShaderTypeIt.Next())
+	{
+		FGlobalShaderType* GlobalShaderType = ShaderTypeIt->GetGlobalShaderType();
+		if (GlobalShaderType && GlobalShaderType->ShouldCache(Platform))
+		{
+			if (!GlobalShaderMap->HasShader(GlobalShaderType))
+			{
+				if (bErrorOnMissing)
+				{
+					UE_LOG(LogShaders, Fatal,TEXT("Missing global shader %s, Please make sure cooking was successful."), GlobalShaderType->GetName());
+				}
+
+				if (!bEmptyMap)
+				{
+					UE_LOG(LogShaders, Warning, TEXT("	%s"), GlobalShaderType->GetName());
+				}
+	
+				// Compile this global shader type.
+				auto* Job = GlobalShaderType->BeginCompileShader(Platform, nullptr, GlobalShaderJobs);
+				check(!SharedShaderJobs.Find(GlobalShaderType));
+				SharedShaderJobs.Add(GlobalShaderType, Job);
+			}
+		}
+	}
+
+	// Now the pipeline jobs; if it's a shareable pipeline, do not add duplicate jobs
 	for (TLinkedList<FShaderPipelineType*>::TIterator ShaderPipelineIt(FShaderPipelineType::GetTypeList());ShaderPipelineIt;ShaderPipelineIt.Next())
 	{
 		const FShaderPipelineType* Pipeline = *ShaderPipelineIt;
@@ -295,69 +340,57 @@ void VerifyGlobalShaders(EShaderPlatform Platform, bool bLoadedFromCacheFile)
 		{
 			if (!GlobalShaderMap->GetShaderPipeline(Pipeline))
 			{
-				// First see if all the ShouldCache pass
-				auto& Stages = Pipeline->GetStages();
-				int32 ShaderStagesToCompile = 0;
+				auto& StageTypes = Pipeline->GetStages();
 				TArray<FGlobalShaderType*> ShaderStages;
-				for (int32 Index = 0; Index < Stages.Num(); ++Index)
+				for (int32 Index = 0; Index < StageTypes.Num(); ++Index)
 				{
-					FGlobalShaderType* GlobalShaderType = ((FShaderType*)(Stages[Index]))->GetGlobalShaderType();
+					FGlobalShaderType* GlobalShaderType = ((FShaderType*)(StageTypes[Index]))->GetGlobalShaderType();
 					if (GlobalShaderType->ShouldCache(Platform))
 					{
 						ShaderStages.Add(GlobalShaderType);
-						++ShaderStagesToCompile;
+					}
+					else
+					{
+						break;
 					}
 				}
 
-				if (ShaderStagesToCompile > 0)
+				if (ShaderStages.Num() == StageTypes.Num())
 				{
 					if (bErrorOnMissing)
 					{
 						UE_LOG(LogShaders, Fatal, TEXT("Missing global shader pipeline %s, Please make sure cooking was successful."), Pipeline->GetName());
 					}
 
-					if (ShaderStagesToCompile != ShaderStages.Num())
+					if (!bEmptyMap)
+					{
+						UE_LOG(LogShaders, Warning, TEXT("	%s"), Pipeline->GetName());
+					}
+
+					if (Pipeline->ShouldOptimizeUnusedOutputs())
 					{
-						UE_LOG(LogShaders, Fatal, TEXT("Missing global shader pipeline %s, has mismatched ShouldCache() functions, Please make sure cooking was successful."), Pipeline->GetName());
+						// Make a pipeline job with all the stages
+						FGlobalShaderType::BeginCompileShaderPipeline(Platform, Pipeline, ShaderStages, GlobalShaderJobs);
 					}
 					else
 					{
-						if (!bEmptyMap)
+						// If sharing shaders amongst pipelines, add this pipeline as a dependency of an existing individual job
+						for (const FShaderType* ShaderType : StageTypes)
 						{
-							UE_LOG(LogShaders, Warning, TEXT("	%s"), Pipeline->GetName());
+							FShaderCompileJob** Job = SharedShaderJobs.Find(ShaderType);
+							checkf(Job, TEXT("Couldn't find existing shared job for global shader %s on pipeline %s!"), ShaderType->GetName(), Pipeline->GetName());
+							auto* SingleJob = (*Job)->GetSingleShaderJob();
+							check(SingleJob);
+							auto& SharedPipelinesInJob = SingleJob->SharingPipelines.FindOrAdd(nullptr);
+							check(!SharedPipelinesInJob.Contains(Pipeline));
+							SharedPipelinesInJob.Add(Pipeline);
 						}
-
-						// Make a pipeline job with all the stages
-						FGlobalShaderType::BeginCompileShaderPipeline(Platform, Pipeline, ShaderStages, GlobalShaderJobs);
 					}
 				}
 			}
 		}
 	}
 
-	for(TLinkedList<FShaderType*>::TIterator ShaderTypeIt(FShaderType::GetTypeList());ShaderTypeIt;ShaderTypeIt.Next())
-	{
-		FGlobalShaderType* GlobalShaderType = ShaderTypeIt->GetGlobalShaderType();
-		if (GlobalShaderType && GlobalShaderType->ShouldCache(Platform))
-		{
-			if (!GlobalShaderMap->HasShader(GlobalShaderType))
-			{
-				if (bErrorOnMissing)
-				{
-					UE_LOG(LogShaders, Fatal,TEXT("Missing global shader %s, Please make sure cooking was successful."), GlobalShaderType->GetName());
-				}
-
-				if (!bEmptyMap)
-				{
-					UE_LOG(LogShaders, Warning, TEXT("	%s"), GlobalShaderType->GetName());
-				}
-	
-				// Compile this global shader type.
-				GlobalShaderType->BeginCompileShader(Platform, nullptr, GlobalShaderJobs);
-			}
-		}
-	}
-
 	if (GlobalShaderJobs.Num() > 0)
 	{
 		GShaderCompilingManager->AddJobs(GlobalShaderJobs, true, true);
@@ -380,7 +413,7 @@ void VerifyGlobalShaders(EShaderPlatform Platform, bool bLoadedFromCacheFile)
 }
 
 /** Serializes the global shader map to an archive. */
-void SerializeGlobalShaders(FArchive& Ar, TShaderMap<FGlobalShaderType>* GlobalShaderMap)
+static void SerializeGlobalShaders(FArchive& Ar, TShaderMap<FGlobalShaderType>* GlobalShaderMap)
 {
 	check(IsInGameThread());
 
@@ -403,7 +436,7 @@ void SerializeGlobalShaders(FArchive& Ar, TShaderMap<FGlobalShaderType>* GlobalS
 	GlobalShaderMap->SerializeInline(Ar, true, false);
 }
 
-FString GetGlobalShaderCacheFilename(EShaderPlatform Platform)
+static FString GetGlobalShaderCacheFilename(EShaderPlatform Platform)
 {
 	return FString(TEXT("Engine")) / TEXT("GlobalShaderCache-") + LegacyShaderPlatformToShaderFormat(Platform).ToString() + TEXT(".bin");
 }
@@ -474,14 +507,14 @@ TShaderMap<FGlobalShaderType>* GetGlobalShaderMap(EShaderPlatform Platform, bool
 			GGlobalShaderMap[Platform] = new TShaderMap<FGlobalShaderType>();
 			return GGlobalShaderMap[Platform];
 		}
-		return NULL;
+		return nullptr;
 	}
 
 	if (bRefreshShaderMap)
 	{
 		// delete the current global shader map
 		delete GGlobalShaderMap[Platform];
-		GGlobalShaderMap[Platform] = NULL;
+		GGlobalShaderMap[Platform] = nullptr;
 
 		// make sure we look for updated shader source files
 		FlushShaderFileCache();
@@ -576,7 +609,7 @@ TShaderMap<FGlobalShaderType>* GetGlobalShaderMap(EShaderPlatform Platform, bool
 	return GGlobalShaderMap[Platform];
 }
 
-static inline bool ShouldCacheGlobalShader(const FGlobalShaderType* GlobalShaderType, const TCHAR* TypeNameSubstring, EShaderPlatform Platform)
+static inline bool ShouldCacheGlobalShaderTypeName(const FGlobalShaderType* GlobalShaderType, const TCHAR* TypeNameSubstring, EShaderPlatform Platform)
 {
 	return GlobalShaderType
 		&& (TypeNameSubstring == nullptr || (FPlatformString::Strstr(GlobalShaderType->GetName(), TypeNameSubstring) != nullptr))
@@ -594,32 +627,46 @@ bool IsGlobalShaderMapComplete(const TCHAR* TypeNameSubstring)
 
 		if (GlobalShaderMap)
 		{
+			// Check if the individual shaders are complete
+			for (TLinkedList<FShaderType*>::TIterator ShaderTypeIt(FShaderType::GetTypeList()); ShaderTypeIt; ShaderTypeIt.Next())
+			{
+				FGlobalShaderType* GlobalShaderType = ShaderTypeIt->GetGlobalShaderType();
+				if (ShouldCacheGlobalShaderTypeName(GlobalShaderType, TypeNameSubstring, Platform))
+				{
+					if (!GlobalShaderMap->HasShader(GlobalShaderType))
+					{
+						return false;
+					}
+				}
+			}
+
+			// Then the pipelines as it may be sharing shaders
 			for (TLinkedList<FShaderPipelineType*>::TIterator ShaderPipelineIt(FShaderPipelineType::GetTypeList()); ShaderPipelineIt; ShaderPipelineIt.Next())
 			{
 				const FShaderPipelineType* Pipeline = *ShaderPipelineIt;
 				if (Pipeline->IsGlobalTypePipeline())
 				{
 					auto& Stages = Pipeline->GetStages();
-					// Only need to check Stage 0 as we validated ShouldCache matching on FGlobalShaderMapId()
-					const FGlobalShaderType* GlobalShaderType = Stages[0]->GetGlobalShaderType();
-					if (ShouldCacheGlobalShader(GlobalShaderType, TypeNameSubstring, Platform))
+					int32 NumStagesNeeded = 0;
+					for (const FShaderType* Shader : Stages)
 					{
-						if (!GlobalShaderMap->GetShaderPipeline(Pipeline))
+						const FGlobalShaderType* GlobalShaderType = Shader->GetGlobalShaderType();
+						if (ShouldCacheGlobalShaderTypeName(GlobalShaderType, TypeNameSubstring, Platform))
 						{
-							return false;
+							++NumStagesNeeded;
+						}
+						else
+						{
+							break;
 						}
 					}
-				}
-			}
 
-			for (TLinkedList<FShaderType*>::TIterator ShaderTypeIt(FShaderType::GetTypeList()); ShaderTypeIt; ShaderTypeIt.Next())
-			{
-				FGlobalShaderType* GlobalShaderType = ShaderTypeIt->GetGlobalShaderType();
-				if (ShouldCacheGlobalShader(GlobalShaderType, TypeNameSubstring, Platform))
-				{
-					if (!GlobalShaderMap->HasShader(GlobalShaderType))
+					if (NumStagesNeeded == Stages.Num())
 					{
-						return false;
+						if (!GlobalShaderMap->GetShaderPipeline(Pipeline))
+						{
+							return false;
+						}
 					}
 				}
 			}
@@ -869,18 +916,30 @@ void FinishRecompileGlobalShaders()
 	GShaderCompilingManager->ProcessAsyncResults(false, true);
 }
 
-static inline FShader* ProcessCompiledJob(FShaderCompileJob* SingleJob, const FShaderPipelineType* Pipeline, TArray<EShaderPlatform>& ShaderPlatformsProcessed)
+static inline FShader* ProcessCompiledJob(FShaderCompileJob* SingleJob, const FShaderPipelineType* Pipeline, TArray<EShaderPlatform>& ShaderPlatformsProcessed, TArray<const FShaderPipelineType*>& OutSharedPipelines)
 {
 	FGlobalShaderType* GlobalShaderType = SingleJob->ShaderType->GetGlobalShaderType();
 	check(GlobalShaderType);
 	FShader* Shader = GlobalShaderType->FinishCompileShader(*SingleJob, Pipeline);
 	if (Shader)
 	{
-		// Add the new global shader instance to the global shader map.
+		// Add the new global shader instance to the global shader map if it's a shared shader
 		EShaderPlatform Platform = (EShaderPlatform)SingleJob->Input.Target.Platform;
-		if (!Pipeline || Pipeline->ShouldAddShadersToIndividualMap())
+		if (!Pipeline || !Pipeline->ShouldOptimizeUnusedOutputs())
 		{
 			GGlobalShaderMap[Platform]->AddShader(GlobalShaderType, Shader);
+			// Add this shared pipeline to the list
+			if (!Pipeline)
+			{
+				auto* JobSharedPipelines = SingleJob->SharingPipelines.Find(nullptr);
+				if (JobSharedPipelines)
+				{
+					for (auto* SharedPipeline : *JobSharedPipelines)
+					{
+						OutSharedPipelines.AddUnique(SharedPipeline);
+					}
+				}
+			}
 		}
 		ShaderPlatformsProcessed.AddUnique(Platform);
 	}
@@ -900,6 +959,7 @@ void ProcessCompiledGlobalShaders(const TArray<FShaderCommonCompileJob*>& Compil
 	UE_LOG(LogShaders, Warning, TEXT("Compiled %u global shaders"), CompilationResults.Num());
 
 	TArray<EShaderPlatform> ShaderPlatformsProcessed;
+	TArray<const FShaderPipelineType*> SharedPipelines;
 
 	for (int32 ResultIndex = 0; ResultIndex < CompilationResults.Num(); ResultIndex++)
 	{
@@ -907,7 +967,7 @@ void ProcessCompiledGlobalShaders(const TArray<FShaderCommonCompileJob*>& Compil
 		FShaderCompileJob* SingleJob = nullptr;
 		if ((SingleJob = (FShaderCompileJob*)CurrentJob.GetSingleShaderJob()) != nullptr)
 		{
-			ProcessCompiledJob(SingleJob, nullptr, ShaderPlatformsProcessed);
+			ProcessCompiledJob(SingleJob, nullptr, ShaderPlatformsProcessed, SharedPipelines);
 		}
 		else
 		{
@@ -917,7 +977,7 @@ void ProcessCompiledGlobalShaders(const TArray<FShaderCommonCompileJob*>& Compil
 			for (int32 Index = 0; Index < PipelineJob->StageJobs.Num(); ++Index)
 			{
 				SingleJob = PipelineJob->StageJobs[Index]->GetSingleShaderJob();
-				FShader* Shader = ProcessCompiledJob(SingleJob, PipelineJob->ShaderPipeline, ShaderPlatformsProcessed);
+				FShader* Shader = ProcessCompiledJob(SingleJob, PipelineJob->ShaderPipeline, ShaderPlatformsProcessed, SharedPipelines);
 				ShaderStages.Add(Shader);
 			}
 
@@ -925,7 +985,7 @@ void ProcessCompiledGlobalShaders(const TArray<FShaderCommonCompileJob*>& Compil
 			if (ShaderPipeline)
 			{
 				EShaderPlatform Platform = (EShaderPlatform)PipelineJob->StageJobs[0]->GetSingleShaderJob()->Input.Target.Platform;
-				check(ShaderPipeline && !GGlobalShaderMap[Platform]->GetShaderPipeline(ShaderPipeline->PipelineType));
+				check(ShaderPipeline && !GGlobalShaderMap[Platform]->HasShaderPipeline(ShaderPipeline->PipelineType));
 				GGlobalShaderMap[Platform]->AddShaderPipeline(PipelineJob->ShaderPipeline, ShaderPipeline);
 			}
 		}
@@ -933,6 +993,39 @@ void ProcessCompiledGlobalShaders(const TArray<FShaderCommonCompileJob*>& Compil
 
 	for (int32 PlatformIndex = 0; PlatformIndex < ShaderPlatformsProcessed.Num(); PlatformIndex++)
 	{
+		{
+			// Process the shader pipelines that share shaders
+			EShaderPlatform Platform = ShaderPlatformsProcessed[PlatformIndex];
+			auto* GlobalShaderMap = GGlobalShaderMap[Platform];
+			for (const FShaderPipelineType* ShaderPipelineType : SharedPipelines)
+			{
+				check(ShaderPipelineType->IsGlobalTypePipeline());
+				if (!GlobalShaderMap->HasShaderPipeline(ShaderPipelineType))
+				{
+					auto& StageTypes = ShaderPipelineType->GetStages();
+					TArray<FShader*> ShaderStages;
+					for (int32 Index = 0; Index < StageTypes.Num(); ++Index)
+					{
+						FGlobalShaderType* GlobalShaderType = ((FShaderType*)(StageTypes[Index]))->GetGlobalShaderType();
+						if (GlobalShaderType->ShouldCache(Platform))
+						{
+							FShader* Shader = GlobalShaderMap->GetShader(GlobalShaderType);
+							check(Shader);
+							ShaderStages.Add(Shader);
+						}
+						else
+						{
+							break;
+						}
+					}
+
+					checkf(StageTypes.Num() == ShaderStages.Num(), TEXT("Internal Error adding Global ShaderPipeline %s"), ShaderPipelineType->GetName());
+					FShaderPipeline* ShaderPipeline = new FShaderPipeline(ShaderPipelineType, ShaderStages);
+					GlobalShaderMap->AddShaderPipeline(ShaderPipelineType, ShaderPipeline);
+				}
+			}
+		}
+
 		// Save the global shader map for any platforms that were recompiled
 		SaveGlobalShaderMapToDerivedDataCache(ShaderPlatformsProcessed[PlatformIndex]);
 	}
diff --git a/Engine/Source/Runtime/Engine/Private/LocalVertexFactory.cpp b/Engine/Source/Runtime/Engine/Private/LocalVertexFactory.cpp
--- a/Engine/Source/Runtime/Engine/Private/LocalVertexFactory.cpp
+++ b/Engine/Source/Runtime/Engine/Private/LocalVertexFactory.cpp
@@ -12,11 +12,19 @@
 void FLocalVertexFactoryShaderParameters::Bind(const FShaderParameterMap& ParameterMap)
 {
 	LODParameter.Bind(ParameterMap, TEXT("SpeedTreeLODInfo"));
+	WindParameter.Bind(ParameterMap, TEXT("SpeedTreeData"));
+
+	bAnySpeedTreeParamIsBound = WindParameter.IsBound() || LODParameter.IsBound();
 }
 
 void FLocalVertexFactoryShaderParameters::Serialize(FArchive& Ar)
 {
 	Ar << LODParameter;
+	Ar << WindParameter;
+	if (Ar.IsLoading())
+	{
+		bAnySpeedTreeParamIsBound = WindParameter.IsBound() || LODParameter.IsBound();
+	}
 }
 
 void FLocalVertexFactoryShaderParameters::SetMesh(FRHICommandList& RHICmdList, FShader* Shader, const FVertexFactory* VertexFactory, const FSceneView& View, const FMeshBatchElement& BatchElement, uint32 DataFlags) const
@@ -27,9 +35,9 @@ void FLocalVertexFactoryShaderParameters::SetMesh(FRHICommandList& RHICmdList, F
 		check(OverrideColorVertexBuffer);
 		static_cast<const FLocalVertexFactory*>(VertexFactory)->SetColorOverrideStream(RHICmdList, OverrideColorVertexBuffer);
 	}
-	//QUICK_SCOPE_CYCLE_COUNTER(STAT_FLocalVertexFactoryShaderParameters_SetMesh_SlowTree);
-	if (View.Family != NULL && View.Family->Scene != NULL)
+	if (bAnySpeedTreeParamIsBound && View.Family != NULL && View.Family->Scene != NULL)
 	{
+		QUICK_SCOPE_CYCLE_COUNTER(STAT_FLocalVertexFactoryShaderParameters_SetMesh_SpeedTree);
 		FUniformBufferRHIParamRef SpeedTreeUniformBuffer = View.Family->Scene->GetSpeedTreeUniformBuffer(VertexFactory);
 		if (SpeedTreeUniformBuffer != NULL)
 		{
diff --git a/Engine/Source/Runtime/Engine/Private/Materials/MaterialInstance.cpp b/Engine/Source/Runtime/Engine/Private/Materials/MaterialInstance.cpp
--- a/Engine/Source/Runtime/Engine/Private/Materials/MaterialInstance.cpp
+++ b/Engine/Source/Runtime/Engine/Private/Materials/MaterialInstance.cpp
@@ -1015,7 +1015,7 @@ void UMaterialInstanceDynamic::CopyScalarAndVectorParameters(const UMaterialInte
 			{
 				const FMaterialUniformExpression* UniformExpression = Array[i];
 
-				// the array can have non scalar parameters in it, those we don't want to inetrpolate
+				// the array can have non scalar parameters in it, those we don't want to interpolate
 				if (UniformExpression->GetType() == &FMaterialUniformExpressionScalarParameter::StaticType)
 				{
 					const FMaterialUniformExpressionScalarParameter* ScalarExpression = static_cast<const FMaterialUniformExpressionScalarParameter*>(UniformExpression);
diff --git a/Engine/Source/Runtime/Engine/Private/Materials/MaterialInstanceDynamic.cpp b/Engine/Source/Runtime/Engine/Private/Materials/MaterialInstanceDynamic.cpp
--- a/Engine/Source/Runtime/Engine/Private/Materials/MaterialInstanceDynamic.cpp
+++ b/Engine/Source/Runtime/Engine/Private/Materials/MaterialInstanceDynamic.cpp
@@ -187,7 +187,7 @@ void UMaterialInstanceDynamic::K2_InterpolateMaterialInstanceParams(UMaterialIns
 
 				if(ParamValueA || ParamValueB)
 				{
-					auto Default = FLinearColor();
+					auto Default = FLinearColor(EForceInit::ForceInit);
 
 					if(!ParamValueA || !ParamValueB)
 					{
diff --git a/Engine/Source/Runtime/Engine/Private/Materials/MaterialShader.cpp b/Engine/Source/Runtime/Engine/Private/Materials/MaterialShader.cpp
--- a/Engine/Source/Runtime/Engine/Private/Materials/MaterialShader.cpp
+++ b/Engine/Source/Runtime/Engine/Private/Materials/MaterialShader.cpp
@@ -29,7 +29,7 @@ static FAutoConsoleVariableRef CVarCreateShadersOnLoad(
 #include "ModuleManager.h"
 
 
-FString GetMaterialShaderMapKeyString(const FMaterialShaderMapId& ShaderMapId, EShaderPlatform Platform);
+static FString GetMaterialShaderMapKeyString(const FMaterialShaderMapId& ShaderMapId, EShaderPlatform Platform);
 
 namespace MaterialStats
 {
@@ -211,7 +211,7 @@ FString GetBlendModeString(EBlendMode BlendMode)
 
 
 /** Creates a string key for the derived data cache given a shader map id. */
-FString GetMaterialShaderMapKeyString(const FMaterialShaderMapId& ShaderMapId, EShaderPlatform Platform)
+static FString GetMaterialShaderMapKeyString(const FMaterialShaderMapId& ShaderMapId, EShaderPlatform Platform)
 {
 	FName Format = LegacyShaderPlatformToShaderFormat(Platform);
 	FString ShaderMapKeyString = Format.ToString() + TEXT("_") + FString(FString::FromInt(GetTargetPlatformManagerRef().ShaderFormatVersion(Format))) + TEXT("_");
@@ -499,6 +499,10 @@ void FMaterialShaderMapId::Serialize(FArchive& Ar)
 	}
 
 	Ar << ShaderTypeDependencies;
+	if (Ar.UE4Ver() >= VER_UE4_PURGED_FMATERIAL_COMPILE_OUTPUTS)
+	{
+		Ar << ShaderPipelineTypeDependencies;
+	}
 	Ar << VertexFactoryTypeDependencies;
 
 	if (Ar.UE4Ver() >= VER_UE4_PURGED_FMATERIAL_COMPILE_OUTPUTS)
@@ -576,6 +580,7 @@ bool FMaterialShaderMapId::operator==(const FMaterialShaderMapId& ReferenceSet)
 		|| ReferencedFunctions.Num() != ReferenceSet.ReferencedFunctions.Num()
 		||  ReferencedParameterCollections.Num() != ReferenceSet.ReferencedParameterCollections.Num()
 		|| ShaderTypeDependencies.Num() != ReferenceSet.ShaderTypeDependencies.Num()
+		|| ShaderPipelineTypeDependencies.Num() != ReferenceSet.ShaderPipelineTypeDependencies.Num()
 		|| VertexFactoryTypeDependencies.Num() != ReferenceSet.VertexFactoryTypeDependencies.Num())
 	{
 		return false;
@@ -606,13 +611,23 @@ bool FMaterialShaderMapId::operator==(const FMaterialShaderMapId& ReferenceSet)
 		const FShaderTypeDependency& ShaderTypeDependency = ShaderTypeDependencies[ShaderIndex];
 		
 		if (ShaderTypeDependency.ShaderType != ReferenceSet.ShaderTypeDependencies[ShaderIndex].ShaderType
-			|| ShaderTypeDependency.ShaderPipelineType != ReferenceSet.ShaderTypeDependencies[ShaderIndex].ShaderPipelineType
 			|| ShaderTypeDependency.SourceHash != ReferenceSet.ShaderTypeDependencies[ShaderIndex].SourceHash)
 		{
 			return false;
 		}
 	}
 
+	for (int32 ShaderPipelineIndex = 0; ShaderPipelineIndex < ShaderPipelineTypeDependencies.Num(); ShaderPipelineIndex++)
+	{
+		const FShaderPipelineTypeDependency& ShaderPipelineTypeDependency = ShaderPipelineTypeDependencies[ShaderPipelineIndex];
+
+		if (ShaderPipelineTypeDependency.ShaderPipelineType != ReferenceSet.ShaderPipelineTypeDependencies[ShaderPipelineIndex].ShaderPipelineType
+			|| ShaderPipelineTypeDependency.StagesSourceHash != ReferenceSet.ShaderPipelineTypeDependencies[ShaderPipelineIndex].StagesSourceHash)
+		{
+			return false;
+		}
+	}
+
 	for (int32 VFIndex = 0; VFIndex < VertexFactoryTypeDependencies.Num(); VFIndex++)
 	{
 		const FVertexFactoryTypeDependency& VFDependency = VertexFactoryTypeDependencies[VFIndex];
@@ -676,11 +691,6 @@ void FMaterialShaderMapId::AppendKeyString(FString& KeyString) const
 	{
 		const FShaderTypeDependency& ShaderTypeDependency = ShaderTypeDependencies[ShaderIndex];
 		KeyString += TEXT("_");
-		if (ShaderTypeDependency.ShaderPipelineType)
-		{
-			KeyString += ShaderTypeDependency.ShaderPipelineType->GetName();
-			KeyString += TEXT("_");
-		}
 		KeyString += ShaderTypeDependency.ShaderType->GetName();
 		KeyString += ShaderTypeDependency.SourceHash.ToString();
 		ShaderTypeDependency.ShaderType->GetSerializationHistory().AppendKeyString(KeyString);
@@ -693,6 +703,26 @@ void FMaterialShaderMapId::AppendKeyString(FString& KeyString) const
 		}
 	}
 
+	// Add the inputs for any shader pipelines that are stored inline in the shader map
+	for (int32 TypeIndex = 0; TypeIndex < ShaderPipelineTypeDependencies.Num(); TypeIndex++)
+	{
+		const FShaderPipelineTypeDependency& Dependency = ShaderPipelineTypeDependencies[TypeIndex];
+		KeyString += TEXT("_");
+		KeyString += Dependency.ShaderPipelineType->GetName();
+		KeyString += Dependency.StagesSourceHash.ToString();
+
+		for (const FShaderType* ShaderType : Dependency.ShaderPipelineType->GetStages())
+		{
+			const TMap<const TCHAR*, FCachedUniformBufferDeclaration>& ReferencedUniformBufferStructsCache = ShaderType->GetReferencedUniformBufferStructsCache();
+
+			// Gather referenced uniform buffers
+			for (TMap<const TCHAR*, FCachedUniformBufferDeclaration>::TConstIterator It(ReferencedUniformBufferStructsCache); It; ++It)
+			{
+				ReferencedUniformBuffers.Add(It.Key(), It.Value());
+			}
+		}
+	}
+
 	// Add the inputs for any shaders that are stored inline in the shader map
 	for (int32 VFIndex = 0; VFIndex < VertexFactoryTypeDependencies.Num(); VFIndex++)
 	{
@@ -752,19 +782,13 @@ void FMaterialShaderMapId::SetShaderDependencies(const TArray<FShaderType*>& Sha
 			VertexFactoryTypeDependencies.Add(Dependency);
 		}
 
-		//#todo-rco: Is this currently necessary? Might be redundant until we only have Pipelines...
 		for (int32 TypeIndex = 0; TypeIndex < ShaderPipelineTypes.Num(); TypeIndex++)
 		{
-			auto& Stages = ShaderPipelineTypes[TypeIndex]->GetStages();
-			for (int32 Index = 0; Index < Stages.Num(); ++Index)
-			{
-				FShaderType* GlobalShaderType = (FShaderType*)Stages[Index];
-				FShaderTypeDependency Dependency;
-				Dependency.ShaderPipelineType = ShaderPipelineTypes[TypeIndex];
-				Dependency.ShaderType = (FShaderType*)Stages[Index];
-				Dependency.SourceHash = Stages[Index]->GetSourceHash();
-				ShaderTypeDependencies.Add(Dependency);
-			}
+			const FShaderPipelineType* Pipeline = ShaderPipelineTypes[TypeIndex];
+			FShaderPipelineTypeDependency Dependency;
+			Dependency.ShaderPipelineType = Pipeline;
+			Dependency.StagesSourceHash = Pipeline->GetSourceHash();
+			ShaderPipelineTypeDependencies.Add(Dependency);
 		}
 	}
 }
@@ -773,7 +797,7 @@ void FMaterialShaderMapId::SetShaderDependencies(const TArray<FShaderType*>& Sha
  * Enqueues a compilation for a new shader of this type.
  * @param Material - The material to link the shader with.
  */
-void FMaterialShaderType::BeginCompileShader(
+FShaderCompileJob* FMaterialShaderType::BeginCompileShader(
 	uint32 ShaderMapId,
 	const FMaterial* Material,
 	FShaderCompilerEnvironment* MaterialEnvironment,
@@ -807,6 +831,7 @@ void FMaterialShaderType::BeginCompileShader(
 		NewJob,
 		NewJobs
 		);
+	return NewJob;
 }
 
 void FMaterialShaderType::BeginCompileShaderPipeline(
@@ -923,6 +948,13 @@ void FMaterialShaderMap::FixupShaderTypes(EShaderPlatform Platform, const TMap<F
 			Key.ShaderTypeDependencies[ShaderTypeIndex].ShaderType = FoundShaderType;
 		}
 
+		for (int32 ShaderPipelineIndex = 0; ShaderPipelineIndex < Key.ShaderPipelineTypeDependencies.Num(); ShaderPipelineIndex++)
+		{
+			const FString& ShaderPipelineTypeName = ShaderPipelineTypeNames.FindChecked(Key.ShaderPipelineTypeDependencies[ShaderPipelineIndex].ShaderPipelineType);
+			const FShaderPipelineType* FoundShaderPipelineType = FShaderPipelineType::GetShaderPipelineTypeByName(*ShaderPipelineTypeName);
+			Key.ShaderPipelineTypeDependencies[ShaderPipelineIndex].ShaderPipelineType = FoundShaderPipelineType;
+		}
+
 		for (int32 VFTypeIndex = 0; VFTypeIndex < Key.VertexFactoryTypeDependencies.Num(); VFTypeIndex++)
 		{
 			const FString& VFTypeName = VertexFactoryTypeNames.FindChecked(Key.VertexFactoryTypeDependencies[VFTypeIndex].VertexFactoryType);
@@ -930,9 +962,6 @@ void FMaterialShaderMap::FixupShaderTypes(EShaderPlatform Platform, const TMap<F
 			Key.VertexFactoryTypeDependencies[VFTypeIndex].VertexFactoryType = FoundVFType;
 		}
 
-		//#todo-rco
-		check(ShaderPipelineTypeNames.Num() == 0);
-
 		FMaterialShaderMap::GIdToMaterialShaderMap[Platform].Add(Key, Values[PairIndex]);
 	}
 }
@@ -970,7 +999,7 @@ void FMaterialShaderMap::LoadFromDerivedDataCache(const FMaterial* Material, con
 			}
 			else
 			{
-				InOutShaderMap = NULL;
+				InOutShaderMap = nullptr;
 			}
 		}
 		INC_FLOAT_STAT_BY(STAT_ShaderCompiling_DDCLoading,(float)MaterialDDCTime);
@@ -1236,8 +1265,6 @@ void FMaterialShaderMap::Compile(
 	}
 	else
 	{
-
-
 		check(!Material->bContainsInlineShaders);
   
 		// Make sure we are operating on a referenced shader map or the below Find will cause this shader map to be deleted,
@@ -1368,6 +1395,7 @@ void FMaterialShaderMap::Compile(
 			}
   
 			// Iterate over all material shader types.
+			TMap<FShaderType*, FShaderCompileJob*> SharedShaderJobs;
 			for(TLinkedList<FShaderType*>::TIterator ShaderTypeIt(FShaderType::GetTypeList());ShaderTypeIt;ShaderTypeIt.Next())
 			{
 				FMaterialShaderType* ShaderType = ShaderTypeIt->GetMaterialShaderType();
@@ -1382,7 +1410,7 @@ void FMaterialShaderMap::Compile(
 					// Only compile the shader if we don't already have it
 					if (!HasShader(ShaderType))
 					{
-						ShaderType->BeginCompileShader(
+						auto* Job = ShaderType->BeginCompileShader(
 							CompilingId,
 							Material,
 							MaterialEnvironment,
@@ -1390,6 +1418,8 @@ void FMaterialShaderMap::Compile(
 							InPlatform,
 							NewJobs
 							);
+						check(!SharedShaderJobs.Find(ShaderType));
+						SharedShaderJobs.Add(ShaderType, Job);
 					}
 					NumShaders++;
 				}
@@ -1400,13 +1430,11 @@ void FMaterialShaderMap::Compile(
 				const FShaderPipelineType* Pipeline = *ShaderPipelineIt;
 				if (Pipeline->IsMaterialTypePipeline())
 				{
-					auto& Stages = Pipeline->GetStages();
-					bool bCompile = true;
+					auto& StageTypes = Pipeline->GetStages();
 					TArray<FMaterialShaderType*> ShaderStagesToCompile;
-					bool bCompilePipeline = true;
-					for (int32 Index = 0; Index < Stages.Num(); ++Index)
+					for (int32 Index = 0; Index < StageTypes.Num(); ++Index)
 					{
-						FMaterialShaderType* ShaderType = (FMaterialShaderType*)(Stages[Index]->GetMaterialShaderType());
+						FMaterialShaderType* ShaderType = (FMaterialShaderType*)(StageTypes[Index]->GetMaterialShaderType());
 						if (ShaderType && ShouldCacheMaterialShader(ShaderType, InPlatform, Material))
 						{
 							// Verify that the shader map Id contains inputs for any shaders that will be put into this shader map
@@ -1415,15 +1443,32 @@ void FMaterialShaderMap::Compile(
 						}
 						else
 						{
-							bCompilePipeline = false;
 							break;
 						}
+					}
 
-						if (bCompilePipeline)
+					if (ShaderStagesToCompile.Num() == StageTypes.Num())
+					{
+						// Verify that the shader map Id contains inputs for any shaders that will be put into this shader map
+						check(InShaderMapId.ContainsShaderPipelineType(Pipeline));
+						if (Pipeline->ShouldOptimizeUnusedOutputs())
 						{
 							NumShaders += ShaderStagesToCompile.Num();
 							FMaterialShaderType::BeginCompileShaderPipeline(CompilingId, InPlatform, Material, MaterialEnvironment, Pipeline, ShaderStagesToCompile, NewJobs);
 						}
+						else
+						{
+							// If sharing shaders amongst pipelines, add this pipeline as a dependency of an existing job
+							for (const FShaderType* ShaderType : StageTypes)
+							{
+								FShaderCompileJob** Job = SharedShaderJobs.Find(ShaderType);
+								checkf(Job, TEXT("Couldn't find existing shared job for material shader %s on pipeline %s!"), ShaderType->GetName(), Pipeline->GetName());
+								auto* SingleJob = (*Job)->GetSingleShaderJob();
+								auto& PipelinesToShare = SingleJob->SharingPipelines.FindOrAdd(nullptr);
+								check(!PipelinesToShare.Contains(Pipeline));
+								PipelinesToShare.Add(Pipeline);
+							}
+						}
 					}
 				}
 			}
@@ -1508,7 +1553,7 @@ FShader* FMaterialShaderMap::ProcessCompilationResultsForSingleJob(FShaderCompil
 	return Shader;
 }
 
-bool FMaterialShaderMap::ProcessCompilationResults(const TArray<FShaderCommonCompileJob*>& InCompilationResults, int32& InOutJobIndex, float& TimeBudget)
+bool FMaterialShaderMap::ProcessCompilationResults(const TArray<FShaderCommonCompileJob*>& InCompilationResults, int32& InOutJobIndex, float& TimeBudget, TMap<const FVertexFactoryType*, TArray<const FShaderPipelineType*> >& SharedPipelines)
 {
 	check(InOutJobIndex < InCompilationResults.Num());
 
@@ -1523,6 +1568,14 @@ bool FMaterialShaderMap::ProcessCompilationResults(const TArray<FShaderCommonCom
 		if (SingleJob)
 		{
 			ProcessCompilationResultsForSingleJob(SingleJob, nullptr, MaterialShaderMapHash);
+			for (auto Pair : SingleJob->SharingPipelines)
+			{
+				auto& SharedPipelinesPerVF = SharedPipelines.FindOrAdd(SingleJob->VFType);
+				for (auto* Pipeline : Pair.Value)
+				{
+					SharedPipelinesPerVF.AddUnique(Pipeline);
+				}
+			}
 		}
 		else
 		{
@@ -1563,12 +1616,12 @@ bool FMaterialShaderMap::ProcessCompilationResults(const TArray<FShaderCommonCom
 					}
 
 					check(MeshShaderMap);
-					check(!MeshShaderMap->GetShaderPipeline(ShaderPipeline->PipelineType));
+					check(!MeshShaderMap->HasShaderPipeline(ShaderPipeline->PipelineType));
 					MeshShaderMap->AddShaderPipeline(PipelineJob->ShaderPipeline, ShaderPipeline);
 				}
 				else
 				{
-					check(!GetShaderPipeline(ShaderPipeline->PipelineType));
+					check(!HasShaderPipeline(ShaderPipeline->PipelineType));
 					AddShaderPipeline(PipelineJob->ShaderPipeline, ShaderPipeline);
 				}
 			}
@@ -1584,6 +1637,63 @@ bool FMaterialShaderMap::ProcessCompilationResults(const TArray<FShaderCommonCom
 
 	if (InOutJobIndex == InCompilationResults.Num())
 	{
+		{
+			// Process the mesh shader pipelines that share shaders
+			for (int32 ShaderMapIndex = 0; ShaderMapIndex < MeshShaderMaps.Num(); ShaderMapIndex++)
+			{
+				auto* MeshShaderMap = &MeshShaderMaps[ShaderMapIndex];
+				auto* VertexFactory = MeshShaderMap->GetVertexFactoryType();
+				auto* FoundSharedPipelines = SharedPipelines.Find(VertexFactory);
+				if (VertexFactory && FoundSharedPipelines)
+				{
+					for (const FShaderPipelineType* ShaderPipelineType : *FoundSharedPipelines)
+					{
+						if (ShaderPipelineType->IsMeshMaterialTypePipeline() && !MeshShaderMap->HasShaderPipeline(ShaderPipelineType))
+						{
+							auto& StageTypes = ShaderPipelineType->GetStages();
+							TArray<FShader*> ShaderStages;
+							for (int32 Index = 0; Index < StageTypes.Num(); ++Index)
+							{
+								FMeshMaterialShaderType* ShaderType = ((FMeshMaterialShaderType*)(StageTypes[Index]))->GetMeshMaterialShaderType();
+								FShader* Shader = MeshShaderMap->GetShader(ShaderType);
+								check(Shader);
+								ShaderStages.Add(Shader);
+							}
+
+							checkf(StageTypes.Num() == ShaderStages.Num(), TEXT("Internal Error adding MeshMaterial ShaderPipeline %s"), ShaderPipelineType->GetName());
+							FShaderPipeline* ShaderPipeline = new FShaderPipeline(ShaderPipelineType, ShaderStages);
+							MeshShaderMap->AddShaderPipeline(ShaderPipelineType, ShaderPipeline);
+						}
+					}
+				}
+			}
+
+			// Process the material shader pipelines that share shaders
+			auto* FoundSharedPipelines = SharedPipelines.Find(nullptr);
+			if (FoundSharedPipelines)
+			{
+				for (const FShaderPipelineType* ShaderPipelineType : *FoundSharedPipelines)
+				{
+					if (ShaderPipelineType->IsMaterialTypePipeline() && !HasShaderPipeline(ShaderPipelineType))
+					{
+						auto& StageTypes = ShaderPipelineType->GetStages();
+						TArray<FShader*> ShaderStages;
+						for (int32 Index = 0; Index < StageTypes.Num(); ++Index)
+						{
+							FMaterialShaderType* ShaderType = ((FMaterialShaderType*)(StageTypes[Index]))->GetMaterialShaderType();
+							FShader* Shader = GetShader(ShaderType);
+							check(Shader);
+							ShaderStages.Add(Shader);
+						}
+
+						checkf(StageTypes.Num() == ShaderStages.Num(), TEXT("Internal Error adding Material ShaderPipeline %s"), ShaderPipelineType->GetName());
+						FShaderPipeline* ShaderPipeline = new FShaderPipeline(ShaderPipelineType, ShaderStages);
+						AddShaderPipeline(ShaderPipelineType, ShaderPipeline);
+					}
+				}
+			}
+		}
+
 		for (int32 ShaderMapIndex = MeshShaderMaps.Num() - 1; ShaderMapIndex >= 0; ShaderMapIndex--)
 		{
 			if (MeshShaderMaps[ShaderMapIndex].GetNumShaders() == 0 && MeshShaderMaps[ShaderMapIndex].GetNumShaderPipelines() == 0)
@@ -1631,7 +1741,7 @@ bool FMaterialShaderMap::IsMaterialShaderComplete(const FMaterial* Material, con
 {
 	// If we should cache this material, it's incomplete if the pipeline is missing or the shader itself is missing
 	if (ShouldCacheMaterialShader(ShaderType, Platform, Material) &&
-		((Pipeline && !GetShaderPipeline(Pipeline)) || (!Pipeline && !HasShader((FShaderType*)ShaderType))))
+		((Pipeline && !HasShaderPipeline(Pipeline)) || (!Pipeline && !HasShader((FShaderType*)ShaderType))))
 	{
 		if (!bSilent)
 		{
@@ -1700,34 +1810,30 @@ bool FMaterialShaderMap::IsComplete(const FMaterial* Material, bool bSilent)
 		const FShaderPipelineType* Pipeline = *ShaderPipelineIt;
 		if (Pipeline->IsMaterialTypePipeline())
 		{
-			auto& Stages = Pipeline->GetStages();
+			auto& StageTypes = Pipeline->GetStages();
 
-			// Verify all the ShouldCache are in sync
 			int32 NumShouldCache = 0;
-			for (int32 Index = 0; Index < Stages.Num(); ++Index)
+			for (int32 Index = 0; Index < StageTypes.Num(); ++Index)
 			{
-				auto* ShaderType = Stages[Index]->GetMaterialShaderType();
+				auto* ShaderType = StageTypes[Index]->GetMaterialShaderType();
 				if (ShouldCacheMaterialShader(ShaderType, Platform, Material))
 				{
 					++NumShouldCache;
 				}
+				else
+				{
+					break;
+				}
 			}
 
-			if (NumShouldCache > 0)
+			if (NumShouldCache == StageTypes.Num())
 			{
-				if (NumShouldCache != Stages.Num())
-				{
-					UE_LOG(LogShaders, Error, TEXT("Mismatched ShouldCache() on Material %s for MeshMaterial ShaderPipeline %s"), *Material->GetFriendlyName(), Pipeline->GetName());
-				}
-				else
+				for (int32 Index = 0; Index < StageTypes.Num(); ++Index)
 				{
-					for (int32 Index = 0; Index < Stages.Num(); ++Index)
+					auto* ShaderType = StageTypes[Index]->GetMaterialShaderType();
+					if (!IsMaterialShaderComplete(Material, ShaderType, Pipeline, bSilent))
 					{
-						auto* ShaderType = Stages[Index]->GetMaterialShaderType();
-						if (!IsMaterialShaderComplete(Material, ShaderType, Pipeline, bSilent))
-						{
-							return false;
-						}
+						return false;
 					}
 				}
 			}
@@ -1755,17 +1861,13 @@ void FMaterialShaderMap::LoadMissingShadersFromMemory(const FMaterial* Material)
 	ShaderMapId.GetMaterialHash(MaterialShaderMapHash);
 
 	// Try to find necessary FMaterialShaderType's in memory
-	for(TLinkedList<FShaderType*>::TIterator ShaderTypeIt(FShaderType::GetTypeList());ShaderTypeIt;ShaderTypeIt.Next())
+	for (TLinkedList<FShaderType*>::TIterator ShaderTypeIt(FShaderType::GetTypeList());ShaderTypeIt;ShaderTypeIt.Next())
 	{
 		FMaterialShaderType* ShaderType = ShaderTypeIt->GetMaterialShaderType();
-
-		const bool bShaderAlreadyExists = HasShader(ShaderType);
-
-		if (ShaderType && ShouldCacheMaterialShader(ShaderType, Platform, Material) && !bShaderAlreadyExists)
+		if (ShaderType && ShouldCacheMaterialShader(ShaderType, Platform, Material) && !HasShader(ShaderType))
 		{
 			FShaderId ShaderId(MaterialShaderMapHash, nullptr, nullptr, ShaderType, FShaderTarget(ShaderType->GetFrequency(), Platform));
-			FShader* FoundShader = ShaderType->FindShaderById(ShaderId);	
-
+			FShader* FoundShader = ShaderType->FindShaderById(ShaderId);
 			if (FoundShader)
 			{
 				AddShader(ShaderType, FoundShader);
@@ -1773,6 +1875,50 @@ void FMaterialShaderMap::LoadMissingShadersFromMemory(const FMaterial* Material)
 		}
 	}
 
+	// Try to find necessary FShaderPipelineTypes in memory
+	for (TLinkedList<FShaderPipelineType*>::TIterator ShaderPipelineIt(FShaderPipelineType::GetTypeList());ShaderPipelineIt;ShaderPipelineIt.Next())
+	{
+		const FShaderPipelineType* PipelineType = *ShaderPipelineIt;
+		if (PipelineType && PipelineType->IsMaterialTypePipeline() && !HasShaderPipeline(PipelineType))
+		{
+			auto& Stages = PipelineType->GetStages();
+			int32 NumShaders = 0;
+			for (const FShaderType* Shader : Stages)
+			{
+				FMaterialShaderType* ShaderType = (FMaterialShaderType*)Shader->GetMaterialShaderType();
+				if (ShaderType && ShouldCacheMaterialShader(ShaderType, Platform, Material))
+				{
+					++NumShaders;
+				}
+			}
+
+			if (NumShaders == Stages.Num())
+			{
+				TArray<FShader*> ShadersForPipeline;
+				for (auto* Shader : Stages)
+				{
+					FMaterialShaderType* ShaderType = (FMaterialShaderType*)Shader->GetMaterialShaderType();
+					if (!HasShader(ShaderType))
+					{
+						FShaderId ShaderId(MaterialShaderMapHash, PipelineType->ShouldOptimizeUnusedOutputs() ? PipelineType : nullptr, nullptr, ShaderType, FShaderTarget(ShaderType->GetFrequency(), Platform));
+						FShader* FoundShader = ShaderType->FindShaderById(ShaderId);
+						if (FoundShader)
+						{
+							AddShader(ShaderType, FoundShader);
+							ShadersForPipeline.Add(FoundShader);
+						}
+					}
+				}
+
+				if (ShadersForPipeline.Num() == NumShaders && !HasShaderPipeline(PipelineType))
+				{
+					auto* Pipeline = new FShaderPipeline(PipelineType, ShadersForPipeline);
+					AddShaderPipeline(PipelineType, Pipeline);
+				}
+			}
+		}
+	}
+
 	// Try to find necessary FMeshMaterialShaderMap's in memory
 	for(TLinkedList<FVertexFactoryType*>::TIterator VertexFactoryTypeIt(FVertexFactoryType::GetTypeList());VertexFactoryTypeIt;VertexFactoryTypeIt.Next())
 	{
@@ -1802,10 +1948,10 @@ void FMaterialShaderMap::GetShaderList(TMap<FShaderId, FShader*>& OutShaders) co
 
 void FMaterialShaderMap::GetShaderPipelineList(TArray<FShaderPipeline*>& OutShaderPipelines) const
 {
-	TShaderMap<FMaterialShaderType>::GetShaderPipelineList(OutShaderPipelines);
+	TShaderMap<FMaterialShaderType>::GetShaderPipelineList(OutShaderPipelines, FShaderPipeline::EAll);
 	for (int32 Index = 0;Index < MeshShaderMaps.Num();Index++)
 	{
-		MeshShaderMaps[Index].GetShaderPipelineList(OutShaderPipelines);
+		MeshShaderMaps[Index].GetShaderPipelineList(OutShaderPipelines, FShaderPipeline::EAll);
 	}
 }
 
@@ -2055,7 +2201,7 @@ void FMaterialShaderMap::Serialize(FArchive& Ar, bool bInlineShaderResources)
 
 		for (int32 VFIndex = 0; VFIndex < NumMeshShaderMaps; VFIndex++)
 		{
-			FVertexFactoryType* VFType = NULL;
+			FVertexFactoryType* VFType = nullptr;
 
 			Ar << VFType;
 
@@ -2161,6 +2307,8 @@ void DumpMaterialStats(EShaderPlatform Platform)
 
 	// Mapping from friendly material name to shaders associated with it.
 	TMultiMap<FString,FShader*> MaterialToShaderMap;
+	TMultiMap<FString, FShaderPipeline*> MaterialToShaderPipelineMap;
+
 	// Set of material names.
 	TSet<FString> MaterialNames;
 
@@ -2190,6 +2338,7 @@ void DumpMaterialStats(EShaderPlatform Platform)
 			{
 				MaterialToShaderMap.AddUnique(FriendlyName, Shader);
 			}
+			MaterialToShaderPipelineMap.AddUnique(FriendlyName, Pipeline);
 		}
 	}
 
@@ -2197,16 +2346,20 @@ void DumpMaterialStats(EShaderPlatform Platform)
 	MaterialViewer.AddColumn(TEXT("Name"));
 	MaterialViewer.AddColumn(TEXT("Shaders"));
 	MaterialViewer.AddColumn(TEXT("Code Size"));
+	MaterialViewer.AddColumn(TEXT("Pipelines"));
 	MaterialViewer.CycleRow();
 
 	// Iterate over all materials, gathering shader stats.
 	int32 TotalCodeSize		= 0;
 	int32 TotalShaderCount	= 0;
+	int32 TotalShaderPipelineCount = 0;
 	for( TSet<FString>::TConstIterator It(MaterialNames); It; ++It )
 	{
 		// Retrieve list of shaders in map.
 		TArray<FShader*> Shaders;
 		MaterialToShaderMap.MultiFind( *It, Shaders );
+		TArray<FShaderPipeline*> ShaderPipelines;
+		MaterialToShaderPipelineMap.MultiFind(*It, ShaderPipelines);
 		
 		// Iterate over shaders and gather stats.
 		int32 CodeSize = 0;
@@ -2218,11 +2371,13 @@ void DumpMaterialStats(EShaderPlatform Platform)
 
 		TotalCodeSize += CodeSize;
 		TotalShaderCount += Shaders.Num();
+		TotalShaderPipelineCount += ShaderPipelines.Num();
 
 		// Dump stats
 		MaterialViewer.AddColumn(**It);
 		MaterialViewer.AddColumn(TEXT("%u"),Shaders.Num());
 		MaterialViewer.AddColumn(TEXT("%u"),CodeSize);
+		MaterialViewer.AddColumn(TEXT("%u"), ShaderPipelines.Num());
 		MaterialViewer.CycleRow();
 	}
 
@@ -2230,6 +2385,7 @@ void DumpMaterialStats(EShaderPlatform Platform)
 	MaterialViewer.AddColumn(TEXT("Total"));
 	MaterialViewer.AddColumn(TEXT("%u"),TotalShaderCount);
 	MaterialViewer.AddColumn(TEXT("%u"),TotalCodeSize);
+	MaterialViewer.AddColumn(TEXT("%u"), TotalShaderPipelineCount);
 	MaterialViewer.CycleRow();
 #endif
 }
diff --git a/Engine/Source/Runtime/Engine/Private/Materials/MaterialShared.cpp b/Engine/Source/Runtime/Engine/Private/Materials/MaterialShared.cpp
--- a/Engine/Source/Runtime/Engine/Private/Materials/MaterialShared.cpp
+++ b/Engine/Source/Runtime/Engine/Private/Materials/MaterialShared.cpp
@@ -911,27 +911,26 @@ void FMaterialResource::NotifyCompilationFinished()
  */
 void FMaterialResource::GetRepresentativeInstructionCounts(TArray<FString> &Descriptions, TArray<int32> &InstructionCounts) const
 {
-	TArray<FString> ShaderTypeNames;
-	TArray<FString> ShaderTypeDescriptions;
+	TMap<FName, FString> ShaderTypeNamesAndDescriptions;
 
 	//when adding a shader type here be sure to update FPreviewMaterial::ShouldCache()
 	//so the shader type will get compiled with preview materials
 	const FMaterialShaderMap* MaterialShaderMap = GetGameThreadShaderMap();
 	if (MaterialShaderMap && MaterialShaderMap->IsCompilationFinalized())
 	{
-		GetRepresentativeShaderTypesAndDescriptions(ShaderTypeNames, ShaderTypeDescriptions);
+		GetRepresentativeShaderTypesAndDescriptions(ShaderTypeNamesAndDescriptions);
 
 		if( IsUIMaterial() )
 		{
-			for (int32 InstructionIndex = 0; InstructionIndex < ShaderTypeNames.Num(); InstructionIndex++)
+			for (const TPair<FName, FString>& ShaderTypePair : ShaderTypeNamesAndDescriptions)
 			{
-				FShaderType* ShaderType = FindShaderTypeByName(*ShaderTypeNames[InstructionIndex]);
+				FShaderType* ShaderType = FindShaderTypeByName(ShaderTypePair.Key);
 				int32 NumInstructions = MaterialShaderMap->GetMaxNumInstructionsForShader(ShaderType);
 				if (NumInstructions > 0)
 				{
 					//if the shader was found, add it to the output arrays
 					InstructionCounts.Push(NumInstructions);
-					Descriptions.Push(ShaderTypeDescriptions[InstructionIndex]);
+					Descriptions.Push(ShaderTypePair.Value);
 				}
 			}
 		}
@@ -943,9 +942,9 @@ void FMaterialResource::GetRepresentativeInstructionCounts(TArray<FString> &Desc
 				Descriptions.Empty();
 				InstructionCounts.Empty();
 
-				for (int32 InstructionIndex = 0; InstructionIndex < ShaderTypeNames.Num(); InstructionIndex++)
+				for (const TPair<FName, FString>& ShaderTypePair : ShaderTypeNamesAndDescriptions)
 				{
-					FShaderType* ShaderType = FindShaderTypeByName(*ShaderTypeNames[InstructionIndex]);
+					FShaderType* ShaderType = FindShaderTypeByName(ShaderTypePair.Key);
 					if (ShaderType)
 					{
 						int32 NumInstructions = MeshShaderMap->GetMaxNumInstructionsForShader(ShaderType);
@@ -953,7 +952,7 @@ void FMaterialResource::GetRepresentativeInstructionCounts(TArray<FString> &Desc
 						{
 							//if the shader was found, add it to the output arrays
 							InstructionCounts.Push(NumInstructions);
-							Descriptions.Push(ShaderTypeDescriptions[InstructionIndex]);
+							Descriptions.Push(ShaderTypePair.Value);
 						}
 					}
 				}
@@ -964,84 +963,99 @@ void FMaterialResource::GetRepresentativeInstructionCounts(TArray<FString> &Desc
 	check(Descriptions.Num() == InstructionCounts.Num());
 }
 
-void FMaterialResource::GetRepresentativeShaderTypesAndDescriptions(TArray<FString> &ShaderTypeNames, TArray<FString> &ShaderTypeDescriptions) const
+void FMaterialResource::GetRepresentativeShaderTypesAndDescriptions(TMap<FName, FString>& ShaderTypeNamesAndDescriptions) const
 {
 	static auto* MobileHDR = IConsoleManager::Get().FindTConsoleVariableDataInt(TEXT("r.MobileHDR"));
 	bool bMobileHDR = MobileHDR && MobileHDR->GetValueOnAnyThread() == 1;
 
 	if( IsUIMaterial() )
 	{
-		new(ShaderTypeNames)FString(TEXT("TSlateMaterialShaderPSDefaultfalse"));
-		new(ShaderTypeDescriptions)FString(TEXT("Default UI Pixel Shader"));
+		static FName TSlateMaterialShaderPSDefaultfalseName = TEXT("TSlateMaterialShaderPSDefaultfalse");
+		ShaderTypeNamesAndDescriptions.Add(TSlateMaterialShaderPSDefaultfalseName, TEXT("Default UI Pixel Shader"));
 
-		new(ShaderTypeNames) FString(TEXT("TSlateMaterialShaderVSfalse"));
-		new(ShaderTypeDescriptions) FString(TEXT("Default UI Vertex Shader"));
+		static FName TSlateMaterialShaderVSfalseName = TEXT("TSlateMaterialShaderVSfalse");
+		ShaderTypeNamesAndDescriptions.Add(TSlateMaterialShaderVSfalseName, TEXT("Default UI Vertex Shader"));
 
-		new(ShaderTypeNames)FString(TEXT("TSlateMaterialShaderVStrue"));
-		new(ShaderTypeDescriptions)FString(TEXT("Instanced UI Vertex Shader"));
+		static FName TSlateMaterialShaderVStrueName = TEXT("TSlateMaterialShaderVStrue");
+		ShaderTypeNamesAndDescriptions.Add(TSlateMaterialShaderVStrueName, TEXT("Instanced UI Vertex Shader"));
 	}
 	else if (GetFeatureLevel() >= ERHIFeatureLevel::SM4)
 	{
 		if (GetShadingModel() == MSM_Unlit)
 		{
 			//unlit materials are never lightmapped
-			new (ShaderTypeNames)FString(TEXT("TBasePassPSFNoLightMapPolicy"));
-			new (ShaderTypeDescriptions)FString(TEXT("Base pass shader without light map"));
+			static FName TBasePassPSFNoLightMapPolicyName = TEXT("TBasePassPSFNoLightMapPolicy");
+			ShaderTypeNamesAndDescriptions.Add(TBasePassPSFNoLightMapPolicyName, TEXT("Base pass shader without light map"));
 		}
 		else
 		{
 			if (IsUsedWithStaticLighting())
 			{
 				//lit materials are usually lightmapped
-				new (ShaderTypeNames)FString(TEXT("TBasePassPSTDistanceFieldShadowsAndLightMapPolicyHQ"));
-				new (ShaderTypeDescriptions)FString(TEXT("Base pass shader with static lighting"));
+				static FName TBasePassPSTDistanceFieldShadowsAndLightMapPolicyHQName = TEXT("TBasePassPSTDistanceFieldShadowsAndLightMapPolicyHQ");
+				ShaderTypeNamesAndDescriptions.Add(TBasePassPSTDistanceFieldShadowsAndLightMapPolicyHQName, TEXT("Base pass shader with static lighting"));
 			}
 
 			//also show a dynamically lit shader
-			new (ShaderTypeNames)FString(TEXT("TBasePassPSFNoLightMapPolicy"));
-			new (ShaderTypeDescriptions)FString(TEXT("Base pass shader with only dynamic lighting"));
+			static FName TBasePassPSFNoLightMapPolicyName = TEXT("TBasePassPSFNoLightMapPolicy");
+			ShaderTypeNamesAndDescriptions.Add(TBasePassPSFNoLightMapPolicyName, TEXT("Base pass shader with only dynamic lighting"));
 
 			if (IsTranslucentBlendMode(GetBlendMode()))
 			{
-				new (ShaderTypeNames)FString(TEXT("TBasePassPSFSelfShadowedTranslucencyPolicy"));
-				new (ShaderTypeDescriptions)FString(TEXT("Base pass shader for self shadowed translucency"));
+				static FName TBasePassPSFSelfShadowedTranslucencyPolicyName = TEXT("TBasePassPSFSelfShadowedTranslucencyPolicy");
+				ShaderTypeNamesAndDescriptions.Add(TBasePassPSFSelfShadowedTranslucencyPolicyName, TEXT("Base pass shader for self shadowed translucency"));
 			}
 		}
 
-		new (ShaderTypeNames)FString(TEXT("TBasePassVSFNoLightMapPolicy"));
-		new (ShaderTypeDescriptions)FString(TEXT("Vertex shader"));
+		static FName TBasePassVSFNoLightMapPolicyName = TEXT("TBasePassVSFNoLightMapPolicy");
+		ShaderTypeNamesAndDescriptions.Add(TBasePassVSFNoLightMapPolicyName, TEXT("Vertex shader"));
 	}
 	else
 	{
-		const TCHAR* ShaderSuffix = bMobileHDR ? TEXT("HDRLinear64") : TEXT("LDRGamma32");
 		const TCHAR* DescSuffix = bMobileHDR ? TEXT(" (HDR)") : TEXT(" (LDR)");
 
 		if (GetShadingModel() == MSM_Unlit)
 		{
 			//unlit materials are never lightmapped
-			new (ShaderTypeNames)FString(FString::Printf(TEXT("TBasePassForForwardShadingPSFNoLightMapPolicy0%s"), ShaderSuffix));
-			new (ShaderTypeDescriptions)FString(FString::Printf(TEXT("Mobile base pass shader without light map%s"), DescSuffix));
+			static FName Name_HDRLinear64 = TEXT("TBasePassForForwardShadingPSFNoLightMapPolicy0HDRLinear64");
+			static FName Name_LDRGamma32 = TEXT("TBasePassForForwardShadingPSFNoLightMapPolicy0LDRGamma32");
+			const FName TBasePassForForwardShadingPSFNoLightMapPolicy0Name = bMobileHDR ? Name_HDRLinear64 : Name_LDRGamma32;
+			ShaderTypeNamesAndDescriptions.Add(TBasePassForForwardShadingPSFNoLightMapPolicy0Name, FString::Printf(TEXT("Mobile base pass shader without light map%s"), DescSuffix));
 		}
 		else
 		{
 			if (IsUsedWithStaticLighting())
 			{
 				//lit materials are usually lightmapped
-				new (ShaderTypeNames)FString(FString::Printf(TEXT("TBasePassForForwardShadingPSTLightMapPolicy0LQ%s"), ShaderSuffix));
-				new (ShaderTypeDescriptions)FString(FString::Printf(TEXT("Mobile base pass shader with static lighting%s"), DescSuffix));
+				{
+					static FName Name_HDRLinear64 = TEXT("TBasePassForForwardShadingPSTLightMapPolicy0LQHDRLinear64");
+					static FName Name_LDRGamma32 = TEXT("TBasePassForForwardShadingPSTLightMapPolicy0LQLDRGamma32");
+					static FName TSlateMaterialShaderVStrueName = bMobileHDR ? Name_HDRLinear64 : Name_LDRGamma32;
+					ShaderTypeNamesAndDescriptions.Add(TSlateMaterialShaderVStrueName, FString::Printf(TEXT("Mobile base pass shader with static lighting%s"), DescSuffix));
+				}
 
 				// + distance field shadows
-				new (ShaderTypeNames)FString(FString::Printf(TEXT("TBasePassForForwardShadingPSTDistanceFieldShadowsAndLightMapPolicy0LQ%s"), ShaderSuffix));
-				new (ShaderTypeDescriptions)FString(FString::Printf(TEXT("Mobile base pass shader with distance field shadows%s"), DescSuffix));
+				{
+					static FName Name_HDRLinear64 = TEXT("TBasePassForForwardShadingPSTDistanceFieldShadowsAndLightMapPolicy0LQHDRLinear64");
+					static FName Name_LDRGamma32 = TEXT("TBasePassForForwardShadingPSTDistanceFieldShadowsAndLightMapPolicy0LQLDRGamma32");
+					static FName TBasePassForForwardShadingPSTDistanceFieldShadowsAndLightMapPolicy0LQName = bMobileHDR ? Name_HDRLinear64 : Name_LDRGamma32;
+					ShaderTypeNamesAndDescriptions.Add(TBasePassForForwardShadingPSTDistanceFieldShadowsAndLightMapPolicy0LQName, FString::Printf(TEXT("Mobile base pass shader with distance field shadows%s"), DescSuffix));
+				}
 			}
 
 			//also show a dynamically lit shader
-			new (ShaderTypeNames)FString(FString::Printf(TEXT("TBasePassForForwardShadingPSFSimpleDirectionalLightAndSHIndirectPolicy0%s"), ShaderSuffix));
-			new (ShaderTypeDescriptions)FString(FString::Printf(TEXT("Mobile base pass shader with only dynamic lighting%s"), DescSuffix));
+			static FName Name_HDRLinear64 = TEXT("TBasePassForForwardShadingPSFSimpleDirectionalLightAndSHIndirectPolicy0HDRLinear64");
+			static FName Name_LDRGamma32 = TEXT("TBasePassForForwardShadingPSFSimpleDirectionalLightAndSHIndirectPolicy0LDRGamma32");
+			const FName TBasePassForForwardShadingPSFSimpleDirectionalLightAndSHIndirectPolicy0Name = bMobileHDR ? Name_HDRLinear64 : Name_LDRGamma32;
+			ShaderTypeNamesAndDescriptions.Add(TBasePassForForwardShadingPSFSimpleDirectionalLightAndSHIndirectPolicy0Name, FString::Printf(TEXT("Mobile base pass shader with only dynamic lighting%s"), DescSuffix));
 		}
 
-		new (ShaderTypeNames)FString(FString::Printf(TEXT("TBasePassForForwardShadingVSFNoLightMapPolicy%s"), ShaderSuffix));
-		new (ShaderTypeDescriptions)FString(FString::Printf(TEXT("Mobile base pass vertex shader%s"), DescSuffix));
+		{
+			static FName Name_HDRLinear64 = TEXT("TBasePassForForwardShadingVSFNoLightMapPolicyHDRLinear64");
+			static FName Name_LDRGamma32 = TEXT("TBasePassForForwardShadingVSFNoLightMapPolicyLDRGamma32");
+			const FName TBasePassForForwardShadingVSFNoLightMapPolicyName = bMobileHDR ? Name_HDRLinear64 : Name_LDRGamma32;
+			ShaderTypeNamesAndDescriptions.Add(TBasePassForForwardShadingVSFNoLightMapPolicyName, FString::Printf(TEXT("Mobile base pass vertex shader%s"), DescSuffix));
+		}
 	}
 }
 
@@ -1548,6 +1562,55 @@ FShader* FMaterial::GetShader(FMeshMaterialShaderType* ShaderType, FVertexFactor
 	return Shader;
 }
 
+FShaderPipeline* FMaterial::GetShaderPipeline(class FShaderPipelineType* ShaderPipelineType, FVertexFactoryType* VertexFactoryType) const
+{
+	const FMeshMaterialShaderMap* MeshShaderMap = RenderingThreadShaderMap->GetMeshShaderMap(VertexFactoryType);
+	FShaderPipeline* ShaderPipeline = MeshShaderMap ? MeshShaderMap->GetShaderPipeline(ShaderPipelineType) : nullptr;
+	if (!ShaderPipeline)
+	{
+		// Get the ShouldCache results that determine whether the shader should be compiled
+		auto ShaderPlatform = GShaderPlatformForFeatureLevel[GetFeatureLevel()];
+		FString MaterialUsage = GetMaterialUsageDescription();
+
+		FString Message;
+		for (auto* ShaderType : ShaderPipelineType->GetStages())
+		{
+			FShader* Shader = MeshShaderMap ? MeshShaderMap->GetShader((FShaderType*)ShaderType) : nullptr;
+			if (Shader && ShaderType->GetMeshMaterialShaderType())
+			{
+				bool bMaterialShouldCache = ShouldCache(ShaderPlatform, ShaderType->GetMeshMaterialShaderType(), VertexFactoryType);
+				bool bVFShouldCache = VertexFactoryType->ShouldCache(ShaderPlatform, this, ShaderType->GetMeshMaterialShaderType());
+				bool bShaderShouldCache = ShaderType->GetMeshMaterialShaderType()->ShouldCache(ShaderPlatform, this, VertexFactoryType);
+
+				Message += FString::Printf(TEXT("%s Freq %d, ShouldCache: Mat=%u, VF=%u, Shader=%u\n"),
+					ShaderType->GetName(), (int32)ShaderType->GetFrequency(), bMaterialShouldCache, bVFShouldCache, bShaderShouldCache);
+			}
+			else
+			{
+				Message += FString::Printf(TEXT("%s Freq %d not Found\n"),
+					ShaderType->GetName(), (int32)ShaderType->GetFrequency());
+			}
+		}
+
+		int BreakPoint = 0;
+
+		// Assert with detailed information if the shader wasn't found for rendering.  
+		// This is usually the result of an incorrect ShouldCache function.
+		UE_LOG(LogMaterial, Fatal,
+			TEXT("Couldn't find ShaderPipeline %s for Material Resource %s!\n")
+			TEXT("		With VF=%s, Platform=%s\n")
+			TEXT("		%s\n")
+			TEXT("		MaterialUsageDesc: %s"),
+			ShaderPipelineType->GetName(), *GetFriendlyName(),
+			VertexFactoryType->GetName(), *LegacyShaderPlatformToShaderFormat(ShaderPlatform).ToString(),
+			*Message,
+			*MaterialUsage
+			);
+	}
+
+	return ShaderPipeline;
+}
+
 /** Returns the index to the Expression in the Expressions array, or -1 if not found. */
 int32 FMaterial::FindExpression( const TArray<TRefCountPtr<FMaterialUniformExpressionTexture> >&Expressions, const FMaterialUniformExpressionTexture &Expression )
 {
@@ -1862,7 +1925,9 @@ void FMaterial::GetDependentShaderAndVFTypes(EShaderPlatform Platform, TArray<FS
 				auto* PipelineType = *PipelineTypeIt;
 				if (PipelineType->IsMeshMaterialTypePipeline())
 				{
-					for (const FShaderType* Type : PipelineType->GetStages())
+					int32 NumShouldCache = 0;
+					auto& ShaderStages = PipelineType->GetStages();
+					for (const FShaderType* Type : ShaderStages)
 					{
 						const FMeshMaterialShaderType* ShaderType = Type->GetMeshMaterialShaderType();
 						if (ShaderType->ShouldCache(Platform, this, VertexFactoryType) &&
@@ -1870,8 +1935,16 @@ void FMaterial::GetDependentShaderAndVFTypes(EShaderPlatform Platform, TArray<FS
 							VertexFactoryType->ShouldCache(Platform, this, ShaderType)
 							)
 						{
-							bAddedTypeFromThisVF = true;
-							OutShaderPipelineTypes.AddUnique(PipelineType);
+							++NumShouldCache;
+						}
+					}
+
+					if (NumShouldCache == ShaderStages.Num())
+					{
+						bAddedTypeFromThisVF = true;
+						OutShaderPipelineTypes.AddUnique(PipelineType);
+						for (const FShaderType* Type : ShaderStages)
+						{
 							OutShaderTypes.AddUnique((FShaderType*)Type);
 						}
 					}
@@ -1904,7 +1977,9 @@ void FMaterial::GetDependentShaderAndVFTypes(EShaderPlatform Platform, TArray<FS
 		auto* PipelineType = *PipelineTypeIt;
 		if (PipelineType->IsMaterialTypePipeline())
 		{
-			for (const FShaderType* Type : PipelineType->GetStages())
+			int32 NumShouldCache = 0;
+			auto& ShaderStages = PipelineType->GetStages();
+			for (const FShaderType* Type : ShaderStages)
 			{
 				const FMaterialShaderType* ShaderType = Type->GetMaterialShaderType();
 				if (ShaderType &&
@@ -1912,6 +1987,15 @@ void FMaterial::GetDependentShaderAndVFTypes(EShaderPlatform Platform, TArray<FS
 					ShouldCache(Platform, ShaderType, nullptr)
 					)
 				{
+					++NumShouldCache;
+				}
+			}
+
+			if (NumShouldCache == ShaderStages.Num())
+			{
+				for (const FShaderType* Type : ShaderStages)
+				{
+					const FMaterialShaderType* ShaderType = Type->GetMaterialShaderType();
 					OutShaderPipelineTypes.AddUnique(PipelineType);
 					OutShaderTypes.AddUnique((FShaderType*)Type);
 				}
@@ -2502,6 +2586,19 @@ bool FMaterialShaderMapId::ContainsShaderType(const FShaderType* ShaderType) con
 	return false;
 }
 
+bool FMaterialShaderMapId::ContainsShaderPipelineType(const FShaderPipelineType* ShaderPipelineType) const
+{
+	for (int32 TypeIndex = 0; TypeIndex < ShaderPipelineTypeDependencies.Num(); TypeIndex++)
+	{
+		if (ShaderPipelineTypeDependencies[TypeIndex].ShaderPipelineType == ShaderPipelineType)
+		{
+			return true;
+		}
+	}
+
+	return false;
+}
+
 bool FMaterialShaderMapId::ContainsVertexFactoryType(const FVertexFactoryType* VFType) const
 {
 	for (int32 TypeIndex = 0; TypeIndex < VertexFactoryTypeDependencies.Num(); TypeIndex++)
diff --git a/Engine/Source/Runtime/Engine/Private/Materials/MeshMaterialShader.cpp b/Engine/Source/Runtime/Engine/Private/Materials/MeshMaterialShader.cpp
--- a/Engine/Source/Runtime/Engine/Private/Materials/MeshMaterialShader.cpp
+++ b/Engine/Source/Runtime/Engine/Private/Materials/MeshMaterialShader.cpp
@@ -22,7 +22,7 @@ static inline bool ShouldCacheMeshShader(const FMeshMaterialShaderType* ShaderTy
  * @param Material - The material to link the shader with.
  * @param VertexFactoryType - The vertex factory to compile with.
  */
-void FMeshMaterialShaderType::BeginCompileShader(
+FShaderCompileJob* FMeshMaterialShaderType::BeginCompileShader(
 	uint32 ShaderMapId,
 	EShaderPlatform Platform,
 	const FMaterial* Material,
@@ -64,6 +64,7 @@ void FMeshMaterialShaderType::BeginCompileShader(
 		NewJobs,
 		bAllowDevelopmentShaderCompile
 		);
+	return NewJob;
 }
 
 void FMeshMaterialShaderType::BeginCompileShaderPipeline(
@@ -146,58 +147,11 @@ uint32 FMeshMaterialShaderMap::BeginCompile(
 	}
 
 	uint32 NumShadersPerVF = 0;
-
-	for (TLinkedList<FShaderPipelineType*>::TIterator ShaderPipelineIt(FShaderPipelineType::GetTypeList());ShaderPipelineIt;ShaderPipelineIt.Next())
-	{
-		const FShaderPipelineType* Pipeline = *ShaderPipelineIt;
-		if (Pipeline->IsMeshMaterialTypePipeline())
-		{
-			auto& Stages = Pipeline->GetStages();
-
-			// Verify all the ShouldCache are in sync
-			int32 NumShouldCache = 0;
-			for (int32 Index = 0; Index < Stages.Num(); ++Index)
-			{
-				auto* ShaderType = Stages[Index]->GetMeshMaterialShaderType();
-				if (ShouldCacheMeshShader(ShaderType, Platform, Material, VertexFactoryType))
-				{
-					++NumShouldCache;
-				}
-			}
-
-			if (NumShouldCache > 0)
-			{
-				NumShadersPerVF += Stages.Num();
-				if (NumShouldCache != Stages.Num())
-				{
-					//#todo-rco: Error or Warning?
-					UE_LOG(LogShaders, Error, TEXT("Mismatched ShouldCache() function on the shaders for MeshMaterial ShaderPipeline %s"), Pipeline->GetName());
-				}
-
-				TArray<FMeshMaterialShaderType*> ShaderStagesToCompile;
-				ShaderStagesToCompile.Reserve(SF_NumFrequencies);
-				for (int32 Index = 0; Index < Stages.Num(); ++Index)
-				{
-					FMeshMaterialShaderType* ShaderType = (FMeshMaterialShaderType*)(Stages[Index]->GetMeshMaterialShaderType());
-
-					// No need to ShouldCache again as we tested each stage beforehand
-
-					// Verify that the shader map Id contains inputs for any shaders that will be put into this shader map
-					check(InShaderMapId.ContainsVertexFactoryType(VertexFactoryType));
-					check(InShaderMapId.ContainsShaderType(ShaderType));
-					ShaderStagesToCompile.Add(ShaderType);
-				}
-
-				//#todo-rco: Check if the MeshShaderMap already contains this pipeline
-
-				// Make a pipeline job with all the stages
-				FMeshMaterialShaderType::BeginCompileShaderPipeline(ShaderMapId, Platform, Material, MaterialEnvironment, VertexFactoryType, Pipeline, ShaderStagesToCompile, NewJobs);
-			}
-		}
-	}
+	TSet<FString> ShaderTypeNames;
 
 	// Iterate over all mesh material shader types.
-	for(TLinkedList<FShaderType*>::TIterator ShaderTypeIt(FShaderType::GetTypeList());ShaderTypeIt;ShaderTypeIt.Next())
+	TMap<FShaderType*, FShaderCompileJob*> SharedShaderJobs;
+	for (TLinkedList<FShaderType*>::TIterator ShaderTypeIt(FShaderType::GetTypeList());ShaderTypeIt;ShaderTypeIt.Next())
 	{
 		FMeshMaterialShaderType* ShaderType = ShaderTypeIt->GetMeshMaterialShaderType();
 		if (ShaderType && ShouldCacheMeshShader(ShaderType, Platform, Material, VertexFactoryType))
@@ -210,8 +164,8 @@ uint32 FMeshMaterialShaderMap::BeginCompile(
 			// only compile the shader if we don't already have it
 			if (!HasShader(ShaderType))
 			{
-			    // Compile this mesh material shader for this material and vertex factory type.
-				ShaderType->BeginCompileShader(
+				// Compile this mesh material shader for this material and vertex factory type.
+				auto* Job = ShaderType->BeginCompileShader(
 					ShaderMapId,
 					Platform,
 					Material,
@@ -220,6 +174,68 @@ uint32 FMeshMaterialShaderMap::BeginCompile(
 					nullptr,
 					NewJobs
 					);
+				check(!SharedShaderJobs.Find(ShaderType));
+				SharedShaderJobs.Add(ShaderType, Job);
+			}
+		}
+	}
+
+	// Now the pipeline jobs; if it's a shareable pipeline, do not add duplicate jobs
+	for (TLinkedList<FShaderPipelineType*>::TIterator ShaderPipelineIt(FShaderPipelineType::GetTypeList());ShaderPipelineIt;ShaderPipelineIt.Next())
+	{
+		const FShaderPipelineType* Pipeline = *ShaderPipelineIt;
+		if (Pipeline->IsMeshMaterialTypePipeline())
+		{
+			auto& StageTypes = Pipeline->GetStages();
+			int32 NumShaderStagesToCompile = 0;
+			for (auto* Shader : StageTypes)
+			{
+				const FMeshMaterialShaderType* ShaderType = Shader->GetMeshMaterialShaderType();
+				if (ShouldCacheMeshShader(ShaderType, Platform, Material, VertexFactoryType))
+				{
+					++NumShaderStagesToCompile;
+				}
+				else
+				{
+					break;
+				}
+			}
+
+			if (NumShaderStagesToCompile == StageTypes.Num())
+			{
+				// Verify that the shader map Id contains inputs for any shaders that will be put into this shader map
+				check(InShaderMapId.ContainsShaderPipelineType(Pipeline));
+
+				if (Pipeline->ShouldOptimizeUnusedOutputs())
+				{
+					NumShadersPerVF += NumShaderStagesToCompile;
+					TArray<FMeshMaterialShaderType*> ShaderStagesToCompile;
+					for (auto* Shader : StageTypes)
+					{
+						const FMeshMaterialShaderType* ShaderType = Shader->GetMeshMaterialShaderType();
+
+						// Verify that the shader map Id contains inputs for any shaders that will be put into this shader map
+						check(InShaderMapId.ContainsVertexFactoryType(VertexFactoryType));
+						check(InShaderMapId.ContainsShaderType(ShaderType));
+						ShaderStagesToCompile.Add((FMeshMaterialShaderType*)ShaderType);
+					}
+
+					// Make a pipeline job with all the stages
+					FMeshMaterialShaderType::BeginCompileShaderPipeline(ShaderMapId, Platform, Material, MaterialEnvironment, VertexFactoryType, Pipeline, ShaderStagesToCompile, NewJobs);
+				}
+				else
+				{
+					// If sharing shaders amongst pipelines, add this pipeline as a dependency of an existing job
+					for (const FShaderType* ShaderType : StageTypes)
+					{
+						FShaderCompileJob** Job = SharedShaderJobs.Find(ShaderType);
+						checkf(Job, TEXT("Couldn't find existing shared job for mesh shader %s on pipeline %s!"), ShaderType->GetName(), Pipeline->GetName());
+						auto* SingleJob = (*Job)->GetSingleShaderJob();
+						auto& PipelinesToShare = SingleJob->SharingPipelines.FindOrAdd(VertexFactoryType);
+						check(!PipelinesToShare.Contains(Pipeline));
+						PipelinesToShare.Add(Pipeline);
+					}
+				}
 			}
 		}
 	}
@@ -239,7 +255,7 @@ inline bool FMeshMaterialShaderMap::IsMeshShaderComplete(const FMeshMaterialShad
 	//		OR If it doesn't have a pipeline and needs one
 	//		OR it's not in the shadermap
 	if (ShouldCacheMeshShader(ShaderType, Platform, Material, InVertexFactoryType) &&
-		(!MeshShaderMap || (Pipeline && !MeshShaderMap->GetShaderPipeline(Pipeline)) || (!Pipeline && !MeshShaderMap->HasShader((FShaderType*)ShaderType))))
+		(!MeshShaderMap || (Pipeline && !MeshShaderMap->HasShaderPipeline(Pipeline)) || (!Pipeline && !MeshShaderMap->HasShader((FShaderType*)ShaderType))))
 	{
 		if (!bSilent)
 		{
@@ -267,6 +283,16 @@ bool FMeshMaterialShaderMap::IsComplete(
 	bool bSilent
 	)
 {
+	// Iterate over all mesh material shader types.
+	for (TLinkedList<FShaderType*>::TIterator ShaderTypeIt(FShaderType::GetTypeList());ShaderTypeIt;ShaderTypeIt.Next())
+	{
+		FMeshMaterialShaderType* ShaderType = ShaderTypeIt->GetMeshMaterialShaderType();
+		if (ShaderType && !IsMeshShaderComplete(MeshShaderMap, Platform, Material, ShaderType, nullptr, InVertexFactoryType, bSilent))
+		{
+			return false;
+		}
+	}
+
 	// Iterate over all pipeline types
 	for (TLinkedList<FShaderPipelineType*>::TIterator ShaderPipelineIt(FShaderPipelineType::GetTypeList());ShaderPipelineIt;ShaderPipelineIt.Next())
 	{
@@ -284,36 +310,27 @@ bool FMeshMaterialShaderMap::IsComplete(
 				{
 					++NumShouldCache;
 				}
+				else
+				{
+					break;
+				}
 			}
 
-			if (NumShouldCache > 0 && NumShouldCache != Stages.Num())
-			{
-				//#todo-rco: Error or Warning?
-				UE_LOG(LogShaders, Error, TEXT("Mismatched ShouldCache() function on the shaders for MeshMaterial ShaderPipeline %s"), ShaderPipelineType->GetName());
-			}
-
-			// Now check the completeness of the shader map
-			for (int32 Index = 0; Index < Stages.Num(); ++Index)
+			if (NumShouldCache == Stages.Num())
 			{
-				auto* ShaderType = Stages[Index]->GetMeshMaterialShaderType();
-				if (ShaderType && !IsMeshShaderComplete(MeshShaderMap, Platform, Material, ShaderType, ShaderPipelineType, InVertexFactoryType, bSilent))
+				// Now check the completeness of the shader map
+				for (int32 Index = 0; Index < Stages.Num(); ++Index)
 				{
-					return false;
+					auto* ShaderType = Stages[Index]->GetMeshMaterialShaderType();
+					if (ShaderType && !IsMeshShaderComplete(MeshShaderMap, Platform, Material, ShaderType, ShaderPipelineType, InVertexFactoryType, bSilent))
+					{
+						return false;
+					}
 				}
 			}
 		}
 	}
 
-	// Iterate over all mesh material shader types.
-	for (TLinkedList<FShaderType*>::TIterator ShaderTypeIt(FShaderType::GetTypeList());ShaderTypeIt;ShaderTypeIt.Next())
-	{
-		FMeshMaterialShaderType* ShaderType = ShaderTypeIt->GetMeshMaterialShaderType();
-		if (ShaderType && !IsMeshShaderComplete(MeshShaderMap, Platform, Material, ShaderType, nullptr, InVertexFactoryType, bSilent))
-		{
-			return false;
-		}
-	}
-
 	return true;
 }
 
@@ -322,21 +339,65 @@ void FMeshMaterialShaderMap::LoadMissingShadersFromMemory(
 	const FMaterial* Material, 
 	EShaderPlatform Platform)
 {
-	//#todo-rco: Pipelines
-	for(TLinkedList<FShaderType*>::TIterator ShaderTypeIt(FShaderType::GetTypeList());ShaderTypeIt;ShaderTypeIt.Next())
+	for (TLinkedList<FShaderType*>::TIterator ShaderTypeIt(FShaderType::GetTypeList());ShaderTypeIt;ShaderTypeIt.Next())
 	{
 		FMeshMaterialShaderType* ShaderType = ShaderTypeIt->GetMeshMaterialShaderType();
-		//#todo-rco: Pass FNullShaderPipeline instead of nullptr
-		const bool bShaderAlreadyExists = HasShader(ShaderType);
-
-		if (ShaderType && ShouldCacheMeshShader(ShaderType, Platform, Material, VertexFactoryType) && !bShaderAlreadyExists) 
+		if (ShaderType && ShouldCacheMeshShader(ShaderType, Platform, Material, VertexFactoryType) && !HasShader((FShaderType*)ShaderType))
 		{
-			const FShaderId ShaderId(MaterialShaderMapHash, nullptr, VertexFactoryType, ShaderType, FShaderTarget(ShaderType->GetFrequency(), Platform));
-			FShader* FoundShader = ShaderType->FindShaderById(ShaderId);	
+			const FShaderId ShaderId(MaterialShaderMapHash, nullptr, VertexFactoryType, (FShaderType*)ShaderType, FShaderTarget(ShaderType->GetFrequency(), Platform));
+			FShader* FoundShader = ((FShaderType*)ShaderType)->FindShaderById(ShaderId);
 
 			if (FoundShader)
 			{
-				AddShader(ShaderType, FoundShader);
+				AddShader((FShaderType*)ShaderType, FoundShader);
+			}
+		}
+	}
+
+	// Try to find necessary FShaderPipelineTypes in memory
+	for (TLinkedList<FShaderPipelineType*>::TIterator ShaderPipelineIt(FShaderPipelineType::GetTypeList());ShaderPipelineIt;ShaderPipelineIt.Next())
+	{
+		const FShaderPipelineType* PipelineType = *ShaderPipelineIt;
+		if (PipelineType && PipelineType->IsMeshMaterialTypePipeline() && !HasShaderPipeline(PipelineType))
+		{
+			auto& Stages = PipelineType->GetStages();
+			int32 NumShaders = 0;
+			for (const FShaderType* Shader : Stages)
+			{
+				FMeshMaterialShaderType* ShaderType = (FMeshMaterialShaderType*)Shader->GetMeshMaterialShaderType();
+				if (ShaderType && ShouldCacheMeshShader(ShaderType, Platform, Material, VertexFactoryType))
+				{
+					++NumShaders;
+				}
+				else
+				{
+					break;
+				}
+			}
+
+			if (NumShaders == Stages.Num())
+			{
+				TArray<FShader*> ShadersForPipeline;
+				for (auto* Shader : Stages)
+				{
+					FMeshMaterialShaderType* ShaderType = (FMeshMaterialShaderType*)Shader->GetMeshMaterialShaderType();
+					if (!HasShader(ShaderType))
+					{
+						const FShaderId ShaderId(MaterialShaderMapHash, PipelineType->ShouldOptimizeUnusedOutputs() ? PipelineType : nullptr, VertexFactoryType, ShaderType, FShaderTarget(ShaderType->GetFrequency(), Platform));
+						FShader* FoundShader = ShaderType->FindShaderById(ShaderId);
+						if (FoundShader)
+						{
+							AddShader(ShaderType, FoundShader);
+							ShadersForPipeline.Add(FoundShader);
+						}
+					}
+				}
+
+				if (ShadersForPipeline.Num() == NumShaders && !HasShaderPipeline(PipelineType))
+				{
+					auto* Pipeline = new FShaderPipeline(PipelineType, ShadersForPipeline);
+					AddShaderPipeline(PipelineType, Pipeline);
+				}
 			}
 		}
 	}
diff --git a/Engine/Source/Runtime/Engine/Private/ShaderCompiler/ShaderCompiler.cpp b/Engine/Source/Runtime/Engine/Private/ShaderCompiler/ShaderCompiler.cpp
--- a/Engine/Source/Runtime/Engine/Private/ShaderCompiler/ShaderCompiler.cpp
+++ b/Engine/Source/Runtime/Engine/Private/ShaderCompiler/ShaderCompiler.cpp
@@ -977,30 +977,30 @@ void FShaderCompileThreadRunnable::CompileDirectlyThroughDll()
 				if (SingleJob)
 				{
 					const FName Format = LegacyShaderPlatformToShaderFormat(EShaderPlatform(SingleJob->Input.Target.Platform));
-				const IShaderFormat* Compiler = TPM.FindShaderFormat(Format);
+					const IShaderFormat* Compiler = TPM.FindShaderFormat(Format);
 
-				if (!Compiler)
-				{
-					UE_LOG(LogShaderCompilers, Fatal, TEXT("Can't compile shaders for format %s, couldn't load compiler dll"), *Format.ToString());
-				}
-				CA_ASSUME(Compiler != NULL);
+					if (!Compiler)
+					{
+						UE_LOG(LogShaderCompilers, Fatal, TEXT("Can't compile shaders for format %s, couldn't load compiler dll"), *Format.ToString());
+					}
+					CA_ASSUME(Compiler != NULL);
 
 					if (IsValidRef(SingleJob->Input.SharedEnvironment))
-				{
-					// Merge the shared environment into the per-shader environment before calling into the compile function
-					// Normally this happens in the worker
+					{
+						// Merge the shared environment into the per-shader environment before calling into the compile function
+						// Normally this happens in the worker
 						SingleJob->Input.Environment.Merge(*SingleJob->Input.SharedEnvironment);
-				}
+					}
 
-				// Compile the shader directly through the platform dll (directly from the shader dir as the working directory)
+					// Compile the shader directly through the platform dll (directly from the shader dir as the working directory)
 					Compiler->CompileShader(Format, SingleJob->Input, SingleJob->Output, FString(FPlatformProcess::ShaderDir()));
 
 					SingleJob->bSucceeded = SingleJob->Output.bSucceeded;
 
 					if (SingleJob->Output.bSucceeded)
-				{
-					// Generate a hash of the output and cache it
-					// The shader processing this output will use it to search for existing FShaderResources
+					{
+						// Generate a hash of the output and cache it
+						// The shader processing this output will use it to search for existing FShaderResources
 						SingleJob->Output.GenerateOutputHash();
 					}
 				}
@@ -1284,10 +1284,10 @@ FProcHandle FShaderCompilingManager::LaunchWorker(const FString& WorkingDirector
 
 	if (DEBUG_SHADERCOMPILEWORKER)
 	{
-	// Note: Set breakpoint here and launch the shadercompileworker with WorkerParameters a cmd-line
-	const TCHAR* WorkerParametersText = *WorkerParameters;
-	FPlatformMisc::LowLevelOutputDebugStringf(TEXT("Launching shader compile worker w/ WorkerParameters\n\t%s\n"), WorkerParametersText);
-	return FProcHandle();
+		// Note: Set breakpoint here and launch the ShaderCompileWorker with WorkerParameters a cmd-line
+		const TCHAR* WorkerParametersText = *WorkerParameters;
+		FPlatformMisc::LowLevelOutputDebugStringf(TEXT("Launching shader compile worker w/ WorkerParameters\n\t%s\n"), WorkerParametersText);
+		return FProcHandle();
 	}
 	else
 	{
@@ -1499,7 +1499,7 @@ void FShaderCompilingManager::ProcessCompiledShaderMaps(
 
 			if (bSuccess)
 			{
-				bShaderMapComplete = ShaderMap->ProcessCompilationResults(ResultArray, CompileResults.FinalizeJobIndex, TimeBudget);
+				bShaderMapComplete = ShaderMap->ProcessCompilationResults(ResultArray, CompileResults.FinalizeJobIndex, TimeBudget, CompileResults.SharedPipelines);
 			}
 
 
diff --git a/Engine/Source/Runtime/Engine/Private/ShaderDerivedDataVersion.h b/Engine/Source/Runtime/Engine/Private/ShaderDerivedDataVersion.h
--- a/Engine/Source/Runtime/Engine/Private/ShaderDerivedDataVersion.h
+++ b/Engine/Source/Runtime/Engine/Private/ShaderDerivedDataVersion.h
@@ -9,5 +9,5 @@
 // In case of merge conflicts with DDC versions, you MUST generate a new GUID and set this new
 // guid as version
 
-#define GLOBALSHADERMAP_DERIVEDDATA_VER			TEXT("753D597996674C67855E3A31377FCDc6")
-#define MATERIALSHADERMAP_DERIVEDDATA_VER		TEXT("4F0644619F9C48FF91E0D494C95EB049")
+#define GLOBALSHADERMAP_DERIVEDDATA_VER			TEXT("753D597996674C67855E3A31377FCDd0")
+#define MATERIALSHADERMAP_DERIVEDDATA_VER		TEXT("4F0644619F9C48FF91E0D494C95EB050")
diff --git a/Engine/Source/Runtime/Engine/Private/SkeletalMesh.cpp b/Engine/Source/Runtime/Engine/Private/SkeletalMesh.cpp
--- a/Engine/Source/Runtime/Engine/Private/SkeletalMesh.cpp
+++ b/Engine/Source/Runtime/Engine/Private/SkeletalMesh.cpp
@@ -1830,7 +1830,10 @@ void FreeSkeletalMeshBuffersSinkCallback()
 		FlushRenderingCommands();
 		for (TObjectIterator<USkeletalMesh> It;It;++It)
 		{
-			It->ReleaseCPUResources();
+			if (!It->GetImportedResource()->RequiresCPUSkinning(GMaxRHIFeatureLevel))
+			{
+				It->ReleaseCPUResources();
+			}
 		}
 	}
 }
diff --git a/Engine/Source/Runtime/Engine/Private/SkeletalRenderGPUSkin.cpp b/Engine/Source/Runtime/Engine/Private/SkeletalRenderGPUSkin.cpp
--- a/Engine/Source/Runtime/Engine/Private/SkeletalRenderGPUSkin.cpp
+++ b/Engine/Source/Runtime/Engine/Private/SkeletalRenderGPUSkin.cpp
@@ -230,8 +230,7 @@ void FSkeletalMeshObjectGPUSkin::Update(int32 LODIndex,USkinnedMeshComponent* In
 static TAutoConsoleVariable<int32> CVarDeferSkeletalDynamicDataUpdateUntilGDME(
 	TEXT("r.DeferSkeletalDynamicDataUpdateUntilGDME"),
 	0,
-	TEXT("If > 0, then do skeletal mesh dynamic data updates will be deferred until GDME. Experimental option."),
-	ECVF_Cheat);
+	TEXT("If > 0, then do skeletal mesh dynamic data updates will be deferred until GDME. Experimental option."));
 
 void FSkeletalMeshObjectGPUSkin::UpdateDynamicData_RenderThread(FRHICommandListImmediate& RHICmdList, FDynamicSkelMeshObjectDataGPUSkin* InDynamicData, uint32 FrameNumber)
 {
diff --git a/Engine/Source/Runtime/Engine/Public/GlobalShader.h b/Engine/Source/Runtime/Engine/Public/GlobalShader.h
--- a/Engine/Source/Runtime/Engine/Public/GlobalShader.h
+++ b/Engine/Source/Runtime/Engine/Public/GlobalShader.h
@@ -29,6 +29,9 @@ public:
 private:
 	/** Shader types that this shader map is dependent on and their stored state. */
 	TArray<FShaderTypeDependency> ShaderTypeDependencies;
+
+	/** Shader pipeline types that this shader map is dependent on and their stored state. */
+	TArray<FShaderPipelineTypeDependency> ShaderPipelineTypeDependencies;
 };
 
 /**
@@ -64,7 +67,7 @@ public:
 	/**
 	 * Enqueues compilation of a shader of this type.  
 	 */
-	ENGINE_API void BeginCompileShader(EShaderPlatform Platform, const FShaderPipelineType* ShaderPipeline, TArray<FShaderCommonCompileJob*>& NewJobs);
+	ENGINE_API class FShaderCompileJob* BeginCompileShader(EShaderPlatform Platform, const FShaderPipelineType* ShaderPipeline, TArray<FShaderCommonCompileJob*>& NewJobs);
 
 	/**
 	 * Enqueues compilation of a shader pipeline of this type.  
diff --git a/Engine/Source/Runtime/Engine/Public/LocalVertexFactory.h b/Engine/Source/Runtime/Engine/Public/LocalVertexFactory.h
--- a/Engine/Source/Runtime/Engine/Public/LocalVertexFactory.h
+++ b/Engine/Source/Runtime/Engine/Public/LocalVertexFactory.h
@@ -89,6 +89,16 @@ public:
 	virtual void Bind(const FShaderParameterMap& ParameterMap) override;
 	virtual void Serialize(FArchive& Ar) override;
 	virtual void SetMesh(FRHICommandList& RHICmdList, FShader* Shader, const FVertexFactory* VertexFactory, const FSceneView& View, const FMeshBatchElement& BatchElement, uint32 DataFlags) const override;
+
+	FLocalVertexFactoryShaderParameters()
+		: bAnySpeedTreeParamIsBound(false)
+	{
+	}
+
+	// True if either of the two below is bound, which puts us on the slow path in SetMesh
+	bool bAnySpeedTreeParamIsBound;
 	// SpeedTree LOD parameter
 	FShaderParameter LODParameter;
+	// SpeedTree wind parameter
+	FShaderResourceParameter WindParameter;
 };
diff --git a/Engine/Source/Runtime/Engine/Public/MaterialShaderType.h b/Engine/Source/Runtime/Engine/Public/MaterialShaderType.h
--- a/Engine/Source/Runtime/Engine/Public/MaterialShaderType.h
+++ b/Engine/Source/Runtime/Engine/Public/MaterialShaderType.h
@@ -86,7 +86,7 @@ public:
 	 * Enqueues a compilation for a new shader of this type.
 	 * @param Material - The material to link the shader with.
 	 */
-	void BeginCompileShader(
+	class FShaderCompileJob* BeginCompileShader(
 		uint32 ShaderMapId,
 		const FMaterial* Material,
 		FShaderCompilerEnvironment* MaterialEnvironment,
diff --git a/Engine/Source/Runtime/Engine/Public/MaterialShared.h b/Engine/Source/Runtime/Engine/Public/MaterialShared.h
--- a/Engine/Source/Runtime/Engine/Public/MaterialShared.h
+++ b/Engine/Source/Runtime/Engine/Public/MaterialShared.h
@@ -24,6 +24,7 @@ class FShaderCommonCompileJob;
 class FShaderCompileJob;
 class FShaderType;
 class FShaderTypeDependency;
+class FShaderPipelineTypeDependency;
 class FVertexFactoryType;
 class FVertexFactoryTypeDependency;
 class UMaterial;
@@ -403,6 +404,9 @@ public:
 	/** Shader types of shaders that are inlined in this shader map in the DDC. */
 	TArray<FShaderTypeDependency> ShaderTypeDependencies;
 
+	/** Shader pipeline types of shader pipelines that are inlined in this shader map in the DDC. */
+	TArray<FShaderPipelineTypeDependency> ShaderPipelineTypeDependencies;
+
 	/** Vertex factory types of shaders that are inlined in this shader map in the DDC. */
 	TArray<FVertexFactoryTypeDependency> VertexFactoryTypeDependencies;
 
@@ -442,6 +446,7 @@ public:
 			+ ReferencedFunctions.GetAllocatedSize()
 			+ ReferencedParameterCollections.GetAllocatedSize()
 			+ ShaderTypeDependencies.GetAllocatedSize()
+			+ ShaderPipelineTypeDependencies.GetAllocatedSize()
 			+ VertexFactoryTypeDependencies.GetAllocatedSize();
 	}
 
@@ -467,6 +472,9 @@ public:
 	/** Returns true if the requested shader type is a dependency of this shader map Id. */
 	bool ContainsShaderType(const FShaderType* ShaderType) const;
 
+	/** Returns true if the requested shader type is a dependency of this shader map Id. */
+	bool ContainsShaderPipelineType(const FShaderPipelineType* ShaderPipelineType) const;
+
 	/** Returns true if the requested vertex factory type is a dependency of this shader map Id. */
 	bool ContainsVertexFactoryType(const FVertexFactoryType* VFType) const;
 };
@@ -530,7 +538,7 @@ public:
 	void FlushShadersByShaderPipelineType(const FShaderPipelineType* ShaderPipelineType);
 
 	// Accessors.
-	FVertexFactoryType* GetVertexFactoryType() const { return VertexFactoryType; }
+	inline FVertexFactoryType* GetVertexFactoryType() const { return VertexFactoryType; }
 
 private:
 	/** The vertex factory type these shaders are for. */
@@ -590,7 +598,7 @@ public:
 		);
 
 	/** Sorts the incoming compiled jobs into the appropriate mesh shader maps, and finalizes this shader map so that it can be used for rendering. */
-	bool ProcessCompilationResults(const TArray<FShaderCommonCompileJob*>& InCompilationResults, int32& ResultIndex, float& TimeBudget);
+	bool ProcessCompilationResults(const TArray<FShaderCommonCompileJob*>& InCompilationResults, int32& ResultIndex, float& TimeBudget, TMap<const FVertexFactoryType*, TArray<const FShaderPipelineType*> >& SharedPipelines);
 
 	/**
 	 * Checks whether the material shader map is missing any shader types necessary for the given material.
@@ -1119,6 +1127,8 @@ public:
 		return (ShaderType*)GetShader(&ShaderType::StaticType, VertexFactoryType);
 	}
 
+	ENGINE_API FShaderPipeline* GetShaderPipeline(class FShaderPipelineType* ShaderPipelineType, FVertexFactoryType* VertexFactoryType) const;
+
 	/** Returns a string that describes the material's usage for debugging purposes. */
 	virtual FString GetMaterialUsageDescription() const = 0;
 
@@ -1587,7 +1597,7 @@ public:
 	 */
 	ENGINE_API void GetRepresentativeInstructionCounts(TArray<FString> &Descriptions, TArray<int32> &InstructionCounts) const;
 
-	ENGINE_API void GetRepresentativeShaderTypesAndDescriptions(TArray<FString> &ShaderTypeNames, TArray<FString> &ShaderTypeDescriptions) const;
+	ENGINE_API void GetRepresentativeShaderTypesAndDescriptions(TMap<FName, FString>& OutShaderTypeNameAndDescriptions) const;
 
 	ENGINE_API SIZE_T GetResourceSizeInclusive();
 
diff --git a/Engine/Source/Runtime/Engine/Public/MeshMaterialShaderType.h b/Engine/Source/Runtime/Engine/Public/MeshMaterialShaderType.h
--- a/Engine/Source/Runtime/Engine/Public/MeshMaterialShaderType.h
+++ b/Engine/Source/Runtime/Engine/Public/MeshMaterialShaderType.h
@@ -59,7 +59,7 @@ public:
 	 * @param Material - The material to link the shader with.
 	 * @param VertexFactoryType - The vertex factory to compile with.
 	 */
-	void BeginCompileShader(
+	class FShaderCompileJob* BeginCompileShader(
 		uint32 ShaderMapId,
 		EShaderPlatform Platform,
 		const FMaterial* Material,
diff --git a/Engine/Source/Runtime/Engine/Public/PrimitiveSceneProxy.h b/Engine/Source/Runtime/Engine/Public/PrimitiveSceneProxy.h
--- a/Engine/Source/Runtime/Engine/Public/PrimitiveSceneProxy.h
+++ b/Engine/Source/Runtime/Engine/Public/PrimitiveSceneProxy.h
@@ -423,8 +423,6 @@ public:
 	inline bool AllowApproximateOcclusion() const { return bAllowApproximateOcclusion; }
 	inline const TUniformBuffer<FPrimitiveUniformShaderParameters>& GetUniformBuffer() const 
 	{
-
-		check(!NeedsUniformBufferUpdate());
 		return UniformBuffer; 
 	}
 	inline bool HasPerInstanceHitProxies () const { return bHasPerInstanceHitProxies; }
diff --git a/Engine/Source/Runtime/Engine/Public/ShaderCompiler.h b/Engine/Source/Runtime/Engine/Public/ShaderCompiler.h
--- a/Engine/Source/Runtime/Engine/Public/ShaderCompiler.h
+++ b/Engine/Source/Runtime/Engine/Public/ShaderCompiler.h
@@ -52,6 +52,9 @@ public:
 	FShaderCompilerInput Input;
 	FShaderCompilerOutput Output;
 
+	// List of pipelines that are sharing this job.
+	TMap<const FVertexFactoryType*, TArray<const FShaderPipelineType*>> SharingPipelines;
+
 	FShaderCompileJob(uint32 InId, FVertexFactoryType* InVFType, FShaderType* InShaderType) :
 		FShaderCommonCompileJob(InId),
 		VFType(InVFType),
@@ -290,6 +293,9 @@ struct FShaderMapFinalizeResults : public FShaderMapCompileResults
 	/** Tracks finalization progress on this shader map. */
 	int32 FinalizeJobIndex;
 
+	// List of pipelines with shared shaders; nullptr for non mesh pipelines
+	TMap<const FVertexFactoryType*, TArray<const FShaderPipelineType*> > SharedPipelines;
+
 	FShaderMapFinalizeResults(const FShaderMapCompileResults& InCompileResults) :
 		FShaderMapCompileResults(InCompileResults),
 		FinalizeJobIndex(0)
diff --git a/Engine/Source/Runtime/RHI/Private/RHICommandList.cpp b/Engine/Source/Runtime/RHI/Private/RHICommandList.cpp
--- a/Engine/Source/Runtime/RHI/Private/RHICommandList.cpp
+++ b/Engine/Source/Runtime/RHI/Private/RHICommandList.cpp
@@ -1891,6 +1891,26 @@ void FDynamicRHI::UpdateTexture2D_RenderThread(class FRHICommandListImmediate& R
 	return GDynamicRHI->RHIUpdateTexture2D(Texture, MipIndex, UpdateRegion, SourcePitch, SourceData);
 }
 
+void* FDynamicRHI::LockTexture2D_RenderThread(class FRHICommandListImmediate& RHICmdList, FTexture2DRHIParamRef Texture, uint32 MipIndex, EResourceLockMode LockMode, uint32& DestStride, bool bLockWithinMiptail, bool bNeedsDefaultRHIFlush)
+{
+	if (bNeedsDefaultRHIFlush)
+	{
+		QUICK_SCOPE_CYCLE_COUNTER(STAT_RHIMETHOD_LockTexture2D_Flush);
+		RHICmdList.ImmediateFlush(EImmediateFlushType::FlushRHIThread);
+	}
+	return GDynamicRHI->RHILockTexture2D(Texture, MipIndex, LockMode, DestStride, bLockWithinMiptail);
+}
+
+void FDynamicRHI::UnlockTexture2D_RenderThread(class FRHICommandListImmediate& RHICmdList, FTexture2DRHIParamRef Texture, uint32 MipIndex, bool bLockWithinMiptail, bool bNeedsDefaultRHIFlush)
+{
+	if (bNeedsDefaultRHIFlush)
+	{
+		QUICK_SCOPE_CYCLE_COUNTER(STAT_RHIMETHOD_UnlockTexture2D_Flush);
+		RHICmdList.ImmediateFlush(EImmediateFlushType::FlushRHIThread);
+	}
+	GDynamicRHI->RHIUnlockTexture2D(Texture, MipIndex, bLockWithinMiptail);
+}
+
 FTexture2DRHIRef FDynamicRHI::RHICreateTexture2D_RenderThread(class FRHICommandListImmediate& RHICmdList, uint32 SizeX, uint32 SizeY, uint8 Format, uint32 NumMips, uint32 NumSamples, uint32 Flags, FRHIResourceCreateInfo& CreateInfo)
 {
 	FScopedRHIThreadStaller StallRHIThread(RHICmdList);
diff --git a/Engine/Source/Runtime/RHI/Public/DynamicRHI.h b/Engine/Source/Runtime/RHI/Public/DynamicRHI.h
--- a/Engine/Source/Runtime/RHI/Public/DynamicRHI.h
+++ b/Engine/Source/Runtime/RHI/Public/DynamicRHI.h
@@ -1108,6 +1108,8 @@ public:
 	virtual void UnlockIndexBuffer_RenderThread(class FRHICommandListImmediate& RHICmdList, FIndexBufferRHIParamRef IndexBuffer);
 	virtual FVertexDeclarationRHIRef CreateVertexDeclaration_RenderThread(class FRHICommandListImmediate& RHICmdList, const FVertexDeclarationElementList& Elements);
 	virtual void UpdateTexture2D_RenderThread(class FRHICommandListImmediate& RHICmdList, FTexture2DRHIParamRef Texture, uint32 MipIndex, const struct FUpdateTextureRegion2D& UpdateRegion, uint32 SourcePitch, const uint8* SourceData);
+	virtual void* LockTexture2D_RenderThread(class FRHICommandListImmediate& RHICmdList, FTexture2DRHIParamRef Texture, uint32 MipIndex, EResourceLockMode LockMode, uint32& DestStride, bool bLockWithinMiptail, bool bNeedsDefaultRHIFlush = true);
+	virtual void UnlockTexture2D_RenderThread(class FRHICommandListImmediate& RHICmdList, FTexture2DRHIParamRef Texture, uint32 MipIndex, bool bLockWithinMiptail, bool bNeedsDefaultRHIFlush = true);
 
 	virtual FTexture2DRHIRef RHICreateTexture2D_RenderThread(class FRHICommandListImmediate& RHICmdList, uint32 SizeX, uint32 SizeY, uint8 Format, uint32 NumMips, uint32 NumSamples, uint32 Flags, FRHIResourceCreateInfo& CreateInfo);
 	virtual FTexture2DArrayRHIRef RHICreateTexture2DArray_RenderThread(class FRHICommandListImmediate& RHICmdList, uint32 SizeX, uint32 SizeY, uint32 SizeZ, uint8 Format, uint32 NumMips, uint32 Flags, FRHIResourceCreateInfo& CreateInfo);
diff --git a/Engine/Source/Runtime/RHI/Public/RHICommandList.h b/Engine/Source/Runtime/RHI/Public/RHICommandList.h
--- a/Engine/Source/Runtime/RHI/Public/RHICommandList.h
+++ b/Engine/Source/Runtime/RHI/Public/RHICommandList.h
@@ -2604,22 +2604,12 @@ public:
 	
 	FORCEINLINE void* LockTexture2D(FTexture2DRHIParamRef Texture, uint32 MipIndex, EResourceLockMode LockMode, uint32& DestStride, bool bLockWithinMiptail, bool bFlushRHIThread = true)
 	{
-		if (bFlushRHIThread)
-		{
-			QUICK_SCOPE_CYCLE_COUNTER(STAT_RHIMETHOD_LockTexture2D_Flush);
-			ImmediateFlush(EImmediateFlushType::FlushRHIThread);  
-		}
-		return GDynamicRHI->RHILockTexture2D(Texture, MipIndex, LockMode, DestStride, bLockWithinMiptail);
+		return GDynamicRHI->LockTexture2D_RenderThread(*this, Texture, MipIndex, LockMode, DestStride, bLockWithinMiptail, bFlushRHIThread);
 	}
 	
 	FORCEINLINE void UnlockTexture2D(FTexture2DRHIParamRef Texture, uint32 MipIndex, bool bLockWithinMiptail, bool bFlushRHIThread = true)
-	{
-		if (bFlushRHIThread)
-		{
-			QUICK_SCOPE_CYCLE_COUNTER(STAT_RHIMETHOD_UnlockTexture2D_Flush);
-			ImmediateFlush(EImmediateFlushType::FlushRHIThread);   
-		}
-		GDynamicRHI->RHIUnlockTexture2D(Texture, MipIndex, bLockWithinMiptail);
+	{		
+		GDynamicRHI->UnlockTexture2D_RenderThread(*this, Texture, MipIndex, bLockWithinMiptail, bFlushRHIThread);
 	}
 	
 	FORCEINLINE void* LockTexture2DArray(FTexture2DArrayRHIParamRef Texture, uint32 TextureIndex, uint32 MipIndex, EResourceLockMode LockMode, uint32& DestStride, bool bLockWithinMiptail)
diff --git a/Engine/Source/Runtime/Renderer/Private/CompositionLighting/CompositionLighting.cpp b/Engine/Source/Runtime/Renderer/Private/CompositionLighting/CompositionLighting.cpp
--- a/Engine/Source/Runtime/Renderer/Private/CompositionLighting/CompositionLighting.cpp
+++ b/Engine/Source/Runtime/Renderer/Private/CompositionLighting/CompositionLighting.cpp
@@ -232,6 +232,10 @@ static FRenderingCompositeOutputRef AddPostProcessingAmbientOcclusion(FRHIComman
 	{
 		AmbientOcclusionInMip1->AddDependency(Context.FinalOutput);
 	}
+	else
+	{
+		AmbientOcclusionPassMip0->AddDependency(Context.FinalOutput);
+	}
 
 	Context.FinalOutput = FRenderingCompositeOutputRef(AmbientOcclusionPassMip0);
 
@@ -403,9 +407,7 @@ void FCompositionLighting::ProcessAfterLighting(FRHICommandListImmediate& RHICmd
 			bool bSimpleDynamicLighting = IsSimpleDynamicLightingEnabled();
 
 			bool bScreenSpaceSubsurfacePassNeeded = (View.ShadingModelMaskInView & (1 << MSM_SubsurfaceProfile)) != 0;
-			if (bScreenSpaceSubsurfacePassNeeded && !bSimpleDynamicLighting &&
-				//@todo-rco: Remove this when we fix the cross-compiler
-				!IsOpenGLPlatform(View.GetShaderPlatform()))
+			if (bScreenSpaceSubsurfacePassNeeded && !bSimpleDynamicLighting)
 			{
 				bool bHalfRes = CVarSSSHalfRes.GetValueOnRenderThread() != 0;
 				bool bSingleViewportMode = View.Family->Views.Num() == 1;
diff --git a/Engine/Source/Runtime/Renderer/Private/DeferredShadingRenderer.cpp b/Engine/Source/Runtime/Renderer/Private/DeferredShadingRenderer.cpp
--- a/Engine/Source/Runtime/Renderer/Private/DeferredShadingRenderer.cpp
+++ b/Engine/Source/Runtime/Renderer/Private/DeferredShadingRenderer.cpp
@@ -1297,8 +1297,7 @@ void FDeferredShadingSceneRenderer::Render(FRHICommandListImmediate& RHICmdList)
 			GPostProcessing.Process(RHICmdList, Views[ ViewIndex ], VelocityRT);
 
 			// we rendered to it during the frame, seems we haven't made use of it, because it should be released
-			FSceneViewState* ViewState = (FSceneViewState*)Views[ ViewIndex ].State;
-			check( !ViewState || !ViewState->SeparateTranslucencyRT );
+			check(!FSceneRenderTargets::Get(RHICmdList).SeparateTranslucencyRT);
 		}
 	}
 	else
diff --git a/Engine/Source/Runtime/Renderer/Private/LightRendering.cpp b/Engine/Source/Runtime/Renderer/Private/LightRendering.cpp
--- a/Engine/Source/Runtime/Renderer/Private/LightRendering.cpp
+++ b/Engine/Source/Runtime/Renderer/Private/LightRendering.cpp
@@ -910,7 +910,7 @@ void FDeferredShadingSceneRenderer::RenderLight(FRHICommandList& RHICmdList, con
 					}
 					else
 					{
-						SetShaderTemplLighting<false, false, false, true>(RHICmdList, View, *VertexShader, LightSceneInfo);
+						SetShaderTemplLighting<false, false, false, false>(RHICmdList, View, *VertexShader, LightSceneInfo);
 					}
 				}
 			}
diff --git a/Engine/Source/Runtime/Renderer/Private/PostProcess/PostProcessMotionBlur.cpp b/Engine/Source/Runtime/Renderer/Private/PostProcess/PostProcessMotionBlur.cpp
--- a/Engine/Source/Runtime/Renderer/Private/PostProcess/PostProcessMotionBlur.cpp
+++ b/Engine/Source/Runtime/Renderer/Private/PostProcess/PostProcessMotionBlur.cpp
@@ -25,12 +25,10 @@ static TAutoConsoleVariable<int32> CVarMotionBlurFiltering(
 	ECVF_Cheat | ECVF_RenderThreadSafe);
 #endif
 
-static TAutoConsoleVariable<int32> CVarMotionBlurSmoothMax(
-	TEXT("r.MotionBlurSmoothMax"),
-	0,
-	TEXT("Useful developer variable\n")
-	TEXT("0: off (default, expected by the shader for better quality)\n")
-	TEXT("1: on"),
+static TAutoConsoleVariable<float> CVarMotionBlur2ndScale(
+	TEXT("r.MotionBlur2ndScale"),
+	1.0f,
+	TEXT(""),
 	ECVF_Cheat | ECVF_RenderThreadSafe);
 
 
@@ -685,6 +683,34 @@ FIntPoint GetNumTiles16x16( FIntPoint PixelExtent )
 	return FIntPoint( TilesX, TilesY );
 }
 
+FVector4 GetMotionBlurParameters( const FViewInfo& View, float Scale = 1.0f )
+{
+	const float TileSize = 16.0f;
+
+	const float SizeX = View.ViewRect.Width();
+	const float SizeY = View.ViewRect.Height();
+	const float AspectRatio = SizeY / SizeX;
+
+	const FSceneViewState* ViewState = (FSceneViewState*) View.State;
+	float MotionBlurTimeScale = ViewState ? ViewState->MotionBlurTimeScale : 1.0f;
+	float MotionBlurScale = 0.5f * MotionBlurTimeScale * View.FinalPostProcessSettings.MotionBlurAmount;
+
+	// 0:no 1:full screen width, percent conversion
+	float MaxVelocity = View.FinalPostProcessSettings.MotionBlurMax / 100.0f;
+
+	// Scale by 0.5 due to blur samples going both ways
+	float PixelScale = Scale * SizeX * 0.5f;
+
+	FVector4 MotionBlurParameters(
+		AspectRatio,
+		PixelScale * MotionBlurScale,			// Scale for pixels
+		PixelScale * MotionBlurScale / TileSize,// Scale for tiles
+		FMath::Abs( PixelScale ) * MaxVelocity	// Max velocity pixels
+	);
+
+	return MotionBlurParameters;
+}
+
 class FPostProcessVelocityFlattenCS : public FGlobalShader
 {
 	DECLARE_SHADER_TYPE(FPostProcessVelocityFlattenCS, Global);
@@ -709,6 +735,7 @@ public:
 		: FGlobalShader(Initializer)
 	{
 		PostprocessParameter.Bind(Initializer.ParameterMap);
+		MotionBlurParameters.Bind(Initializer.ParameterMap, TEXT("MotionBlurParameters"));
 		OutVelocityFlat.Bind(Initializer.ParameterMap, TEXT("OutVelocityFlat"));
 		OutMaxTileVelocity.Bind(Initializer.ParameterMap, TEXT("OutMaxTileVelocity"));
 	}
@@ -721,12 +748,15 @@ public:
 		PostprocessParameter.SetCS(ShaderRHI, Context, TStaticSamplerState<SF_Point,AM_Clamp,AM_Clamp,AM_Clamp>::GetRHI());
 
 		SetUniformBufferParameter(RHICmdList, ShaderRHI, GetUniformBufferParameter<FCameraMotionParameters>(), CreateCameraMotionParametersUniformBuffer(Context.View));
+
+		SetShaderValue(Context.RHICmdList, ShaderRHI, MotionBlurParameters, GetMotionBlurParameters( Context.View ) );
 	}
 	
 	virtual bool Serialize(FArchive& Ar) override
 	{
 		bool bShaderHasOutdatedParameters = FGlobalShader::Serialize(Ar);
 		Ar << PostprocessParameter;
+		Ar << MotionBlurParameters;
 		Ar << OutVelocityFlat;
 		Ar << OutMaxTileVelocity;
 		return bShaderHasOutdatedParameters;
@@ -734,6 +764,7 @@ public:
 
 private:
 	FPostProcessPassParameters	PostprocessParameter;
+	FShaderParameter			MotionBlurParameters;
 };
 
 IMPLEMENT_SHADER_TYPE(,FPostProcessVelocityFlattenCS,TEXT("PostProcessVelocityFlatten"),TEXT("VelocityFlattenMain"),SF_Compute);
@@ -908,7 +939,7 @@ class FPostProcessVelocityScatterVS : public FGlobalShader
 public:
 	FPostProcessPassParameters PostprocessParameter;
 	FShaderParameter DrawMax;
-	FShaderParameter VelocityScale;
+	FShaderParameter MotionBlurParameters;
 
 	/** Initialization constructor. */
 	FPostProcessVelocityScatterVS(const ShaderMetaType::CompiledShaderInitializerType& Initializer)
@@ -916,14 +947,14 @@ public:
 	{
 		PostprocessParameter.Bind(Initializer.ParameterMap);
 		DrawMax.Bind(Initializer.ParameterMap, TEXT("bDrawMax"));
-		VelocityScale.Bind( Initializer.ParameterMap, TEXT("VelocityScale") );
+		MotionBlurParameters.Bind( Initializer.ParameterMap, TEXT("MotionBlurParameters") );
 	}
 
 	// FShader interface.
 	virtual bool Serialize(FArchive& Ar) override
 	{
 		bool bShaderHasOutdatedParameters = FGlobalShader::Serialize(Ar);
-		Ar << PostprocessParameter << DrawMax << VelocityScale;
+		Ar << PostprocessParameter << DrawMax << MotionBlurParameters;
 		return bShaderHasOutdatedParameters;
 	}
 
@@ -936,14 +967,7 @@ public:
 		PostprocessParameter.SetVS(ShaderRHI, Context, TStaticSamplerState<SF_Point,AM_Clamp,AM_Clamp,AM_Clamp>::GetRHI());
 
 		SetShaderValue(Context.RHICmdList, ShaderRHI, DrawMax, bDrawMax);
-
-		{
-			const FSceneViewState* ViewState = (FSceneViewState*) Context.View.State;
-			const float MotionBlurTimeScale = ViewState ? ViewState->MotionBlurTimeScale : 1.0f;
-
-			const float ViewMotionBlurScale = 0.5f * MotionBlurTimeScale * Context.View.FinalPostProcessSettings.MotionBlurAmount;
-			SetShaderValue(Context.RHICmdList, ShaderRHI, VelocityScale, FVector4(ViewMotionBlurScale, ViewMotionBlurScale, 0, 0));
-		}
+		SetShaderValue(Context.RHICmdList, ShaderRHI, MotionBlurParameters, GetMotionBlurParameters( Context.View ) );
 	}
 
 	static const TCHAR* GetSourceFilename()
@@ -1095,6 +1119,7 @@ public:
 		: FGlobalShader(Initializer)
 	{
 		PostprocessParameter.Bind(Initializer.ParameterMap);
+		MotionBlurParameters.Bind(Initializer.ParameterMap, TEXT("MotionBlurParameters"));
 		OutScatteredMaxVelocity.Bind(Initializer.ParameterMap, TEXT("OutScatteredMaxVelocity"));
 	}
 
@@ -1103,6 +1128,7 @@ public:
 	{
 		bool bShaderHasOutdatedParameters = FGlobalShader::Serialize(Ar);
 		Ar << PostprocessParameter;
+		Ar << MotionBlurParameters;
 		Ar << OutScatteredMaxVelocity;
 		return bShaderHasOutdatedParameters;
 	}
@@ -1114,10 +1140,13 @@ public:
 		FGlobalShader::SetParameters(Context.RHICmdList, ShaderRHI, Context.View);
 
 		PostprocessParameter.SetCS(ShaderRHI, Context, TStaticSamplerState<SF_Point,AM_Clamp,AM_Clamp,AM_Clamp>::GetRHI());
+
+		SetShaderValue(Context.RHICmdList, ShaderRHI, MotionBlurParameters, GetMotionBlurParameters( Context.View ) );
 	}
 	
 private:
-	FPostProcessPassParameters PostprocessParameter;
+	FPostProcessPassParameters	PostprocessParameter;
+	FShaderParameter			MotionBlurParameters;
 };
 
 IMPLEMENT_SHADER_TYPE(,FPostProcessVelocityGatherCS,TEXT("PostProcessVelocityFlatten"),TEXT("VelocityGatherCS"),SF_Compute);
@@ -1151,7 +1180,8 @@ void FRCPassPostProcessVelocityGather::Process(FRenderingCompositePassContext& C
 
 	ComputeShader->SetParameters( Context );
 
-	DispatchComputeShader(Context.RHICmdList, *ComputeShader, TileCount.X, TileCount.Y, 1);
+	FIntPoint GroupCount = GetNumTiles16x16( TileCount );
+	DispatchComputeShader(Context.RHICmdList, *ComputeShader, GroupCount.X, GroupCount.Y, 1);
 
 	// un-set destination
 	Context.RHICmdList.SetUAVParameter( ComputeShader->GetComputeShader(), ComputeShader->OutScatteredMaxVelocity.GetBaseIndex(), NULL );
@@ -1216,7 +1246,7 @@ public:
 		return bShaderHasOutdatedParameters;
 	}
 
-	void SetParameters(const FRenderingCompositePassContext& Context)
+	void SetParameters(const FRenderingCompositePassContext& Context, float Scale)
 	{
 		const FPixelShaderRHIParamRef ShaderRHI = GetPixelShader();
 
@@ -1245,51 +1275,13 @@ public:
 
 				PostprocessParameter.SetPS( ShaderRHI, Context, 0, eFC_0000, Filters );
 			}
-			else if( CVarMotionBlurSmoothMax.GetValueOnRenderThread() )
-			{
-				FSamplerStateRHIParamRef Filters[] =
-				{
-					TStaticSamplerState<SF_Point,AM_Clamp,AM_Clamp,AM_Clamp>::GetRHI(),
-					TStaticSamplerState<SF_Point,AM_Clamp,AM_Clamp,AM_Clamp>::GetRHI(),
-					TStaticSamplerState<SF_Point,AM_Clamp,AM_Clamp,AM_Clamp>::GetRHI(),
-					TStaticSamplerState<SF_Bilinear,AM_Clamp,AM_Clamp,AM_Clamp>::GetRHI(),
-				};
-
-				PostprocessParameter.SetPS( ShaderRHI, Context, 0, eFC_0000, Filters );
-			}
 			else
 			{
 				PostprocessParameter.SetPS(ShaderRHI, Context, TStaticSamplerState<SF_Point,AM_Clamp,AM_Clamp,AM_Clamp>::GetRHI());
 			}
 		}
 
-		TRefCountPtr<IPooledRenderTarget> InputPooledElement = Context.Pass->GetInput(ePId_Input0)->GetOutput()->RequestInput();
-
-		{
-			const float SizeX = Context.View.ViewRect.Width();
-			const float SizeY = Context.View.ViewRect.Height();
-		
-			const float AspectRatio = SizeX / SizeY;
-			const float InvAspectRatio = SizeY / SizeX;
-
-			const FSceneViewState* ViewState = (FSceneViewState*) Context.View.State;
-			const float MotionBlurTimeScale = ViewState ? ViewState->MotionBlurTimeScale : 1.0f;
-
-			const float ViewMotionBlurScale = 0.5f * MotionBlurTimeScale * Context.View.FinalPostProcessSettings.MotionBlurAmount;
-
-			// 0:no 1:full screen width
-			float MaxVelocity = Context.View.FinalPostProcessSettings.MotionBlurMax / 100.0f;
-			float InvMaxVelocity = 1.0f / MaxVelocity;
-
-			// *2 to convert to -1..1 -1..1 screen space
-			// / MaxFraction to map screenpos to -1..1 normalized MaxFraction
-			FVector4 MotionBlurParametersValue(
-				ViewMotionBlurScale,
-				AspectRatio,
-				MaxVelocity,
-				InvMaxVelocity);
-			SetShaderValue(Context.RHICmdList, ShaderRHI, MotionBlurParameters, MotionBlurParametersValue);
-		}
+		SetShaderValue(Context.RHICmdList, ShaderRHI, MotionBlurParameters, GetMotionBlurParameters( Context.View, Scale ) );
 	}
 
 	static const TCHAR* GetSourceFilename()
@@ -1307,14 +1299,14 @@ public:
 #define VARIATION1(A) typedef FPostProcessMotionBlurNewPS<A> FPostProcessMotionBlurNewPS##A; \
 	IMPLEMENT_SHADER_TYPE2(FPostProcessMotionBlurNewPS##A, SF_Pixel);
 
-VARIATION1(0)			VARIATION1(1)			VARIATION1(2)			VARIATION1(3)			VARIATION1(4)
+VARIATION1(1)			VARIATION1(2)			VARIATION1(3)			VARIATION1(4)
 #undef VARIATION1
 
 
 
 // @param Quality 0: visualize, 1:low, 2:medium, 3:high, 4:very high
 template< uint32 Quality >
-static void SetMotionBlurShaderNewTempl(const FRenderingCompositePassContext& Context)
+static void SetMotionBlurShaderNewTempl( const FRenderingCompositePassContext& Context, float Scale )
 {
 	TShaderMapRef< FPostProcessVS > VertexShader( Context.GetShaderMap() );
 	TShaderMapRef< FPostProcessMotionBlurNewPS< Quality > >	PixelShader( Context.GetShaderMap() );
@@ -1323,14 +1315,7 @@ static void SetMotionBlurShaderNewTempl(const FRenderingCompositePassContext& Co
 	SetGlobalBoundShaderState(Context.RHICmdList, Context.GetFeatureLevel(), BoundShaderState, GFilterVertexDeclaration.VertexDeclarationRHI, *VertexShader, *PixelShader);
 
 	VertexShader->SetParameters(Context);
-	PixelShader->SetParameters(Context);
-}
-
-FRCPassPostProcessMotionBlurNew::FRCPassPostProcessMotionBlurNew(uint32 InQuality) 
-	: Quality(InQuality)
-{
-	// internal error
-	check(Quality >= 1 && Quality <= 4);
+	PixelShader->SetParameters(Context, Scale);
 }
 
 void FRCPassPostProcessMotionBlurNew::Process(FRenderingCompositePassContext& Context)
@@ -1372,22 +1357,35 @@ void FRCPassPostProcessMotionBlurNew::Process(FRenderingCompositePassContext& Co
 	Context.RHICmdList.SetRasterizerState(TStaticRasterizerState<>::GetRHI());
 	Context.RHICmdList.SetDepthStencilState(TStaticDepthStencilState<false, CF_Always>::GetRHI());
 
+	float BlurScaleLUT[] =
+	{
+		1.0f - 0.5f / 4.0f,
+		1.0f - 0.5f / 6.0f,
+		1.0f - 0.5f / 8.0f,
+		1.0f - 0.5f / 16.0f,
+		1.0f /  4.0f * CVarMotionBlur2ndScale.GetValueOnRenderThread(),
+		1.0f /  6.0f * CVarMotionBlur2ndScale.GetValueOnRenderThread(),
+		1.0f /  8.0f * CVarMotionBlur2ndScale.GetValueOnRenderThread(),
+		1.0f / 16.0f * CVarMotionBlur2ndScale.GetValueOnRenderThread(),
+	};
+	float Scale = Pass >= 0 ? BlurScaleLUT[ (Pass * 4) + (Quality - 1) ] : 1.0f;
+	
 	if(Quality == 1)
 	{
-		SetMotionBlurShaderNewTempl<1>(Context);
+		SetMotionBlurShaderNewTempl<1>( Context, Scale );
 	}
 	else if(Quality == 2)
 	{
-		SetMotionBlurShaderNewTempl<2>(Context);
+		SetMotionBlurShaderNewTempl<2>( Context, Scale );
 	}
-	else if(Quality == 3)
+	else if(Quality == 3 || Pass > 0 )
 	{
-		SetMotionBlurShaderNewTempl<3>(Context);
+		SetMotionBlurShaderNewTempl<3>( Context, Scale );
 	}
 	else
 	{
 		check(Quality == 4);
-		SetMotionBlurShaderNewTempl<4>(Context);
+		SetMotionBlurShaderNewTempl<4>( Context, Scale );
 	}
 
 	TShaderMapRef<FPostProcessVS> VertexShader(Context.GetShaderMap());
@@ -1414,7 +1412,6 @@ FPooledRenderTargetDesc FRCPassPostProcessMotionBlurNew::ComputeOutputDesc(EPass
 	FPooledRenderTargetDesc Ret = GetInput(ePId_Input0)->GetOutput()->RenderTargetDesc;
 
 	Ret.Reset();
-	// we don't need the alpha channel and 32bit is faster and costs less memory
 	Ret.Format = PF_FloatRGB;
 	Ret.DebugName = TEXT("MotionBlur");
 	Ret.AutoWritable = false;
diff --git a/Engine/Source/Runtime/Renderer/Private/PostProcess/PostProcessMotionBlur.h b/Engine/Source/Runtime/Renderer/Private/PostProcess/PostProcessMotionBlur.h
--- a/Engine/Source/Runtime/Renderer/Private/PostProcess/PostProcessMotionBlur.h
+++ b/Engine/Source/Runtime/Renderer/Private/PostProcess/PostProcessMotionBlur.h
@@ -125,7 +125,13 @@ class FRCPassPostProcessMotionBlurNew : public TRenderingCompositePassBase<4, 1>
 {
 public:
 	// @param InQuality 0xffffffff to visualize, 0:off(no shader is used), 1:low, 2:medium, 3:high, 4:very high
-	FRCPassPostProcessMotionBlurNew(uint32 InQuality);
+	FRCPassPostProcessMotionBlurNew( uint32 InQuality, int32 InPass )
+		: Quality(InQuality)
+		, Pass(InPass)
+	{
+		// internal error
+		check(Quality >= 1 && Quality <= 4);
+	}
 
 	// interface FRenderingCompositePass ---------
 	virtual void Process(FRenderingCompositePassContext& Context) override;
@@ -133,5 +139,6 @@ public:
 	virtual FPooledRenderTargetDesc ComputeOutputDesc(EPassOutputId InPassOutputId) const override;
 
 	// 1:low, 2:medium, 3:high, 4: very high
-	uint32 Quality;
+	uint32	Quality;
+	int32	Pass;
 };
diff --git a/Engine/Source/Runtime/Renderer/Private/PostProcess/PostProcessing.cpp b/Engine/Source/Runtime/Renderer/Private/PostProcess/PostProcessing.cpp
--- a/Engine/Source/Runtime/Renderer/Private/PostProcess/PostProcessing.cpp
+++ b/Engine/Source/Runtime/Renderer/Private/PostProcess/PostProcessing.cpp
@@ -154,6 +154,27 @@ static TAutoConsoleVariable<int32> CVarTonemapperScreenPercentage(
 	TEXT(" 2: enabled even if some other feature would prevent it (good for testing, likely rendering artifacts)"),
 	ECVF_Scalability | ECVF_RenderThreadSafe);
 
+static TAutoConsoleVariable<int32> CVarMotionBlurNew(
+	TEXT("r.MotionBlurNew"),
+	0,
+	TEXT(""),
+	ECVF_RenderThreadSafe
+);
+
+static TAutoConsoleVariable<int32> CVarMotionBlurScatter(
+	TEXT("r.MotionBlurScatter"),
+	0,
+	TEXT(""),
+	ECVF_RenderThreadSafe
+);
+
+static TAutoConsoleVariable<int32> CVarMotionBlurSeparable(
+	TEXT("r.MotionBlurSeparable"),
+	0,
+	TEXT(""),
+	ECVF_RenderThreadSafe
+);
+
 IMPLEMENT_SHADER_TYPE(,FPostProcessVS,TEXT("PostProcessBloom"),TEXT("MainPostprocessCommonVS"),SF_Vertex);
 
 static bool HasPostProcessMaterial(FPostprocessContext& Context, EBlendableLocation InLocation);
@@ -1067,34 +1088,6 @@ bool FPostProcessing::AllowFullPostProcessing(const FViewInfo& View, ERHIFeature
 		&& !View.Family->EngineShowFlags.VisualizeMeshDistanceFields;
 }
 
-static TAutoConsoleVariable<int32> CVarMotionBlurNew(
-	TEXT("r.MotionBlurNew"),
-	0,
-	TEXT(""),
-	ECVF_RenderThreadSafe
-);
-
-static TAutoConsoleVariable<int32> CVarMotionBlurScatter(
-	TEXT("r.MotionBlurScatter"),
-	1,
-	TEXT(""),
-	ECVF_RenderThreadSafe
-);
-
-static TAutoConsoleVariable<int32> CVarMotionBlurDilate(
-	TEXT("r.MotionBlurDilate"),
-	0,
-	TEXT(""),
-	ECVF_RenderThreadSafe
-);
-
-static TAutoConsoleVariable<int32> CVarMotionBlurSeparable(
-	TEXT("r.MotionBlurSeparable"),
-	0,
-	TEXT(""),
-	ECVF_RenderThreadSafe
-);
-
 void FPostProcessing::Process(FRHICommandListImmediate& RHICmdList, FViewInfo& View, TRefCountPtr<IPooledRenderTarget>& VelocityRT)
 {
 	QUICK_SCOPE_CYCLE_COUNTER( STAT_PostProcessing_Process );
@@ -1143,15 +1136,15 @@ void FPostProcessing::Process(FRHICommandListImmediate& RHICmdList, FViewInfo& V
 		FSceneViewState* ViewState = (FSceneViewState*)Context.View.State;
 
 		{
-			if(ViewState && ViewState->SeparateTranslucencyRT)
+			if (FSceneRenderTargets::Get(RHICmdList).SeparateTranslucencyRT)
 			{
-				FRenderingCompositePass* NodeSeparateTranslucency = Context.Graph.RegisterPass(new(FMemStack::Get()) FRCPassPostProcessInput(ViewState->SeparateTranslucencyRT));
+				FRenderingCompositePass* NodeSeparateTranslucency = Context.Graph.RegisterPass(new(FMemStack::Get()) FRCPassPostProcessInput(FSceneRenderTargets::Get(RHICmdList).SeparateTranslucencyRT));
 
 				SeparateTranslucency = FRenderingCompositeOutputRef(NodeSeparateTranslucency);
 				// the node keeps another reference so the RT will not be release too early
-				ViewState->FreeSeparateTranslucency();
+				FSceneRenderTargets::Get(RHICmdList).FreeSeparateTranslucency();
 
-				check(!ViewState->SeparateTranslucencyRT);
+				check(!FSceneRenderTargets::Get(RHICmdList).SeparateTranslucencyRT);
 			}
 		}
 
@@ -1306,24 +1299,30 @@ void FPostProcessing::Process(FRHICommandListImmediate& RHICmdList, FViewInfo& V
 						MaxTileVelocity	= FRenderingCompositeOutputRef( VelocityFlattenPass, ePId_Output1 );
 					}
 
-					if( CVarMotionBlurScatter.GetValueOnRenderThread() )
+					const float SizeX = View.ViewRect.Width();
+
+					// 0:no 1:full screen width, percent conversion
+					float MaxVelocity = View.FinalPostProcessSettings.MotionBlurMax / 100.0f;
+					float MaxVelocityTiles = MaxVelocity * SizeX * (0.5f / 16.0f);
+					float MaxTileDistGathered = 3.0f;
+					if( MaxVelocityTiles > MaxTileDistGathered || CVarMotionBlurScatter.GetValueOnRenderThread() )
 					{
 						FRenderingCompositePass* VelocityScatterPass = Context.Graph.RegisterPass( new(FMemStack::Get()) FRCPassPostProcessVelocityScatter() );
 						VelocityScatterPass->SetInput( ePId_Input0, MaxTileVelocity );
 
 						MaxTileVelocity	= FRenderingCompositeOutputRef( VelocityScatterPass );
 					}
-
-					if( !CVarMotionBlurScatter.GetValueOnRenderThread() )
+					else
 					{
-						FRenderingCompositePass* VelocityDilatePass = Context.Graph.RegisterPass( new(FMemStack::Get()) FRCPassPostProcessVelocityGather() );
-						VelocityDilatePass->SetInput( ePId_Input0, MaxTileVelocity );
+						FRenderingCompositePass* VelocityGatherPass = Context.Graph.RegisterPass( new(FMemStack::Get()) FRCPassPostProcessVelocityGather() );
+						VelocityGatherPass->SetInput( ePId_Input0, MaxTileVelocity );
 
-						MaxTileVelocity	= FRenderingCompositeOutputRef( VelocityDilatePass );
+						MaxTileVelocity	= FRenderingCompositeOutputRef( VelocityGatherPass );
 					}
 
+					bool bTwoPass = CVarMotionBlurSeparable.GetValueOnRenderThread() != 0;
 					{
-						FRenderingCompositePass* MotionBlurPass = Context.Graph.RegisterPass( new(FMemStack::Get()) FRCPassPostProcessMotionBlurNew( GetMotionBlurQualityFromCVar() ) );
+						FRenderingCompositePass* MotionBlurPass = Context.Graph.RegisterPass( new(FMemStack::Get()) FRCPassPostProcessMotionBlurNew( GetMotionBlurQualityFromCVar(), bTwoPass ? 0 : -1 ) );
 						MotionBlurPass->SetInput( ePId_Input0, Context.FinalOutput );
 						MotionBlurPass->SetInput( ePId_Input1, SceneDepth );
 						MotionBlurPass->SetInput( ePId_Input2, VelocityInput );
@@ -1332,9 +1331,9 @@ void FPostProcessing::Process(FRHICommandListImmediate& RHICmdList, FViewInfo& V
 						Context.FinalOutput = FRenderingCompositeOutputRef( MotionBlurPass );
 					}
 
-					if( CVarMotionBlurSeparable.GetValueOnRenderThread() )
+					if( bTwoPass )
 					{
-						FRenderingCompositePass* MotionBlurPass = Context.Graph.RegisterPass( new(FMemStack::Get()) FRCPassPostProcessMotionBlurNew( GetMotionBlurQualityFromCVar() ) );
+						FRenderingCompositePass* MotionBlurPass = Context.Graph.RegisterPass( new(FMemStack::Get()) FRCPassPostProcessMotionBlurNew( GetMotionBlurQualityFromCVar(), 1 ) );
 						MotionBlurPass->SetInput( ePId_Input0, Context.FinalOutput );
 						MotionBlurPass->SetInput( ePId_Input1, SceneDepth );
 						MotionBlurPass->SetInput( ePId_Input2, VelocityInput );
@@ -1607,11 +1606,7 @@ void FPostProcessing::Process(FRHICommandListImmediate& RHICmdList, FViewInfo& V
 		}
 		else
 		{
-			if(ViewState)
-			{
-				check(!ViewState->SeparateTranslucencyRT);
-			}
-
+			check(!FSceneRenderTargets::Get(RHICmdList).SeparateTranslucencyRT);
 			AddGammaOnlyTonemapper(Context);
 		}
 		
diff --git a/Engine/Source/Runtime/Renderer/Private/PostProcess/SceneRenderTargets.cpp b/Engine/Source/Runtime/Renderer/Private/PostProcess/SceneRenderTargets.cpp
--- a/Engine/Source/Runtime/Renderer/Private/PostProcess/SceneRenderTargets.cpp
+++ b/Engine/Source/Runtime/Renderer/Private/PostProcess/SceneRenderTargets.cpp
@@ -240,6 +240,8 @@ FSceneRenderTargets::FSceneRenderTargets(const FViewInfo& View, const FSceneRend
 	, bSceneDepthCleared(SnapshotSource.bSceneDepthCleared)	
 	, bSnapshot(true)
 	, QuadOverdrawIndex(SnapshotSource.QuadOverdrawIndex)
+	, SeparateTranslucencyRT(SnapshotSource.SeparateTranslucencyRT)
+	, SeparateTranslucencyDepthRT(SnapshotSource.SeparateTranslucencyDepthRT)
 {
 	SnapshotArray(SceneColor, SnapshotSource.SceneColor);
 	SnapshotArray(TranslucencyShadowTransmission, SnapshotSource.TranslucencyShadowTransmission);
@@ -249,13 +251,6 @@ FSceneRenderTargets::FSceneRenderTargets(const FViewInfo& View, const FSceneRend
 	SnapshotArray(ReflectionBrightness, SnapshotSource.ReflectionBrightness);
 	SnapshotArray(TranslucencyLightingVolumeAmbient, SnapshotSource.TranslucencyLightingVolumeAmbient);
 	SnapshotArray(TranslucencyLightingVolumeDirectional, SnapshotSource.TranslucencyLightingVolumeDirectional);
-
-	FSceneViewState* ViewState = (FSceneViewState*)View.State;
-	if (ViewState)
-	{
-		SeparateTranslucencyRT = GRenderTargetPool.MakeSnapshot(ViewState->SeparateTranslucencyRT);
-		SeparateTranslucencyDepthRT = GRenderTargetPool.MakeSnapshot(ViewState->SeparateTranslucencyDepthRT);
-	}
 }
 
 
@@ -1291,11 +1286,9 @@ bool FSceneRenderTargets::BeginRenderingSeparateTranslucency(FRHICommandList& RH
 		}
 		else
 		{
-			FSceneViewState* ViewState = (FSceneViewState*)View.State;
-			SeparateTranslucency = &ViewState->GetSeparateTranslucency(RHICmdList, ScaledSize);
+			SeparateTranslucency = &GetSeparateTranslucency(RHICmdList, ScaledSize);
 		}
-		FSceneViewState* ViewState = (FSceneViewState*)View.State;
-		const FTexture2DRHIRef &SeparateTranslucencyDepth = Scale < 1.0f ? (const FTexture2DRHIRef&)ViewState->GetSeparateTranslucencyDepth(RHICmdList, GetBufferSizeXY())->GetRenderTargetItem().TargetableTexture : GetSceneDepthSurface();
+		const FTexture2DRHIRef &SeparateTranslucencyDepth = Scale < 1.0f ? (const FTexture2DRHIRef&)GetSeparateTranslucencyDepth(RHICmdList, GetBufferSizeXY())->GetRenderTargetItem().TargetableTexture : GetSceneDepthSurface();
 
 		check((*SeparateTranslucency)->GetRenderTargetItem().TargetableTexture->GetClearColor() == FLinearColor::Black);
 		// clear the render target the first time, re-use afterwards
@@ -1334,10 +1327,9 @@ void FSceneRenderTargets::FinishRenderingSeparateTranslucency(FRHICommandList& R
 		{
 			static const auto CVar = IConsoleManager::Get().FindTConsoleVariableDataFloat(TEXT("r.SeparateTranslucencyScreenPercentage"));
 			float Scale = CVar->GetValueOnRenderThread() / 100.0f;
-			FSceneViewState* ViewState = (FSceneViewState*)View.State;
 			FIntPoint ScaledSize(GetBufferSizeXY().X * Scale, GetBufferSizeXY().Y * Scale);
-			SeparateTranslucency = &ViewState->GetSeparateTranslucency( RHICmdList, ScaledSize );
-			SeparateTranslucencyDepth = &ViewState->GetSeparateTranslucencyDepth(RHICmdList, GetBufferSizeXY());
+			SeparateTranslucency = &GetSeparateTranslucency( RHICmdList, ScaledSize );
+			SeparateTranslucencyDepth = &GetSeparateTranslucencyDepth(RHICmdList, GetBufferSizeXY());
 		}
 
 		RHICmdList.CopyToResolveTarget((*SeparateTranslucency)->GetRenderTargetItem().TargetableTexture, (*SeparateTranslucency)->GetRenderTargetItem().ShaderResourceTexture, true, FResolveParams());
@@ -1447,8 +1439,7 @@ bool FSceneRenderTargets::IsSeparateTranslucencyActive(const FViewInfo& View) co
 
 	return (Value != 0) && CurrentFeatureLevel >= ERHIFeatureLevel::SM4
 		&& View.Family->EngineShowFlags.PostProcessing
-		&& View.Family->EngineShowFlags.SeparateTranslucency
-		&& View.State != NULL;	// We require a ViewState in order for separate translucency to work (it keeps track of our SeparateTranslucencyRT)
+		&& View.Family->EngineShowFlags.SeparateTranslucency;
 }
 
 void FSceneRenderTargets::InitEditorPrimitivesColor(FRHICommandList& RHICmdList)
@@ -2269,15 +2260,13 @@ void FSceneTextureShaderParameters::Set(
 
 		if(SceneDepthTextureParameter.IsBound() || SceneDepthTextureParameterSampler.IsBound())
 		{
-			FSceneViewState* ViewState = (FSceneViewState*)View.State;
 			static const auto CVar = IConsoleManager::Get().FindTConsoleVariableDataFloat(TEXT("r.SeparateTranslucencyScreenPercentage"));
 			const FTexture2DRHIRef* DepthTexture = SceneContext.GetActualDepthTexture();
 			if (SceneContext.IsSeparateTranslucencyPass()
 				&& CVar->GetValueOnRenderThread() < 100
-				&& ViewState
-				&& ViewState->IsSeparateTranslucencyDepthValid())
+				&& SceneContext.IsSeparateTranslucencyDepthValid())
 			{
-				DepthTexture = &ViewState->GetSeparateTranslucencyDepthSurface();
+				DepthTexture = &SceneContext.GetSeparateTranslucencyDepthSurface();
 			}
 
 			SetTextureParameter(
diff --git a/Engine/Source/Runtime/Renderer/Private/PostProcess/SceneRenderTargets.h b/Engine/Source/Runtime/Renderer/Private/PostProcess/SceneRenderTargets.h
--- a/Engine/Source/Runtime/Renderer/Private/PostProcess/SceneRenderTargets.h
+++ b/Engine/Source/Runtime/Renderer/Private/PostProcess/SceneRenderTargets.h
@@ -267,7 +267,11 @@ public:
 
 	bool BeginRenderingSeparateTranslucency(FRHICommandList& RHICmdList, const FViewInfo& View, bool bFirstTimeThisFrame);
 	void FinishRenderingSeparateTranslucency(FRHICommandList& RHICmdList, const FViewInfo& View);
-	void FreeSeparateTranslucency();
+	void FreeSeparateTranslucency()
+	{
+		SeparateTranslucencyRT.SafeRelease();
+		check(!SeparateTranslucencyRT);
+	}
 
 	void ResolveSceneDepthTexture(FRHICommandList& RHICmdList);
 	void ResolveSceneDepthToAuxiliaryTexture(FRHICommandList& RHICmdList);
@@ -281,6 +285,42 @@ public:
 	void BeginRenderingLightAttenuation(FRHICommandList& RHICmdList, bool bClearToWhite = false);
 	void FinishRenderingLightAttenuation(FRHICommandList& RHICmdList);
 
+
+
+	TRefCountPtr<IPooledRenderTarget>& GetSeparateTranslucency(FRHICommandList& RHICmdList, FIntPoint Size)
+	{
+		if (!SeparateTranslucencyRT || SeparateTranslucencyRT->GetDesc().Extent != Size)
+		{
+			// Create the SeparateTranslucency render target (alpha is needed to lerping)
+			FPooledRenderTargetDesc Desc(FPooledRenderTargetDesc::Create2DDesc(Size, PF_FloatRGBA, FClearValueBinding::Black, TexCreate_None, TexCreate_RenderTargetable, false));
+			Desc.AutoWritable = false;
+			GRenderTargetPool.FindFreeElement(RHICmdList, Desc, SeparateTranslucencyRT, TEXT("SeparateTranslucency"));
+		}
+		return SeparateTranslucencyRT;
+	}
+
+	bool IsSeparateTranslucencyDepthValid()
+	{
+		return SeparateTranslucencyDepthRT != nullptr;
+	}
+
+	TRefCountPtr<IPooledRenderTarget>& GetSeparateTranslucencyDepth(FRHICommandList& RHICmdList, FIntPoint Size)
+	{
+		if (!SeparateTranslucencyDepthRT || SeparateTranslucencyDepthRT->GetDesc().Extent != Size)
+		{
+			// Create the SeparateTranslucency depth render target 
+			FPooledRenderTargetDesc Desc(FPooledRenderTargetDesc::Create2DDesc(Size, PF_DepthStencil, FClearValueBinding::None, TexCreate_None, TexCreate_DepthStencilTargetable, true));
+			GRenderTargetPool.FindFreeElement(RHICmdList, Desc, SeparateTranslucencyDepthRT, TEXT("SeparateTranslucencyDepth"));
+		}
+		return SeparateTranslucencyDepthRT;
+	}
+	const FTexture2DRHIRef& GetSeparateTranslucencyDepthSurface()
+	{
+		return (const FTexture2DRHIRef&)SeparateTranslucencyDepthRT->GetRenderTargetItem().TargetableTexture;
+	}
+
+
+
 	/**
 	 * Cleans up editor primitive targets that we no longer need
 	 */
diff --git a/Engine/Source/Runtime/Renderer/Private/PreCullTriangles.cpp b/Engine/Source/Runtime/Renderer/Private/PreCullTriangles.cpp
--- a/Engine/Source/Runtime/Renderer/Private/PreCullTriangles.cpp
+++ b/Engine/Source/Runtime/Renderer/Private/PreCullTriangles.cpp
@@ -43,6 +43,7 @@ public:
 	{
 		FGlobalShader::ModifyCompilationEnvironment(Platform,OutEnvironment);
 		OutEnvironment.SetDefine(TEXT("UPDATEOBJECTS_THREADGROUP_SIZE"), UpdateObjectsGroupSize);
+		OutEnvironment.CompilerFlags.Add(CFLAG_StandardOptimization);
 	}
 
 	FCullObjectsForBoundsCS(const ShaderMetaType::CompiledShaderInitializerType& Initializer)
@@ -193,6 +194,7 @@ public:
 		FGlobalShader::ModifyCompilationEnvironment(Platform,OutEnvironment);
 		OutEnvironment.SetDefine(TEXT("THREADGROUP_SIZEX"), GDistanceFieldAOTileSizeX);
 		OutEnvironment.SetDefine(TEXT("THREADGROUP_SIZEY"), GDistanceFieldAOTileSizeY);
+		OutEnvironment.CompilerFlags.Add(CFLAG_StandardOptimization);
 	}
 
 	FPreCullTrianglesCS(const ShaderMetaType::CompiledShaderInitializerType& Initializer)
diff --git a/Engine/Source/Runtime/Renderer/Private/PrimitiveSceneInfo.cpp b/Engine/Source/Runtime/Renderer/Private/PrimitiveSceneInfo.cpp
--- a/Engine/Source/Runtime/Renderer/Private/PrimitiveSceneInfo.cpp
+++ b/Engine/Source/Runtime/Renderer/Private/PrimitiveSceneInfo.cpp
@@ -135,7 +135,6 @@ FPrimitiveSceneInfo::~FPrimitiveSceneInfo()
 
 void FPrimitiveSceneInfo::AddStaticMeshes(FRHICommandListImmediate& RHICmdList)
 {
-	ConditionalUpdateUniformBuffer(RHICmdList);
 	// Cache the primitive's static mesh elements.
 	FBatchingSPDI BatchingSPDI(this);
 	BatchingSPDI.SetHitProxy(DefaultDynamicHitProxy);
diff --git a/Engine/Source/Runtime/Renderer/Private/ReflectionEnvironmentCapture.cpp b/Engine/Source/Runtime/Renderer/Private/ReflectionEnvironmentCapture.cpp
--- a/Engine/Source/Runtime/Renderer/Private/ReflectionEnvironmentCapture.cpp
+++ b/Engine/Source/Runtime/Renderer/Private/ReflectionEnvironmentCapture.cpp
@@ -586,10 +586,11 @@ public:
 		SourceTexture.Bind(Initializer.ParameterMap,TEXT("SourceTexture"));
 		SourceTextureSampler.Bind(Initializer.ParameterMap,TEXT("SourceTextureSampler"));
 		SkyLightCaptureParameters.Bind(Initializer.ParameterMap,TEXT("SkyLightCaptureParameters"));
+		SinCosSourceCubemapRotation.Bind(Initializer.ParameterMap,TEXT("SinCosSourceCubemapRotation"));
 	}
 	FCopyCubemapToCubeFacePS() {}
 
-	void SetParameters(FRHICommandList& RHICmdList, const FTexture* SourceCubemap, uint32 CubeFaceValue, bool bIsSkyLight, bool bLowerHemisphereIsBlack)
+	void SetParameters(FRHICommandList& RHICmdList, const FTexture* SourceCubemap, uint32 CubeFaceValue, bool bIsSkyLight, bool bLowerHemisphereIsBlack, float SourceCubemapRotation)
 	{
 		const FPixelShaderRHIParamRef ShaderRHI = GetPixelShader();
 
@@ -603,6 +604,8 @@ public:
 			SourceCubemap);
 
 		SetShaderValue(RHICmdList, ShaderRHI, SkyLightCaptureParameters, FVector(bIsSkyLight ? 1.0f : 0.0f, 0.0f, bLowerHemisphereIsBlack ? 1.0f : 0.0f));
+
+		SetShaderValue(RHICmdList, ShaderRHI, SinCosSourceCubemapRotation, FVector2D(FMath::Sin(SourceCubemapRotation), FMath::Cos(SourceCubemapRotation)));
 	}
 
 	virtual bool Serialize(FArchive& Ar) override
@@ -612,6 +615,7 @@ public:
 		Ar << SourceTexture;
 		Ar << SourceTextureSampler;
 		Ar << SkyLightCaptureParameters;
+		Ar << SinCosSourceCubemapRotation;
 		return bShaderHasOutdatedParameters;
 	}
 
@@ -620,6 +624,7 @@ private:
 	FShaderResourceParameter SourceTexture;
 	FShaderResourceParameter SourceTextureSampler;
 	FShaderParameter SkyLightCaptureParameters;
+	FShaderParameter SinCosSourceCubemapRotation;
 };
 
 IMPLEMENT_SHADER_TYPE(,FCopyCubemapToCubeFacePS,TEXT("ReflectionEnvironmentShaders"),TEXT("CopyCubemapToCubeFaceColorPS"),SF_Pixel);
@@ -766,7 +771,7 @@ void CaptureSceneToScratchCubemap(FRHICommandListImmediate& RHICmdList, FSceneRe
 	FSceneRenderer::WaitForTasksClearSnapshotsAndDeleteSceneRenderer(RHICmdList, SceneRenderer);
 }
 
-void CopyCubemapToScratchCubemap(FRHICommandList& RHICmdList, ERHIFeatureLevel::Type FeatureLevel, UTextureCube* SourceCubemap, bool bIsSkyLight, bool bLowerHemisphereIsBlack)
+void CopyCubemapToScratchCubemap(FRHICommandList& RHICmdList, ERHIFeatureLevel::Type FeatureLevel, UTextureCube* SourceCubemap, bool bIsSkyLight, bool bLowerHemisphereIsBlack, float SourceCubemapRotation)
 {
 	check(SourceCubemap);
 	
@@ -790,7 +795,7 @@ void CopyCubemapToScratchCubemap(FRHICommandList& RHICmdList, ERHIFeatureLevel::
 		TShaderMapRef<FCopyCubemapToCubeFacePS> PixelShader(GetGlobalShaderMap(FeatureLevel));
 
 		SetGlobalBoundShaderState(RHICmdList, FeatureLevel, CopyFromCubemapToCubemapBoundShaderState, GFilterVertexDeclaration.VertexDeclarationRHI, *VertexShader, *PixelShader);
-		PixelShader->SetParameters(RHICmdList, SourceCubemapResource, CubeFace, bIsSkyLight, bLowerHemisphereIsBlack);
+		PixelShader->SetParameters(RHICmdList, SourceCubemapResource, CubeFace, bIsSkyLight, bLowerHemisphereIsBlack, SourceCubemapRotation);
 
 		DrawRectangle(
 			RHICmdList,
@@ -1398,14 +1403,14 @@ void FScene::UpdateSkyCaptureContents(const USkyLightComponent* CaptureComponent
 		}
 		else if (CaptureComponent->SourceType == SLS_SpecifiedCubemap)
 		{
-			ENQUEUE_UNIQUE_RENDER_COMMAND_THREEPARAMETER( 
+			ENQUEUE_UNIQUE_RENDER_COMMAND_FOURPARAMETER( 
 				CopyCubemapCommand,
 				UTextureCube*, SourceTexture, SourceCubemap,
 				bool, bLowerHemisphereIsBlack, CaptureComponent->bLowerHemisphereIsBlack,
+				float, SourceCubemapRotation, CaptureComponent->SourceCubemapAngle * (PI / 180.f),
 				ERHIFeatureLevel::Type, FeatureLevel, GetFeatureLevel(),
 			{
-				
-				CopyCubemapToScratchCubemap(RHICmdList, FeatureLevel, SourceTexture, true, bLowerHemisphereIsBlack);
+				CopyCubemapToScratchCubemap(RHICmdList, FeatureLevel, SourceTexture, true, bLowerHemisphereIsBlack, SourceCubemapRotation);
 			});
 		}
 		else
diff --git a/Engine/Source/Runtime/Renderer/Private/ScenePrivate.h b/Engine/Source/Runtime/Renderer/Private/ScenePrivate.h
--- a/Engine/Source/Runtime/Renderer/Private/ScenePrivate.h
+++ b/Engine/Source/Runtime/Renderer/Private/ScenePrivate.h
@@ -505,10 +505,6 @@ public:
 
 	FHeightfieldLightingAtlas* HeightfieldLightingAtlas;
 
-	// If Translucency should be rendered into a separate RT and composited without DepthOfField, can be disabled in the materials (affects sorting)
-	TRefCountPtr<IPooledRenderTarget> SeparateTranslucencyRT;
-	// Depth buffer for separate translucency rendering (needed if SeparateTranslucencyScreenPercentage is <100)
-	TRefCountPtr<IPooledRenderTarget> SeparateTranslucencyDepthRT;
 	// Temporal AA result of last frame
 	TRefCountPtr<IPooledRenderTarget> TemporalAAHistoryRT;
 	TRefCountPtr<IPooledRenderTarget> PendingTemporalAAHistoryRT;
@@ -607,12 +603,6 @@ public:
 		return DistanceFieldTemporalSampleIndex;
 	}
 
-	void FreeSeparateTranslucency()
-	{
-		SeparateTranslucencyRT.SafeRelease();
-
-		check(!SeparateTranslucencyRT);
-	}
 
 	// call only if not yet created
 	void SetupLightPropagationVolume(FSceneView& View, FSceneViewFamily& ViewFamily);
@@ -696,37 +686,6 @@ public:
 		bValidEyeAdaptation = true;
 	}
 
-	TRefCountPtr<IPooledRenderTarget>& GetSeparateTranslucency(FRHICommandList& RHICmdList, FIntPoint Size)
-	{
-		if (!SeparateTranslucencyRT || SeparateTranslucencyRT->GetDesc().Extent != Size)
-		{
-			// Create the SeparateTranslucency render target (alpha is needed to lerping)
-			FPooledRenderTargetDesc Desc(FPooledRenderTargetDesc::Create2DDesc(Size, PF_FloatRGBA, FClearValueBinding::Black, TexCreate_None, TexCreate_RenderTargetable, false));
-			Desc.AutoWritable = false;
-			GRenderTargetPool.FindFreeElement(RHICmdList, Desc, SeparateTranslucencyRT, TEXT("SeparateTranslucency"));
-		}
-		return SeparateTranslucencyRT;
-	}
-
-	bool IsSeparateTranslucencyDepthValid()
-	{
-		return SeparateTranslucencyDepthRT != nullptr;
-	}
-
-	TRefCountPtr<IPooledRenderTarget>& GetSeparateTranslucencyDepth(FRHICommandList& RHICmdList, FIntPoint Size)
-	{
-		if (!SeparateTranslucencyDepthRT || SeparateTranslucencyDepthRT->GetDesc().Extent != Size)
-		{
-			// Create the SeparateTranslucency depth render target 
-			FPooledRenderTargetDesc Desc(FPooledRenderTargetDesc::Create2DDesc(Size, PF_DepthStencil, FClearValueBinding::None, TexCreate_None, TexCreate_DepthStencilTargetable, true));
-			GRenderTargetPool.FindFreeElement(RHICmdList, Desc, SeparateTranslucencyDepthRT, TEXT("SeparateTranslucencyDepth"));
-		}
-		return SeparateTranslucencyDepthRT;
-	}
-	const FTexture2DRHIRef& GetSeparateTranslucencyDepthSurface()
-	{
-		return (const FTexture2DRHIRef&)SeparateTranslucencyDepthRT->GetRenderTargetItem().TargetableTexture;
-	}
 
 	// FRenderResource interface.
 	virtual void InitDynamicRHI() override
@@ -745,8 +704,6 @@ public:
 		OcclusionQueryPool.Release();
 		HZBOcclusionTests.ReleaseDynamicRHI();
 		EyeAdaptationRT.SafeRelease();
-		SeparateTranslucencyRT.SafeRelease();
-		SeparateTranslucencyDepthRT.SafeRelease();
 		TemporalAAHistoryRT.SafeRelease();
 		PendingTemporalAAHistoryRT.SafeRelease();
 		DOFHistoryRT.SafeRelease();
diff --git a/Engine/Source/Runtime/Renderer/Private/TranslucentRendering.cpp b/Engine/Source/Runtime/Renderer/Private/TranslucentRendering.cpp
--- a/Engine/Source/Runtime/Renderer/Private/TranslucentRendering.cpp
+++ b/Engine/Source/Runtime/Renderer/Private/TranslucentRendering.cpp
@@ -964,17 +964,16 @@ void FDeferredShadingSceneRenderer::RenderTranslucencyParallel(FRHICommandListIm
 		{
 			if (SceneContext.IsSeparateTranslucencyActive(View))
 			{
-				FSceneViewState* ViewState = (FSceneViewState*)View.State;
 				// we need to allocate this now so it ends up in the snapshot
 				static const auto CVar = IConsoleManager::Get().FindTConsoleVariableDataFloat(TEXT("r.SeparateTranslucencyScreenPercentage"));
 				float Scale = CVar->GetValueOnRenderThread() / 100.0f;
 				FIntPoint ScaledSize(SceneContext.GetBufferSizeXY().X * Scale, SceneContext.GetBufferSizeXY().Y * (Scale / 100.0f) );
-				ViewState->GetSeparateTranslucency(RHICmdList, ScaledSize);
+				SceneContext.GetSeparateTranslucency(RHICmdList, ScaledSize);
 
 				if (Scale<1.0f)
 				{
-					ViewState->GetSeparateTranslucencyDepth(RHICmdList, SceneContext.GetBufferSizeXY());
-					DownsampleDepthSurface(RHICmdList, ViewState->GetSeparateTranslucencyDepthSurface(), View, Scale, 1.0f);
+					SceneContext.GetSeparateTranslucencyDepth(RHICmdList, SceneContext.GetBufferSizeXY());
+					DownsampleDepthSurface(RHICmdList, SceneContext.GetSeparateTranslucencyDepthSurface(), View, Scale, 1.0f);
 				}
 			}
 			FTranslucencyPassParallelCommandListSet ParallelCommandListSet(View, RHICmdList, 
@@ -1143,16 +1142,12 @@ void FDeferredShadingSceneRenderer::RenderTranslucency(FRHICommandListImmediate&
 				if (SceneContext.IsSeparateTranslucencyActive(View))
 				{
 					// always call BeginRenderingSeparateTranslucency() even if there are no primitives to we keep the RT allocated
-					FSceneViewState* ViewState = (FSceneViewState*)View.State;
-					if (ViewState)
+					static const auto CVar = IConsoleManager::Get().FindTConsoleVariableDataFloat(TEXT("r.SeparateTranslucencyScreenPercentage"));
+					float Scale = CVar->GetValueOnRenderThread() / 100.0f;
+					if (Scale < 1.0f)
 					{
-						static const auto CVar = IConsoleManager::Get().FindTConsoleVariableDataFloat(TEXT("r.SeparateTranslucencyScreenPercentage"));
-						float Scale = CVar->GetValueOnRenderThread() / 100.0f;
-						if (Scale < 1.0f)
-						{
-							ViewState->GetSeparateTranslucencyDepth(RHICmdList, SceneContext.GetBufferSizeXY());
-							DownsampleDepthSurface(RHICmdList, ViewState->GetSeparateTranslucencyDepthSurface(), View, Scale, 1.0f);
-						}
+						SceneContext.GetSeparateTranslucencyDepth(RHICmdList, SceneContext.GetBufferSizeXY());
+						DownsampleDepthSurface(RHICmdList, SceneContext.GetSeparateTranslucencyDepthSurface(), View, Scale, 1.0f);
 					}
 					bool bSetupTranslucency = SceneContext.BeginRenderingSeparateTranslucency(RHICmdList, View, true);
 
diff --git a/Engine/Source/Runtime/Renderer/Public/RendererInterface.h b/Engine/Source/Runtime/Renderer/Public/RendererInterface.h
--- a/Engine/Source/Runtime/Renderer/Public/RendererInterface.h
+++ b/Engine/Source/Runtime/Renderer/Public/RendererInterface.h
@@ -183,7 +183,7 @@ public:
 			&& TargetableFlags == rhs.TargetableFlags
 			&& bForceSeparateTargetAndShaderResource == rhs.bForceSeparateTargetAndShaderResource
 			&& ClearValue == rhs.ClearValue
-			&& AutoWritable == AutoWritable;
+			&& AutoWritable == rhs.AutoWritable;
 	}
 
 	bool IsCubemap() const
diff --git a/Engine/Source/Runtime/ShaderCore/Private/Shader.cpp b/Engine/Source/Runtime/ShaderCore/Private/Shader.cpp
--- a/Engine/Source/Runtime/ShaderCore/Private/Shader.cpp
+++ b/Engine/Source/Runtime/ShaderCore/Private/Shader.cpp
@@ -120,6 +120,7 @@ FShaderType::FShaderType(
 	GetStreamOutElementsType InGetStreamOutElementsRef
 	):
 	Name(InName),
+	TypeName(InName),
 	SourceFilename(InSourceFilename),
 	FunctionName(InFunctionName),
 	Frequency(InFrequency),
@@ -141,7 +142,7 @@ FShaderType::FShaderType(
 
 	// register this shader type
 	GlobalListLink.LinkHead(GetTypeList());
-	GetNameToTypeMap().Add(FName(InName), this);
+	GetNameToTypeMap().Add(TypeName, this);
 
 	// Assign the shader type the next unassigned hash index.
 	static uint32 NextHashIndex = 0;
@@ -151,12 +152,12 @@ FShaderType::FShaderType(
 FShaderType::~FShaderType()
 {
 	GlobalListLink.Unlink();
-	GetNameToTypeMap().Remove(FName(Name));
+	GetNameToTypeMap().Remove(TypeName);
 }
 
 TLinkedList<FShaderType*>*& FShaderType::GetTypeList()
 {
-	static TLinkedList<FShaderType*>* GShaderTypeList = NULL;
+	static TLinkedList<FShaderType*>* GShaderTypeList = nullptr;
 	return GShaderTypeList;
 }
 
@@ -172,7 +173,7 @@ FShaderType* FShaderType::GetShaderTypeByName(const TCHAR* Name)
 		}
 	}
 
-	return NULL;
+	return nullptr;
 }
 
 TArray<FShaderType*> FShaderType::GetShaderTypesByFilename(const TCHAR* Filename)
@@ -649,8 +650,8 @@ FShaderId::FShaderId(const FSHAHash& InMaterialShaderMapHash, const FShaderPipel
 	}
 	else
 	{
-		VFSerializationHistory = NULL;
-		VertexFactoryType = NULL;
+		VFSerializationHistory = nullptr;
+		VertexFactoryType = nullptr;
 	}
 }
 
@@ -983,11 +984,11 @@ FShaderPipelineType::FShaderPipelineType(
 	const FShaderType* InDomainShader,
 	const FShaderType* InGeometryShader,
 	const FShaderType* InPixelShader,
-	bool bInShouldAddShadersToIndividualMap) :
+	bool bInShouldOptimizeUnusedOutputs) :
 	Name(InName),
 	TypeName(Name),
 	GlobalListLink(this),
-	bAddShadersToIndividualMap(bInShouldAddShadersToIndividualMap)
+	bShouldOptimizeUnusedOutputs(bInShouldOptimizeUnusedOutputs)
 {
 	checkf(Name && *Name, TEXT("Shader Pipeline Type requires a valid Name!"));
 
@@ -1035,12 +1036,6 @@ FShaderPipelineType::FShaderPipelineType(
 
 FShaderPipelineType::~FShaderPipelineType()
 {
-	if (!bAddShadersToIndividualMap)
-	{
-		//#todo-rco: Delete & remove each FShader
-		check(0);
-	}
-
 	GetNameToTypeMap().Remove(FName(Name));
 	GlobalListLink.Unlink();
 }
@@ -1157,6 +1152,30 @@ void FShaderPipelineType::GetOutdatedTypes(TArray<FShaderType*>& OutdatedShaderT
 	}
 }
 
+const FShaderPipelineType* FShaderPipelineType::GetShaderPipelineTypeByName(FName Name)
+{
+	for (TLinkedList<FShaderPipelineType*>::TIterator It(GetTypeList()); It; It.Next())
+	{
+		const FShaderPipelineType* Type = *It;
+		if (Name == Type->GetFName())
+		{
+			return Type;
+		}
+	}
+
+	return nullptr;
+}
+
+const FSHAHash& FShaderPipelineType::GetSourceHash() const
+{
+	TArray<FString> Filenames;
+	for (const FShaderType* ShaderType : Stages)
+	{
+		Filenames.Add(ShaderType->GetShaderFilename());
+	}
+	return GetShaderFilesHash(Filenames);
+}
+
 
 FShaderPipeline::FShaderPipeline(
 	const FShaderPipelineType* InPipelineType,
@@ -1221,8 +1240,59 @@ FShaderPipeline::FShaderPipeline(const FShaderPipelineType* InPipelineType, cons
 	Validate();
 }
 
+FShaderPipeline::FShaderPipeline(const FShaderPipelineType* InPipelineType, const TArray< TRefCountPtr<FShader> >& InStages) :
+	PipelineType(InPipelineType),
+	VertexShader(nullptr),
+	HullShader(nullptr),
+	DomainShader(nullptr),
+	GeometryShader(nullptr),
+	PixelShader(nullptr)
+{
+	check(InPipelineType);
+	for (FShader* Shader : InStages)
+	{
+		if (Shader)
+		{
+			switch (Shader->GetType()->GetFrequency())
+			{
+			case SF_Vertex:
+				check(!VertexShader);
+				VertexShader = Shader;
+				break;
+			case SF_Pixel:
+				check(!PixelShader);
+				PixelShader = Shader;
+				break;
+			case SF_Hull:
+				check(!HullShader);
+				HullShader = Shader;
+				break;
+			case SF_Domain:
+				check(!DomainShader);
+				DomainShader = Shader;
+				break;
+			case SF_Geometry:
+				check(!GeometryShader);
+				GeometryShader = Shader;
+				break;
+			default:
+				checkf(0, TEXT("Invalid stage %u found!"), Shader->GetType()->GetFrequency());
+				break;
+			}
+		}
+	}
+
+	Validate();
+}
+
 FShaderPipeline::~FShaderPipeline()
 {
+	// Manually set references to nullptr, helps debugging
+	VertexShader = nullptr;
+	HullShader = nullptr;
+	DomainShader = nullptr;
+	GeometryShader = nullptr;
+	PixelShader = nullptr;
 }
 
 void FShaderPipeline::Validate()
@@ -1260,9 +1330,10 @@ void DumpShaderStats(EShaderPlatform Platform, EShaderFrequency Frequency)
 
 	// Iterate over all shader types and log stats.
 	int32 TotalShaderCount		= 0;
-	int32 TotalTypeCount			= 0;
+	int32 TotalTypeCount		= 0;
 	int32 TotalInstructionCount	= 0;
 	int32 TotalSize				= 0;
+	int32 TotalPipelineCount	= 0;
 	float TotalSizePerType		= 0;
 
 	// Write a row of headings for the table's columns.
@@ -1272,6 +1343,7 @@ void DumpShaderStats(EShaderPlatform Platform, EShaderFrequency Frequency)
 	ShaderTypeViewer.AddColumn(TEXT("Size"));
 	ShaderTypeViewer.AddColumn(TEXT("AvgSizePerInstance"));
 	ShaderTypeViewer.AddColumn(TEXT("Pipelines"));
+	ShaderTypeViewer.AddColumn(TEXT("Shared Pipelines"));
 	ShaderTypeViewer.CycleRow();
 
 	for( TLinkedList<FShaderType*>::TIterator It(FShaderType::GetTypeList()); It; It.Next() )
@@ -1285,6 +1357,7 @@ void DumpShaderStats(EShaderPlatform Platform, EShaderFrequency Frequency)
 			int32 Size						= 0;
 			int32 NumShaders					= 0;
 			int32 NumPipelines = 0;
+			int32 NumSharedPipelines = 0;
 			for (TMap<FShaderId,FShader*>::TConstIterator ShaderIt(Type->ShaderIdMap);ShaderIt;++ShaderIt)
 			{
 				const FShader* Shader = ShaderIt.Value();
@@ -1299,17 +1372,39 @@ void DumpShaderStats(EShaderPlatform Platform, EShaderFrequency Frequency)
 					continue;
 				}
 
-				if (Shader->GetShaderPipelineType())
-				{
-					++NumPipelines;
-				}
-
 				NumInitializedInstructions += Shader->GetNumInstructions();
 				Size += Shader->GetCode().Num();
 				NumShaders++;
 			}
 			AverageNumInstructions = (float)NumInitializedInstructions / (float)Type->GetNumShaders();
 			
+			for (TLinkedList<FShaderPipelineType*>::TConstIterator PipelineIt(FShaderPipelineType::GetTypeList()); PipelineIt; PipelineIt.Next())
+			{
+				const FShaderPipelineType* PipelineType = *PipelineIt;
+				bool bFound = false;
+				if (Frequency == SF_NumFrequencies)
+				{
+					if (PipelineType->GetShader(Type->GetFrequency()) == Type)
+					{
+						++NumPipelines;
+						bFound = true;
+					}
+				}
+				else
+				{
+					if (PipelineType->GetShader(Frequency) == Type)
+					{
+						++NumPipelines;
+						bFound = true;
+					}
+				}
+
+				if (!PipelineType->ShouldOptimizeUnusedOutputs() && bFound)
+				{
+					++NumSharedPipelines;
+				}
+			}
+
 			// Only add rows if there is a matching shader.
 			if( NumShaders )
 			{
@@ -1320,9 +1415,11 @@ void DumpShaderStats(EShaderPlatform Platform, EShaderFrequency Frequency)
 				ShaderTypeViewer.AddColumn(TEXT("%u"),Size);
 				ShaderTypeViewer.AddColumn(TEXT("%.1f"),Size / (float)NumShaders);
 				ShaderTypeViewer.AddColumn(TEXT("%d"), NumPipelines);
+				ShaderTypeViewer.AddColumn(TEXT("%d"), NumSharedPipelines);
 				ShaderTypeViewer.CycleRow();
 
 				TotalShaderCount += NumShaders;
+				TotalPipelineCount += NumPipelines;
 				TotalInstructionCount += NumInitializedInstructions;
 				TotalTypeCount++;
 				TotalSize += Size;
@@ -1331,12 +1428,15 @@ void DumpShaderStats(EShaderPlatform Platform, EShaderFrequency Frequency)
 		}
 	}
 
+	// go through non shared pipelines
+
 	// Write a total row.
 	ShaderTypeViewer.AddColumn(TEXT("Total"));
 	ShaderTypeViewer.AddColumn(TEXT("%u"),TotalShaderCount);
 	ShaderTypeViewer.AddColumn(TEXT("%u"),TotalInstructionCount);
 	ShaderTypeViewer.AddColumn(TEXT("%u"),TotalSize);
 	ShaderTypeViewer.AddColumn(TEXT("0"));
+	ShaderTypeViewer.AddColumn(TEXT("%u"), TotalPipelineCount);
 	ShaderTypeViewer.AddColumn(TEXT("-"));
 	ShaderTypeViewer.CycleRow();
 
@@ -1347,23 +1447,21 @@ void DumpShaderStats(EShaderPlatform Platform, EShaderFrequency Frequency)
 	ShaderTypeViewer.AddColumn(TEXT("%.1f"),TotalSize / (float)TotalShaderCount);
 	ShaderTypeViewer.AddColumn(TEXT("%.1f"),TotalSizePerType / TotalTypeCount);
 	ShaderTypeViewer.AddColumn(TEXT("-"));
+	ShaderTypeViewer.AddColumn(TEXT("-"));
 	ShaderTypeViewer.CycleRow();
 #endif
 }
 
 
-FShaderType* FindShaderTypeByName(const TCHAR* ShaderTypeName)
+FShaderType* FindShaderTypeByName(FName ShaderTypeName)
 {
-	for(TLinkedList<FShaderType*>::TIterator ShaderTypeIt(FShaderType::GetTypeList());ShaderTypeIt;ShaderTypeIt.Next())
+	FShaderType** FoundShader = FShaderType::GetNameToTypeMap().Find(ShaderTypeName);
+	if (FoundShader)
 	{
-		const TCHAR* LocalName = ShaderTypeIt->GetName();
-
-		if(!FCString::Stricmp(LocalName, ShaderTypeName))
-		{
-			return *ShaderTypeIt;
-		}
+		return *FoundShader;
 	}
-	return NULL;
+
+	return nullptr;
 }
 
 
diff --git a/Engine/Source/Runtime/ShaderCore/Public/Shader.h b/Engine/Source/Runtime/ShaderCore/Public/Shader.h
--- a/Engine/Source/Runtime/ShaderCore/Public/Shader.h
+++ b/Engine/Source/Runtime/ShaderCore/Public/Shader.h
@@ -429,6 +429,7 @@ public:
 	 */
 	FString VertexFactoryTypeName;
 
+	// Required to differentiate amongst unique shaders in the global map per Type
 	FString ShaderPipelineName;
 
 	/** Used to detect changes to the vertex factory source files. */
@@ -577,7 +578,6 @@ public:
 	inline FSHAHash GetOutputHash() const { return OutputHash; }
 	FShaderId GetId() const;
 	inline FVertexFactoryType* GetVertexFactoryType() const { return VFType; }
-	inline const FShaderPipelineType* GetShaderPipelineType() const { return ShaderPipeline; }
 	inline FSHAHash GetMaterialShaderMapHash() const { return MaterialShaderMapHash; }
 	inline int32 GetNumRefs() const { return NumRefs; }
 
@@ -834,6 +834,10 @@ public:
 	{ 
 		return Name; 
 	}
+	const FName& GetFName() const
+	{
+		return TypeName;
+	}
 	const TCHAR* GetShaderFilename() const 
 	{ 
 		return SourceFilename; 
@@ -904,9 +908,9 @@ public:
 	}
 
 private:
-
 	uint32 HashIndex;
 	const TCHAR* Name;
+	FName TypeName;
 	const TCHAR* SourceFilename;
 	const TCHAR* FunctionName;
 	uint32 Frequency;
@@ -1012,7 +1016,8 @@ private:
 class SHADERCORE_API FShaderPipelineType
 {
 public:
-	// Set bInShouldShadersToIndividualMap to false if you want to skip adding the FShaderTypes from the shadermap individual shader type list
+	// Set bShouldOptimizeUnusedOutputs to true if we want unique FShaders for each shader pipeline
+	// Set bShouldOptimizeUnusedOutputs to false if the FShaders will point to the individual shaders in the map
 	FShaderPipelineType(
 		const TCHAR* InName,
 		const FShaderType* InVertexShader,
@@ -1020,8 +1025,8 @@ public:
 		const FShaderType* InDomainShader,
 		const FShaderType* InGeometryShader,
 		const FShaderType* InPixelShader,
-		bool bInShouldAddShadersToIndividualMap = true);
-	virtual ~FShaderPipelineType();
+		bool bInShouldOptimizeUnusedOutputs);
+	~FShaderPipelineType();
 
 	FORCEINLINE bool HasTessellation() const { return AllStages[SF_Domain] != nullptr; }
 	FORCEINLINE bool HasGeometry() const { return AllStages[SF_Geometry] != nullptr; }
@@ -1043,6 +1048,7 @@ public:
 
 	/** @return The global shader pipeline name to type map */
 	static TMap<FName, FShaderPipelineType*>& GetNameToTypeMap();
+	static const FShaderPipelineType* GetShaderPipelineTypeByName(FName Name);
 
 	/** Initialize static members, this must be called before any shader types are created. */
 	static void Initialize();
@@ -1061,11 +1067,14 @@ public:
 	bool IsGlobalTypePipeline() const { return Stages[0]->GetGlobalShaderType() != nullptr; }
 	bool IsMaterialTypePipeline() const { return Stages[0]->GetMaterialShaderType() != nullptr; }
 	bool IsMeshMaterialTypePipeline() const { return Stages[0]->GetMeshMaterialShaderType() != nullptr; }
-	bool ShouldAddShadersToIndividualMap() const { return bAddShadersToIndividualMap; }
+	bool ShouldOptimizeUnusedOutputs() const { return bShouldOptimizeUnusedOutputs; }
 
 	/** Gets a list of FShaderTypes & PipelineTypes whose source file no longer matches what that type was compiled with */
 	static void GetOutdatedTypes(TArray<FShaderType*>& OutdatedShaderTypes, TArray<const FShaderPipelineType*>& ShaderPipelineTypesToFlush, TArray<const FVertexFactoryType*>& OutdatedFactoryTypes);
 
+	/** Calculates a Hash based on this shader pipeline type stages' source code and includes */
+	const FSHAHash& GetSourceHash() const;
+
 protected:
 	const TCHAR* const Name;
 	FName TypeName;
@@ -1078,11 +1087,37 @@ protected:
 	TLinkedList<FShaderPipelineType*> GlobalListLink;
 
 	uint32 HashIndex;
-	bool bAddShadersToIndividualMap;
+	bool bShouldOptimizeUnusedOutputs;
 
 	static bool bInitialized;
 };
 
+#if !UE_BUILD_DOCS
+// Vertex+Pixel
+#define IMPLEMENT_SHADERPIPELINE_TYPE_VSPS(PipelineName, VertexShaderType, PixelShaderType, bRemoveUnused)	\
+	static FShaderPipelineType PipelineName(TEXT(PREPROCESSOR_TO_STRING(PipelineName)), &VertexShaderType::StaticType, nullptr, nullptr, nullptr, &PixelShaderType::StaticType, bRemoveUnused);
+// Only VS
+#define IMPLEMENT_SHADERPIPELINE_TYPE_VS(PipelineName, VertexShaderType, bRemoveUnused)	\
+	static FShaderPipelineType PipelineName(TEXT(PREPROCESSOR_TO_STRING(PipelineName)), &VertexShaderType::StaticType, nullptr, nullptr, nullptr, nullptr, bRemoveUnused);
+// Vertex+Geometry+Pixel
+#define IMPLEMENT_SHADERPIPELINE_TYPE_VSGSPS(PipelineName, VertexShaderType, GeometryShaderType, PixelShaderType, bRemoveUnused)	\
+	static FShaderPipelineType PipelineName(TEXT(PREPROCESSOR_TO_STRING(PipelineName)), &VertexShaderType::StaticType, nullptr, nullptr, &GeometryShaderType::StaticType, &PixelShaderType::StaticType, bRemoveUnused);
+// Vertex+Geometry
+#define IMPLEMENT_SHADERPIPELINE_TYPE_VSGS(PipelineName, VertexShaderType, GeometryShaderType, bRemoveUnused)	\
+	static FShaderPipelineType PipelineName(TEXT(PREPROCESSOR_TO_STRING(PipelineName)), &VertexShaderType::StaticType, nullptr, nullptr, &GeometryShaderType::StaticType, nullptr, bRemoveUnused);
+// Vertex+Hull+Domain+Pixel
+#define IMPLEMENT_SHADERPIPELINE_TYPE_VSHSDSPS(PipelineName, VertexShaderType, HullShaderType, DomainShaderType, PixelShaderType, bRemoveUnused)	\
+	static FShaderPipelineType PipelineName(TEXT(PREPROCESSOR_TO_STRING(PipelineName)), &VertexShaderType::StaticType, &HullShaderType::StaticType, &DomainShaderType::StaticType, nullptr, &PixelShaderType::StaticType, bRemoveUnused);
+// Vertex+Hull+Domain+Geometry+Pixel
+#define IMPLEMENT_SHADERPIPELINE_TYPE_VSHSDSGSPS(PipelineName, VertexShaderType, HullShaderType, DomainShaderType, GeometryShaderType, PixelShaderType, bRemoveUnused)	\
+	static FShaderPipelineType PipelineName(TEXT(PREPROCESSOR_TO_STRING(PipelineName)), &VertexShaderType::StaticType, &HullShaderType::StaticType, &DomainShaderType::StaticType, &GeometryShaderType::StaticType, &PixelShaderType::StaticType, bRemoveUnused);
+// Vertex+Hull+Domain
+#define IMPLEMENT_SHADERPIPELINE_TYPE_VSHSDS(PipelineName, VertexShaderType, HullShaderType, DomainShaderType, bRemoveUnused)	\
+	static FShaderPipelineType PipelineName(TEXT(PREPROCESSOR_TO_STRING(PipelineName)), &VertexShaderType::StaticType, &HullShaderType::StaticType, &DomainShaderType::StaticType, nullptr, nullptr, bRemoveUnused);
+// Vertex+Hull+Domain+Geometry
+#define IMPLEMENT_SHADERPIPELINE_TYPE_VSHSDSGS(PipelineName, VertexShaderType, HullShaderType, DomainShaderType, GeometryShaderType, bRemoveUnused)	\
+	static FShaderPipelineType PipelineName(TEXT(PREPROCESSOR_TO_STRING(PipelineName)), &VertexShaderType::StaticType, &HullShaderType::StaticType, &DomainShaderType::StaticType, &GeometryShaderType::StaticType, nullptr, bRemoveUnused);
+#endif
 
 /** Encapsulates a dependency on a shader type and saved state from that shader type. */
 class FShaderTypeDependency
@@ -1090,36 +1125,55 @@ class FShaderTypeDependency
 public:
 
 	FShaderTypeDependency() :
-		ShaderType(nullptr),
-		ShaderPipelineType(nullptr)
+		ShaderType(nullptr)
 	{}
 
 	/** Shader type */
 	FShaderType* ShaderType;
 
-	const FShaderPipelineType* ShaderPipelineType;
-
 	/** Used to detect changes to the shader source files. */
 	FSHAHash SourceHash;
 
 	friend FArchive& operator<<(FArchive& Ar,class FShaderTypeDependency& Ref)
 	{
 		Ar << Ref.ShaderType;
-		if (Ar.UE4Ver() >= VER_UE4_PURGED_FMATERIAL_COMPILE_OUTPUTS)
-		{
-			Ar << Ref.ShaderPipelineType;
-		}
 		Ar << Ref.SourceHash;
 		return Ar;
 	}
 
 	bool operator==(const FShaderTypeDependency& Reference) const
 	{
-		return ShaderType == Reference.ShaderType && ShaderPipelineType == Reference.ShaderPipelineType && SourceHash == Reference.SourceHash;
+		return ShaderType == Reference.ShaderType && SourceHash == Reference.SourceHash;
 	}
 };
 
 
+class FShaderPipelineTypeDependency
+{
+public:
+	FShaderPipelineTypeDependency() :
+		ShaderPipelineType(nullptr)
+	{}
+
+	/** Shader Pipeline type */
+	const FShaderPipelineType* ShaderPipelineType;
+
+	/** Used to detect changes to the shader source files. */
+	FSHAHash StagesSourceHash;
+
+	friend FArchive& operator<<(FArchive& Ar, class FShaderPipelineTypeDependency& Ref)
+	{
+		Ar << Ref.ShaderPipelineType;
+		Ar << Ref.StagesSourceHash;
+		return Ar;
+	}
+
+	bool operator==(const FShaderPipelineTypeDependency& Reference) const
+	{
+		return ShaderPipelineType == Reference.ShaderPipelineType && StagesSourceHash == Reference.StagesSourceHash;
+	}
+};
+
 /** Used to compare two shader types by name. */
 class FCompareShaderTypes
 {																				
@@ -1170,49 +1224,8 @@ public:
 	}
 };
 
-/** Used to compare two shader pipeline types by name. */
-class FCompareShaderPipelineTypes
-{
-public:
-	FORCEINLINE bool operator()(const FShaderPipelineType& A, const FShaderPipelineType& B) const
-	{
-		int32 AL = FCString::Strlen(A.GetName());
-		int32 BL = FCString::Strlen(B.GetName());
-		if (AL == BL)
-		{
-			int32 Result = FCString::Strncmp(A.GetName(), B.GetName(), AL);
-			if (Result != 0)
-			{
-				return Result > 0;
-			}
-
-			const auto& PipeA = A.GetStages();
-			const auto& PipeB = B.GetStages();
-
-			if (PipeA.Num() != PipeB.Num())
-			{
-				return PipeA.Num() > PipeB.Num();
-			}
-
-			bool bResult = false;
-			FCompareShaderTypes Comparator;
-			for (int32 Index = 0; Index < PipeA.Num(); ++Index)
-			{
-				bResult = Comparator(*PipeA[Index], *PipeB[Index]);
-				if (bResult != 0)
-				{
-					break;
-				}
-			}
-
-			return bResult;
-		}
-		return AL > BL;
-	}
-};
-
 // A Shader Pipeline instance with compiled stages
-class SHADERCORE_API FShaderPipeline : public FRefCountedObject
+class SHADERCORE_API FShaderPipeline
 {
 public:
 	const FShaderPipelineType* PipelineType;
@@ -1231,8 +1244,9 @@ public:
 		FShader* InPixelShader);
 
 	FShaderPipeline(const FShaderPipelineType* InPipelineType, const TArray<FShader*>& InStages);
+	FShaderPipeline(const FShaderPipelineType* InPipelineType, const TArray< TRefCountPtr<FShader> >& InStages);
 
-	virtual ~FShaderPipeline();
+	~FShaderPipeline();
 
 	// Find a shader inside the pipeline
 	template<typename ShaderType>
@@ -1324,8 +1338,20 @@ public:
 	}
 
 	void Validate();
+
+	enum EFilter
+	{
+		EAll,			// All pipelines
+		EOnlyShared,	// Only pipelines with shared shaders
+		EOnlyUnique,	// Only pipelines with unique shaders
+	};
 };
 
+inline bool operator<(const FShaderPipeline& Lhs, const FShaderPipeline& Rhs)
+{
+	FCompareShaderPipelineNameTypes Comparator;
+	return Comparator(*Lhs.PipelineType, *Rhs.PipelineType);
+}
 
 /** A collection of shaders of different types, but the same meta type. */
 template<typename ShaderMetaType>
@@ -1381,9 +1407,13 @@ public:
 
 	void RemoveShaderPipelineType(const FShaderPipelineType* ShaderPipelineType)
 	{
-		auto* Found = ShaderPipelines.Find(ShaderPipelineType);
+		FShaderPipeline** Found = ShaderPipelines.Find(ShaderPipelineType);
 		if (Found)
 		{
+			if (*Found)
+			{
+				delete *Found;
+			}
 			ShaderPipelines.Remove(ShaderPipelineType);
 		}
 	}
@@ -1401,11 +1431,20 @@ public:
 	}
 
 	/** Builds a list of the shader pipelines in a shader map. */
-	void GetShaderPipelineList(TArray<FShaderPipeline*>& OutShaderPipelines) const
+	void GetShaderPipelineList(TArray<FShaderPipeline*>& OutShaderPipelines, FShaderPipeline::EFilter Filter) const
 	{
 		for (auto Pair : ShaderPipelines)
 		{
-			OutShaderPipelines.Add(Pair.Value);
+			FShaderPipeline* Pipeline = Pair.Value;
+			if (Pipeline->PipelineType->ShouldOptimizeUnusedOutputs() && Filter == FShaderPipeline::EOnlyShared)
+			{
+				continue;
+			}
+			else if (!Pipeline->PipelineType->ShouldOptimizeUnusedOutputs() && Filter == FShaderPipeline::EOnlyUnique)
+			{
+				continue;
+			}
+			OutShaderPipelines.Add(Pipeline);
 		}
 	}
 
@@ -1535,14 +1574,14 @@ public:
 				SerializeShaderForSaving(CurrentShader, Ar, bHandleShaderKeyChanges, bInlineShaderResource);
 			}
 
-			int32 NumPipelines = ShaderPipelines.Num();
+			TArray<FShaderPipeline*> SortedPipelines;
+			GetShaderPipelineList(SortedPipelines, FShaderPipeline::EAll);
+			int32 NumPipelines = SortedPipelines.Num();
 			Ar << NumPipelines;
 			// Sort the shader pipelines by type name before saving, to make sure the saved result is binary equivalent to what is generated on other machines, Which is a requirement of the Derived Data Cache.
-			auto SortedPipelines = ShaderPipelines;
-			SortedPipelines.KeySort(FCompareShaderPipelineNameTypes());
-			for (TMap<const FShaderPipelineType*, TRefCountPtr<FShaderPipeline> >::TConstIterator PipelineIt(SortedPipelines); PipelineIt; ++PipelineIt)
+			SortedPipelines.Sort();
+			for (FShaderPipeline* CurrentPipeline : SortedPipelines)
 			{
-				FShaderPipeline* CurrentPipeline = PipelineIt.Value();
 				const FShaderPipelineType* PipelineType = CurrentPipeline->PipelineType;
 				Ar << PipelineType;
 
@@ -1584,7 +1623,8 @@ public:
 				Ar << ShaderPipelineType;
 				int32 NumStages = 0;
 				Ar << NumStages;
-				TArray<FShader*> ShaderStages;
+				// Make a list of references so they can be deleted when going out of scope if needed
+				TArray< TRefCountPtr<FShader> > ShaderStages;
 				for (int32 Index = 0; Index < NumStages; ++Index)
 				{
 					FShaderType* Type = nullptr;
@@ -1596,37 +1636,30 @@ public:
 					}
 				}
 
-				if (ShaderStages.Num() == ShaderPipelineType->GetStages().Num())
+				// ShaderPipelineType can be nullptr if the pipeline existed but now is gone!
+				if (ShaderPipelineType && ShaderStages.Num() == ShaderPipelineType->GetStages().Num())
 				{
 					FShaderPipeline* ShaderPipeline = new FShaderPipeline(ShaderPipelineType, ShaderStages);
 					AddShaderPipeline(ShaderPipelineType, ShaderPipeline);
 				}
-				else
-				{
-					//#todo-rco: Log that the # of stages changed for this pipeline
-					for (FShader* Shader : ShaderStages)
-					{
-						delete Shader;
-					}
-				}
 			}
 		}
 	}
 
 	/** @return true if the map is empty */
-	bool IsEmpty() const
+	inline bool IsEmpty() const
 	{
 		return Shaders.Num() == 0;
 	}
 
 	/** @return The number of shaders in the map. */
-	uint32 GetNumShaders() const
+	inline uint32 GetNumShaders() const
 	{
 		return Shaders.Num();
 	}
 
 	/** @return The number of shader pipelines in the map. */
-	uint32 GetNumShaderPipelines() const
+	inline uint32 GetNumShaderPipelines() const
 	{
 		return ShaderPipelines.Num();
 	}
@@ -1638,14 +1671,25 @@ public:
 		ShaderPipelines.Empty();
 	}
 
-	// Returns nullptr if not found
-	FShaderPipeline* GetShaderPipeline(const FShaderPipelineType* PipelineType) const
+	inline FShaderPipeline* GetShaderPipeline(const FShaderPipelineType* PipelineType)
+	{
+		FShaderPipeline** Found = ShaderPipelines.Find(PipelineType);
+		return Found ? *Found : nullptr;
+	}
+
+	inline FShaderPipeline* GetShaderPipeline(const FShaderPipelineType* PipelineType) const
 	{
-		auto* Found = ShaderPipelines.Find(PipelineType);
+		FShaderPipeline* const* Found = ShaderPipelines.Find(PipelineType);
 		return Found ? *Found : nullptr;
 	}
 
-	void AddShaderPipeline(const FShaderPipelineType* Type, FShaderPipeline* ShaderPipeline)
+	// Returns nullptr if not found
+	inline bool HasShaderPipeline(const FShaderPipelineType* PipelineType) const
+	{
+		return (GetShaderPipeline(PipelineType) != nullptr);
+	}
+
+	inline void AddShaderPipeline(const FShaderPipelineType* Type, FShaderPipeline* ShaderPipeline)
 	{
 		check(Type);
 		check(!ShaderPipeline || ShaderPipeline->PipelineType == Type);
@@ -1676,7 +1720,7 @@ public:
 
 protected:
 	TMap<FShaderType*, TRefCountPtr<FShader> > Shaders;
-	TMap<const FShaderPipelineType*, TRefCountPtr<FShaderPipeline> > ShaderPipelines;
+	TMap<const FShaderPipelineType*, FShaderPipeline*> ShaderPipelines;
 };
 
 /** A reference which is initialized with the requested shader type from a shader map. */
@@ -1905,7 +1949,7 @@ extern SHADERCORE_API void DumpShaderStats( EShaderPlatform Platform, EShaderFre
  * @param ShaderTypeName - The name of the shader type to find.
  * @return The shader type, or NULL if none matched.
  */
-extern SHADERCORE_API FShaderType* FindShaderTypeByName(const TCHAR* ShaderTypeName);
+extern SHADERCORE_API FShaderType* FindShaderTypeByName(FName ShaderTypeName);
 
 /** Helper function to dispatch a compute shader while checking that parameters have been set correctly. */
 extern SHADERCORE_API void DispatchComputeShader(
diff --git a/Engine/Source/Runtime/ShaderCore/Public/ShaderCore.h b/Engine/Source/Runtime/ShaderCore/Public/ShaderCore.h
--- a/Engine/Source/Runtime/ShaderCore/Public/ShaderCore.h
+++ b/Engine/Source/Runtime/ShaderCore/Public/ShaderCore.h
@@ -111,7 +111,9 @@ enum ECompilerFlags
 	CFLAG_StandardOptimization,
 	/** Shader should use on chip memory instead of main memory ring buffer memory. */
 	CFLAG_OnChip,
-	CFLAG_KeepDebugInfo
+	CFLAG_KeepDebugInfo,
+	// Skip using the cached usf file system and directly open the file (for debugging)
+	CFLAG_OpenFileFromFullPath,
 };
 
 /**
