Prevent LiveCodingServer module building for targets by default.

#rnx

diff --git a/Engine/Source/Developer/Windows/LiveCodingServer/LiveCodingServer.Build.cs b/Engine/Source/Developer/Windows/LiveCodingServer/LiveCodingServer.Build.cs
--- a/Engine/Source/Developer/Windows/LiveCodingServer/LiveCodingServer.Build.cs
+++ b/Engine/Source/Developer/Windows/LiveCodingServer/LiveCodingServer.Build.cs
@@ -23,5 +23,7 @@ public class LiveCodingServer : ModuleRules
 		PrivateIncludePaths.Add(Path.Combine(DiaSdkDir, "include"));
 		PublicAdditionalLibraries.Add(Path.Combine(DiaSdkDir, "lib", "amd64", "diaguids.lib"));
 		RuntimeDependencies.Add("$(TargetOutputDir)/msdia140.dll", Path.Combine(DiaSdkDir, "bin", "amd64", "msdia140.dll"));
+
+		PrecompileForTargets = PrecompileTargetsType.None;
 	}
 }
