Fixes for static analysis warnings.

diff --git a/Engine/Source/Developer/Windows/LiveCoding/Private/External/LC_ClientCommandActions.cpp b/Engine/Source/Developer/Windows/LiveCoding/Private/External/LC_ClientCommandActions.cpp
--- a/Engine/Source/Developer/Windows/LiveCoding/Private/External/LC_ClientCommandActions.cpp
+++ b/Engine/Source/Developer/Windows/LiveCoding/Private/External/LC_ClientCommandActions.cpp
@@ -139,14 +139,17 @@ bool actions::LoadPatch::Execute(CommandType* command, const DuplexPipe* pipe, v
 	HMODULE module = ::LoadLibraryW(command->path);
 
 	// BEGIN EPIC MOD - Support for UE4 debug visualizers
-	typedef void InitNatvisHelpersFunc(FNameEntry*** NameTable, FChunkedFixedUObjectArray* ObjectArray);
-
-	InitNatvisHelpersFunc* InitNatvisHelpers = (InitNatvisHelpersFunc*)(void*)GetProcAddress(module, "InitNatvisHelpers");
-	if (InitNatvisHelpers != nullptr)
+	if (module != nullptr)
 	{
-		extern FNameEntry*** GFNameTableForDebuggerVisualizers_MT;
-		extern FChunkedFixedUObjectArray*& GObjectArrayForDebugVisualizers;
-		(*InitNatvisHelpers)(GFNameTableForDebuggerVisualizers_MT, GObjectArrayForDebugVisualizers);
+		typedef void InitNatvisHelpersFunc(FNameEntry*** NameTable, FChunkedFixedUObjectArray* ObjectArray);
+
+		InitNatvisHelpersFunc* InitNatvisHelpers = (InitNatvisHelpersFunc*)(void*)GetProcAddress(module, "InitNatvisHelpers");
+		if (InitNatvisHelpers != nullptr)
+		{
+			extern FNameEntry*** GFNameTableForDebuggerVisualizers_MT;
+			extern FChunkedFixedUObjectArray*& GObjectArrayForDebugVisualizers;
+			(*InitNatvisHelpers)(GFNameTableForDebuggerVisualizers_MT, GObjectArrayForDebugVisualizers);
+		}
 	}
 	// END EPIC MOD
 
diff --git a/Engine/Source/Developer/Windows/LiveCoding/Private/External/LC_ClientStartupThread.cpp b/Engine/Source/Developer/Windows/LiveCoding/Private/External/LC_ClientStartupThread.cpp
--- a/Engine/Source/Developer/Windows/LiveCoding/Private/External/LC_ClientStartupThread.cpp
+++ b/Engine/Source/Developer/Windows/LiveCoding/Private/External/LC_ClientStartupThread.cpp
@@ -343,11 +343,11 @@ unsigned int ClientStartupThread::ThreadFunction(const std::wstring& processGrou
 			commandLine += processGroupName;
 
 			// BEGIN EPIC MOD - Additional arguments for console app
-			extern const wchar_t* GLiveCodingConsoleArguments;
-			if(GLiveCodingConsoleArguments != nullptr)
+			extern FString GLiveCodingConsoleArguments;
+			if(GLiveCodingConsoleArguments.Len() > 0)
 			{
 				commandLine += L" ";
-				commandLine += GLiveCodingConsoleArguments;
+				commandLine += *GLiveCodingConsoleArguments;
 			}
 			if (!FApp::IsProjectNameEmpty())
 			{
diff --git a/Engine/Source/Developer/Windows/LiveCoding/Private/External/LC_Process.cpp b/Engine/Source/Developer/Windows/LiveCoding/Private/External/LC_Process.cpp
--- a/Engine/Source/Developer/Windows/LiveCoding/Private/External/LC_Process.cpp
+++ b/Engine/Source/Developer/Windows/LiveCoding/Private/External/LC_Process.cpp
@@ -7,6 +7,12 @@
 #include "LC_Logging.h"
 #include <Psapi.h>
 
+// BEGIN EPIC MODS
+#pragma warning(push)
+#pragma warning(disable:6011) // warning C6011: Dereferencing NULL pointer 'processInfo'.
+#pragma warning(disable:6335) // warning C6335: Leaking process information handle 'context->pi.hProcess'.
+// END EPIC MODS
+
 // internal ntdll.dll definitions
 typedef LONG NTSTATUS;
 typedef LONG KPRIORITY;
@@ -837,3 +843,7 @@ namespace process
 		delete[] memory;
 	}
 }
+
+// BEGIN EPIC MODS
+#pragma warning(pop)
+// END EPIC MODS
diff --git a/Engine/Source/Developer/Windows/LiveCoding/Private/External/LC_Thread.cpp b/Engine/Source/Developer/Windows/LiveCoding/Private/External/LC_Thread.cpp
--- a/Engine/Source/Developer/Windows/LiveCoding/Private/External/LC_Thread.cpp
+++ b/Engine/Source/Developer/Windows/LiveCoding/Private/External/LC_Thread.cpp
@@ -5,6 +5,11 @@
 #include "LC_Logging.h"
 #include <process.h>
 
+// BEGIN EPIC MODS
+#pragma warning(push)
+#pragma warning(disable:6322) // warning C6322: Empty _except block.
+#pragma warning(disable:6258) // warning C6258: Using TerminateThread does not allow proper thread clean up.
+// END EPIC MODS
 
 namespace
 {
@@ -175,3 +180,7 @@ namespace thread
 		SetThreadName(name);
 	}
 }
+
+// BEGIN EPIC MODS
+#pragma warning(pop)
+// END EPIC MODS
diff --git a/Engine/Source/Developer/Windows/LiveCoding/Private/External/xxhash.c b/Engine/Source/Developer/Windows/LiveCoding/Private/External/xxhash.c
--- a/Engine/Source/Developer/Windows/LiveCoding/Private/External/xxhash.c
+++ b/Engine/Source/Developer/Windows/LiveCoding/Private/External/xxhash.c
@@ -32,6 +32,11 @@
 *  - xxHash source repository : [REDACTED_PROFILE_URL]/xxHash
 */
 
+// BEGIN EPIC MODS
+#pragma warning(push)
+#pragma warning(disable:6385) // warning C6385: Reading invalid data from '&state'
+// END EPIC MODS
+
 // MM: tuning
 #pragma warning(disable : 4711)
 #pragma warning(disable : 4820)
@@ -892,3 +897,7 @@ XXH_PUBLIC_API XXH64_hash_t XXH64_hashFromCanonical(const XXH64_canonical_t* src
 }
 
 #endif  /* XXH_NO_LONG_LONG */
+
+// BEGIN EPIC MODS
+#pragma warning(pop)
+// END EPIC MODS
diff --git a/Engine/Source/Developer/Windows/LiveCoding/Private/LiveCodingModule.cpp b/Engine/Source/Developer/Windows/LiveCoding/Private/LiveCodingModule.cpp
--- a/Engine/Source/Developer/Windows/LiveCoding/Private/LiveCodingModule.cpp
+++ b/Engine/Source/Developer/Windows/LiveCoding/Private/LiveCodingModule.cpp
@@ -19,7 +19,7 @@ IMPLEMENT_MODULE(FLiveCodingModule, LiveCoding)
 #define LOCTEXT_NAMESPACE "LiveCodingModule"
 
 FString GLiveCodingConsolePath;
-const wchar_t* GLiveCodingConsoleArguments = nullptr;
+FString GLiveCodingConsoleArguments;
 
 FLiveCodingModule::FLiveCodingModule()
 	: bEnabled(false)
diff --git a/Engine/Source/Developer/Windows/LiveCodingServer/Private/External/LC_LiveModule.cpp b/Engine/Source/Developer/Windows/LiveCodingServer/Private/External/LC_LiveModule.cpp
--- a/Engine/Source/Developer/Windows/LiveCodingServer/Private/External/LC_LiveModule.cpp
+++ b/Engine/Source/Developer/Windows/LiveCodingServer/Private/External/LC_LiveModule.cpp
@@ -631,16 +631,16 @@ namespace
 		// the code to inject on x64 is:
 		//		B0 01		mov al, 1
 		//		C3			ret				different calling convention than x86
-		const uint8_t PATCH[3u] = { 0xB0, 0x01, 0xC3 };
+		const uint8_t PatchData[3u] = { 0xB0, 0x01, 0xC3 };
 #else
 		// the code to inject on x86 is:
 		//		B0 01		mov al, 1
 		//		C2 0C 00	ret 0Ch			different calling convention than x64
-		const uint8_t PATCH[5u] = { 0xB0, 0x01, 0xC2, 0x0C, 0x00 };
+		const uint8_t PatchData[5u] = { 0xB0, 0x01, 0xC2, 0x0C, 0x00 };
 #endif
 
 		uint8_t* address = pointer::Offset<uint8_t*>(patchBase, dllMainRva);
-		process::WriteProcessMemory(processHandle, address, PATCH, sizeof(PATCH));
+		process::WriteProcessMemory(processHandle, address, PatchData, sizeof(PatchData));
 	}
 
 
diff --git a/Engine/Source/Developer/Windows/LiveCodingServer/Private/External/LC_Scheduler.cpp b/Engine/Source/Developer/Windows/LiveCodingServer/Private/External/LC_Scheduler.cpp
--- a/Engine/Source/Developer/Windows/LiveCodingServer/Private/External/LC_Scheduler.cpp
+++ b/Engine/Source/Developer/Windows/LiveCodingServer/Private/External/LC_Scheduler.cpp
@@ -5,6 +5,11 @@
 #include "LC_Thread.h"
 
 
+// BEGIN EPIC MODS
+#pragma warning(push)
+#pragma warning(disable:6011) // warning C6011: Dereferencing NULL pointer 'info'. 
+// END EPIC MODS
+
 namespace
 {
 	scheduler::TaskQueue* g_taskQueue = nullptr;
@@ -147,3 +152,7 @@ void scheduler::WaitForTask(TaskBase* task)
 		}
 	}
 }
+
+// BEGIN EPIC MODS
+#pragma warning(pop)
+// END EPIC MODS
diff --git a/Engine/Source/Developer/Windows/LiveCodingServer/Private/External/LC_ServerCommandThread.cpp b/Engine/Source/Developer/Windows/LiveCodingServer/Private/External/LC_ServerCommandThread.cpp
--- a/Engine/Source/Developer/Windows/LiveCodingServer/Private/External/LC_ServerCommandThread.cpp
+++ b/Engine/Source/Developer/Windows/LiveCodingServer/Private/External/LC_ServerCommandThread.cpp
@@ -27,6 +27,10 @@
 // unreachable code
 #pragma warning (disable : 4702)
 
+// BEGIN EPIC MODS
+#pragma warning(push)
+#pragma warning(disable:6031) // warning C6031: Return value ignored: 'CoInitialize'.
+// END EPIC MODS
 
 namespace
 {
@@ -1796,3 +1800,7 @@ bool ServerCommandThread::ApplySettingStringAction::Execute(CommandType* command
 
 	return true;
 }
+
+// BEGIN EPIC MODS
+#pragma warning(pop)
+// END EPIC MODS
