Fix copyright notices.

#rnx

diff --git a/Engine/Source/Developer/Windows/LiveCoding/Public/ILiveCodingModule.h b/Engine/Source/Developer/Windows/LiveCoding/Public/ILiveCodingModule.h
--- a/Engine/Source/Developer/Windows/LiveCoding/Public/ILiveCodingModule.h
+++ b/Engine/Source/Developer/Windows/LiveCoding/Public/ILiveCodingModule.h
@@ -1,3 +1,5 @@
+// Copyright 1998-2019 Epic Games, Inc. All Rights Reserved.
+
 #pragma once
 
 #include "CoreTypes.h"
diff --git a/Engine/Source/Developer/Windows/LiveCodingServer/Private/LiveCodingServer.cpp b/Engine/Source/Developer/Windows/LiveCodingServer/Private/LiveCodingServer.cpp
--- a/Engine/Source/Developer/Windows/LiveCodingServer/Private/LiveCodingServer.cpp
+++ b/Engine/Source/Developer/Windows/LiveCodingServer/Private/LiveCodingServer.cpp
@@ -1,3 +1,5 @@
+// Copyright 1998-2019 Epic Games, Inc. All Rights Reserved.
+
 #include "LiveCodingServer.h"
 #include "External/LC_Scheduler.h"
 #include "External/LC_UniqueId.h"
diff --git a/Engine/Source/Developer/Windows/LiveCodingServer/Private/LiveCodingServer.h b/Engine/Source/Developer/Windows/LiveCodingServer/Private/LiveCodingServer.h
--- a/Engine/Source/Developer/Windows/LiveCodingServer/Private/LiveCodingServer.h
+++ b/Engine/Source/Developer/Windows/LiveCodingServer/Private/LiveCodingServer.h
@@ -1,3 +1,5 @@
+// Copyright 1998-2019 Epic Games, Inc. All Rights Reserved.
+
 #pragma once
 
 #include "CoreTypes.h"
diff --git a/Engine/Source/Developer/Windows/LiveCodingServer/Private/LiveCodingServerModule.cpp b/Engine/Source/Developer/Windows/LiveCodingServer/Private/LiveCodingServerModule.cpp
--- a/Engine/Source/Developer/Windows/LiveCodingServer/Private/LiveCodingServerModule.cpp
+++ b/Engine/Source/Developer/Windows/LiveCodingServer/Private/LiveCodingServerModule.cpp
@@ -1,3 +1,5 @@
+// Copyright 1998-2019 Epic Games, Inc. All Rights Reserved.
+
 #include "LiveCodingServerModule.h"
 #include "LiveCodingServer.h"
 #include "Features/IModularFeatures.h"
diff --git a/Engine/Source/Developer/Windows/LiveCodingServer/Private/LiveCodingServerModule.h b/Engine/Source/Developer/Windows/LiveCodingServer/Private/LiveCodingServerModule.h
--- a/Engine/Source/Developer/Windows/LiveCodingServer/Private/LiveCodingServerModule.h
+++ b/Engine/Source/Developer/Windows/LiveCodingServer/Private/LiveCodingServerModule.h
@@ -1,3 +1,5 @@
+// Copyright 1998-2019 Epic Games, Inc. All Rights Reserved.
+
 #pragma once
 
 #include "CoreTypes.h"
diff --git a/Engine/Source/Developer/Windows/LiveCodingServer/Public/ILiveCodingServer.h b/Engine/Source/Developer/Windows/LiveCodingServer/Public/ILiveCodingServer.h
--- a/Engine/Source/Developer/Windows/LiveCodingServer/Public/ILiveCodingServer.h
+++ b/Engine/Source/Developer/Windows/LiveCodingServer/Public/ILiveCodingServer.h
@@ -1,3 +1,5 @@
+// Copyright 1998-2019 Epic Games, Inc. All Rights Reserved.
+
 #pragma once
 
 #include "CoreTypes.h"
