Suppress warning in static analysis.

#rnx

diff --git a/Engine/Source/Developer/Windows/LiveCoding/Private/External/xxhash.c b/Engine/Source/Developer/Windows/LiveCoding/Private/External/xxhash.c
--- a/Engine/Source/Developer/Windows/LiveCoding/Private/External/xxhash.c
+++ b/Engine/Source/Developer/Windows/LiveCoding/Private/External/xxhash.c
@@ -376,7 +376,7 @@ XXH_PUBLIC_API void XXH32_copyState(XXH32_state_t* dstState, const XXH32_state_t
 XXH_PUBLIC_API XXH_errorcode XXH32_reset(XXH32_state_t* statePtr, unsigned int seed)
 {
     XXH32_state_t state;   /* using a local state to memcpy() in order to avoid strict-aliasing warnings */
-    memset(&state, 0, sizeof(state)-4);   /* do not write into reserved, for future removal */
+    memset(&state, 0, sizeof(state)-4);   /* do not write into reserved, for future removal */ //-V512
     state.v1 = seed + PRIME32_1 + PRIME32_2;
     state.v2 = seed + PRIME32_2;
     state.v3 = seed + 0;
@@ -746,7 +746,7 @@ XXH_PUBLIC_API void XXH64_copyState(XXH64_state_t* dstState, const XXH64_state_t
 XXH_PUBLIC_API XXH_errorcode XXH64_reset(XXH64_state_t* statePtr, unsigned long long seed)
 {
     XXH64_state_t state;   /* using a local state to memcpy() in order to avoid strict-aliasing warnings */
-    memset(&state, 0, sizeof(state)-8);   /* do not write into reserved, for future removal */
+    memset(&state, 0, sizeof(state)-8);   /* do not write into reserved, for future removal */ //-V512
     state.v1 = seed + PRIME64_1 + PRIME64_2;
     state.v2 = seed + PRIME64_2;
     state.v3 = seed + 0;
