Disable warning when building xxhash.c.

#rnx

diff --git a/Engine/Source/Developer/Windows/LiveCoding/Private/External/xxhash.c b/Engine/Source/Developer/Windows/LiveCoding/Private/External/xxhash.c
--- a/Engine/Source/Developer/Windows/LiveCoding/Private/External/xxhash.c
+++ b/Engine/Source/Developer/Windows/LiveCoding/Private/External/xxhash.c
@@ -35,6 +35,7 @@
 // BEGIN EPIC MODS
 #pragma warning(push)
 #pragma warning(disable:6385) // warning C6385: Reading invalid data from '&state'
+#pragma warning(disable:4668) // warning C4668: '__cplusplus' is not defined as a preprocessor macro, replacing with '0' for '#if/#elif'
 // END EPIC MODS
 
 // MM: tuning
@@ -119,7 +120,6 @@ static void* XXH_memcpy(void* dest, const void* src, size_t size) { return memcp
 #define XXH_STATIC_LINKING_ONLY
 #include "xxhash.h"
 
-
 /* *************************************
 *  Compiler Specific Options
 ***************************************/
