LiveCoding: Add support for lazy loading modules. This lets us enable live coding for the entire engine without a large startup penalty. Project modules are still preloaded by default.

diff --git a/Engine/Source/Developer/Windows/LiveCoding/Private/External/LC_ClientStartupThread.cpp b/Engine/Source/Developer/Windows/LiveCoding/Private/External/LC_ClientStartupThread.cpp
--- a/Engine/Source/Developer/Windows/LiveCoding/Private/External/LC_ClientStartupThread.cpp
+++ b/Engine/Source/Developer/Windows/LiveCoding/Private/External/LC_ClientStartupThread.cpp
@@ -255,6 +255,18 @@ void ClientStartupThread::SetBuildArguments(const wchar_t* arguments)
 }
 // END EPIC MOD
 
+// BEGIN EPIC MOD - Support for lazy-loading modules
+void ClientStartupThread::EnableLazyLoadedModule(const wchar_t* fileName, Windows::HMODULE moduleBase)
+{
+	// we cannot wait for commands in the user command thread as long as startup hasn't finished
+	Join();
+
+	if (m_userCommandThread)
+	{
+		m_userCommandThread->EnableLazyLoadedModule(fileName, moduleBase);
+	}
+}
+// END EPIC MOD
 
 void ClientStartupThread::ApplySettingBool(const char* const settingName, int value)
 {
diff --git a/Engine/Source/Developer/Windows/LiveCoding/Private/External/LC_ClientStartupThread.h b/Engine/Source/Developer/Windows/LiveCoding/Private/External/LC_ClientStartupThread.h
--- a/Engine/Source/Developer/Windows/LiveCoding/Private/External/LC_ClientStartupThread.h
+++ b/Engine/Source/Developer/Windows/LiveCoding/Private/External/LC_ClientStartupThread.h
@@ -59,6 +59,10 @@ public:
 	void SetBuildArguments(const wchar_t* arguments);
 	// END EPIC MOD
 
+	// BEGIN EPIC MOD - Support for lazy-loading modules
+	void EnableLazyLoadedModule(const wchar_t* fileName, Windows::HMODULE moduleBase);
+	// END EPIC MOD
+
 private:
 	struct ThreadContext
 	{
diff --git a/Engine/Source/Developer/Windows/LiveCoding/Private/External/LC_ClientUserCommandThread.cpp b/Engine/Source/Developer/Windows/LiveCoding/Private/External/LC_ClientUserCommandThread.cpp
--- a/Engine/Source/Developer/Windows/LiveCoding/Private/External/LC_ClientUserCommandThread.cpp
+++ b/Engine/Source/Developer/Windows/LiveCoding/Private/External/LC_ClientUserCommandThread.cpp
@@ -233,6 +233,30 @@ namespace userCommands
 	};
 	// END EPIC MOD
 
+	// BEGIN EPIC MOD - Support for lazy-loading modules
+	struct EnableLazyLoadedModuleCommand : public BaseCommand
+	{
+		EnableLazyLoadedModuleCommand(void)
+			: BaseCommand(Scope::NONE)
+		{
+		}
+
+		virtual ~EnableLazyLoadedModuleCommand(void) {}
+
+		virtual void Execute(DuplexPipe* pipe) override
+		{
+			commands::EnableLazyLoadedModule serverCommand;
+			serverCommand.processId = process::GetId();
+			wcscpy_s(serverCommand.fileName, fileName.c_str());
+			serverCommand.moduleBase = moduleBase;
+			pipe->SendCommandAndWaitForAck(serverCommand);
+		}
+
+		std::wstring fileName;
+		Windows::HMODULE moduleBase;
+	};
+	// END EPIC MOD
+
 
 	struct BuildPatchCommand : public BaseCommand
 	{
@@ -604,6 +628,23 @@ void ClientUserCommandThread::SetBuildArguments(const wchar_t* arguments)
 }
 // END EPIC MOD
 
+// BEGIN EPIC MOD - Adding support for lazy-loading modules
+void ClientUserCommandThread::EnableLazyLoadedModule(const wchar_t* fileName, Windows::HMODULE moduleBase)
+{
+	userCommands::EnableLazyLoadedModuleCommand* command = new userCommands::EnableLazyLoadedModuleCommand;
+	command->fileName = fileName;
+	command->moduleBase = moduleBase;
+	{
+		CriticalSection::ScopedLock lock(&g_userCommandQueueCS);
+		g_userCommandQueue.push_front(command);
+	}
+
+	// signal to the thread that a new item is in the queue
+	m_itemInQueueEvent->Signal();
+}
+// END EPIC MOD
+
+
 void ClientUserCommandThread::ApplySettingBool(const char* const settingName, int value)
 {
 	userCommands::ApplySettingBoolCommand* command = new userCommands::ApplySettingBoolCommand;
diff --git a/Engine/Source/Developer/Windows/LiveCoding/Private/External/LC_ClientUserCommandThread.h b/Engine/Source/Developer/Windows/LiveCoding/Private/External/LC_ClientUserCommandThread.h
--- a/Engine/Source/Developer/Windows/LiveCoding/Private/External/LC_ClientUserCommandThread.h
+++ b/Engine/Source/Developer/Windows/LiveCoding/Private/External/LC_ClientUserCommandThread.h
@@ -63,6 +63,10 @@ public:
 	void SetBuildArguments(const wchar_t* arguments);
 	// END EPIC MOD
 
+	// BEGIN EPIC MOD - Support for lazy-loading modules
+	void EnableLazyLoadedModule(const wchar_t* fileName, Windows::HMODULE moduleBase);
+	// END EPIC MOD
+
 	void ApplySettingBool(const char* const settingName, int value);
 	void ApplySettingInt(const char* const settingName, int value);
 	void ApplySettingString(const char* const settingName, const wchar_t* const value);
diff --git a/Engine/Source/Developer/Windows/LiveCoding/Private/External/LC_Commands.h b/Engine/Source/Developer/Windows/LiveCoding/Private/External/LC_Commands.h
--- a/Engine/Source/Developer/Windows/LiveCoding/Private/External/LC_Commands.h
+++ b/Engine/Source/Developer/Windows/LiveCoding/Private/External/LC_Commands.h
@@ -316,11 +316,27 @@ namespace commands
 	};
 	// END EPIC MOD
 
+	// BEGIN EPIC MOD - Support for lazy-loading modules
+	struct EnableLazyLoadedModule
+	{
+		static const uint32_t ID = SetBuildArguments::ID + 1u;
+
+		unsigned int processId;
+		wchar_t fileName[260];
+		Windows::HMODULE moduleBase;
+	};
+
+	struct FinishedLazyLoadingModules
+	{
+		static const uint32_t ID = EnableLazyLoadedModule::ID + 1u;
+	};
+	// END EPIC MOD
+
 	// tell the EXE that a bool setting needs to be changed
 	struct ApplySettingBool
 	{
-		// BEGIN EPIC MOD - Adding SetBuildArguments command
-		static const uint32_t ID = SetBuildArguments::ID + 1u;
+		// BEGIN EPIC MOD - Support for lazy-loading modules
+		static const uint32_t ID = FinishedLazyLoadingModules::ID + 1u;
 		// END EPIC MOD
 
 		char settingName[256];
diff --git a/Engine/Source/Developer/Windows/LiveCoding/Private/External/LC_EntryPoint.cpp b/Engine/Source/Developer/Windows/LiveCoding/Private/External/LC_EntryPoint.cpp
--- a/Engine/Source/Developer/Windows/LiveCoding/Private/External/LC_EntryPoint.cpp
+++ b/Engine/Source/Developer/Windows/LiveCoding/Private/External/LC_EntryPoint.cpp
@@ -175,6 +175,13 @@ LPP_DLL_API void __cdecl LppSetBuildArguments(const wchar_t* arguments)
 }
 // END EPIC MOD
 
+// BEGIN EPIC MOD - Support for lazy-loading modules
+LPP_DLL_API void __cdecl LppEnableLazyLoadedModule(const wchar_t* const nameOfExeOrDll)
+{
+	HMODULE baseAddress = GetModuleHandle(nameOfExeOrDll);
+	g_startupThread->EnableLazyLoadedModule(nameOfExeOrDll, baseAddress);
+}
+// END EPIC MOD
 
 LPP_DLL_API void __cdecl LppApplySettingBool(const char* const settingName, int value)
 {
diff --git a/Engine/Source/Developer/Windows/LiveCoding/Private/External/LC_EntryPoint.h b/Engine/Source/Developer/Windows/LiveCoding/Private/External/LC_EntryPoint.h
--- a/Engine/Source/Developer/Windows/LiveCoding/Private/External/LC_EntryPoint.h
+++ b/Engine/Source/Developer/Windows/LiveCoding/Private/External/LC_EntryPoint.h
@@ -18,3 +18,7 @@ void __cdecl LppSetBuildArguments(const wchar_t* const arguments);
 void __cdecl LppApplySettingBool(const char* const settingName, int value);
 void __cdecl LppApplySettingInt(const char* const settingName, int value);
 void __cdecl LppApplySettingString(const char* const settingName, const wchar_t* const value);
+
+// BEGIN EPIC MOD - Support for lazy-loading modules
+void __cdecl LppEnableLazyLoadedModule(const wchar_t* const nameOfExeOrDll);
+// END EPIC MODS
\ No newline at end of file
diff --git a/Engine/Source/Developer/Windows/LiveCoding/Private/LiveCodingModule.cpp b/Engine/Source/Developer/Windows/LiveCoding/Private/LiveCodingModule.cpp
--- a/Engine/Source/Developer/Windows/LiveCoding/Private/LiveCodingModule.cpp
+++ b/Engine/Source/Developer/Windows/LiveCoding/Private/LiveCodingModule.cpp
@@ -258,7 +258,7 @@ void FLiveCodingModule::UpdateModules()
 #if IS_MONOLITHIC
 	wchar_t FullFilePath[WINDOWS_MAX_PATH];
 	verify(GetModuleFileName(hInstance, FullFilePath, ARRAY_COUNT(FullFilePath)));
-	EnableModule(FullFilePath);
+	LppEnableModule(FullFilePath);
 #else
 	TArray<FModuleStatus> ModuleStatuses;
 	FModuleManager::Get().QueryModules(ModuleStatuses);
@@ -280,24 +280,6 @@ void FLiveCodingModule::UpdateModules()
 #endif
 }
 
-void FLiveCodingModule::EnableModule(const FString& FullFilePath)
-{
-	if (!EnabledModules.Contains(FullFilePath))
-	{
-		LppEnableModule(*FullFilePath);
-		EnabledModules.Add(FullFilePath);
-	}
-}
-
-void FLiveCodingModule::DisableModule(const FString& FullFilePath)
-{
-	if(EnabledModules.Contains(FullFilePath))
-	{
-		LppDisableModule(*FullFilePath);
-		EnabledModules.Remove(FullFilePath);
-	}
-}
-
 void FLiveCodingModule::OnModulesChanged(FName ModuleName, EModuleChangeReason Reason)
 {
 #if !IS_MONOLITHIC
@@ -316,42 +298,42 @@ void FLiveCodingModule::OnModulesChanged(FName ModuleName, EModuleChangeReason R
 void FLiveCodingModule::ConfigureModule(const FName& Name, const FString& FullFilePath)
 {
 #if !IS_MONOLITHIC
-	if (ShouldEnableModule(Name, FullFilePath))
+	if (!ConfiguredModules.Contains(Name))
 	{
-		EnableModule(FullFilePath);
-	}
-	else
-	{
-		DisableModule(FullFilePath);
+		if (ShouldPreloadModule(Name, FullFilePath))
+		{
+			LppEnableModule(*FullFilePath);
+		}
+		else
+		{
+			LppEnableLazyLoadedModule(*FullFilePath);
+		}
+		ConfiguredModules.Add(Name);
 	}
 #endif
 }
 
-bool FLiveCodingModule::ShouldEnableModule(const FName& Name, const FString& FullFilePath) const
+bool FLiveCodingModule::ShouldPreloadModule(const FName& Name, const FString& FullFilePath) const
 {
-	if (Settings->ExcludeSpecificModules.Contains(Name))
-	{
-		return false;
-	}
-	if (Settings->IncludeSpecificModules.Contains(Name))
+	if (Settings->PreloadNamedModules.Contains(Name))
 	{
 		return true;
 	}
 
 	if (FullFilePath.StartsWith(FullProjectDir))
 	{
-		if (Settings->bIncludeProjectModules == Settings->bIncludeProjectPluginModules)
+		if (Settings->bPreloadProjectModules == Settings->bPreloadProjectPluginModules)
 		{
-			return Settings->bIncludeProjectModules;
+			return Settings->bPreloadProjectModules;
 		}
 
 		if(FullFilePath.StartsWith(FullProjectPluginsDir))
 		{
-			return Settings->bIncludeProjectPluginModules;
+			return Settings->bPreloadProjectPluginModules;
 		}
 		else
 		{
-			return Settings->bIncludeProjectModules;
+			return Settings->bPreloadProjectModules;
 		}
 	}
 	else
@@ -361,18 +343,18 @@ bool FLiveCodingModule::ShouldEnableModule(const FName& Name, const FString& Ful
 			return false;
 		}
 
-		if (Settings->bIncludeEngineModules == Settings->bIncludeEnginePluginModules)
+		if (Settings->bPreloadEngineModules == Settings->bPreloadEnginePluginModules)
 		{
-			return Settings->bIncludeEngineModules;
+			return Settings->bPreloadEngineModules;
 		}
 
 		if(FullFilePath.StartsWith(FullEnginePluginsDir))
 		{
-			return Settings->bIncludeEnginePluginModules;
+			return Settings->bPreloadEnginePluginModules;
 		}
 		else
 		{
-			return Settings->bIncludeEngineModules;
+			return Settings->bPreloadEngineModules;
 		}
 	}
 }
diff --git a/Engine/Source/Developer/Windows/LiveCoding/Private/LiveCodingModule.h b/Engine/Source/Developer/Windows/LiveCoding/Private/LiveCodingModule.h
--- a/Engine/Source/Developer/Windows/LiveCoding/Private/LiveCodingModule.h
+++ b/Engine/Source/Developer/Windows/LiveCoding/Private/LiveCodingModule.h
@@ -39,7 +39,7 @@ private:
 	bool bEnabledLastTick;
 	bool bEnabledForSession;
 	bool bStarted;
-	TSet<FString> EnabledModules;
+	TSet<FName> ConfiguredModules;
 
 	const FString FullEnginePluginsDir;
 	const FString FullProjectDir;
@@ -55,10 +55,8 @@ private:
 	void OnModulesChanged(FName ModuleName, EModuleChangeReason Reason);
 
 	void UpdateModules();
-	void EnableModule(const FString& FullFilePath);
-	void DisableModule(const FString& FullFilePath);
 
 	void ConfigureModule(const FName& Name, const FString& FullFilePath);
-	bool ShouldEnableModule(const FName& Name, const FString& FullFilePath) const;
+	bool ShouldPreloadModule(const FName& Name, const FString& FullFilePath) const;
 };
 
diff --git a/Engine/Source/Developer/Windows/LiveCoding/Private/LiveCodingSettings.cpp b/Engine/Source/Developer/Windows/LiveCoding/Private/LiveCodingSettings.cpp
--- a/Engine/Source/Developer/Windows/LiveCoding/Private/LiveCodingSettings.cpp
+++ b/Engine/Source/Developer/Windows/LiveCoding/Private/LiveCodingSettings.cpp
@@ -5,10 +5,10 @@
 
 ULiveCodingSettings::ULiveCodingSettings(const FObjectInitializer& ObjectInitializer) : Super(ObjectInitializer)
 {
-	UProperty* EngineModulesProperty = StaticClass()->FindPropertyByName("bIncludeEngineModules");
+	UProperty* EngineModulesProperty = StaticClass()->FindPropertyByName("bPreloadEngineModules");
 	check(EngineModulesProperty != nullptr);
 
-	UProperty* EnginePluginModulesProperty = StaticClass()->FindPropertyByName("bIncludeEnginePluginModules");
+	UProperty* EnginePluginModulesProperty = StaticClass()->FindPropertyByName("bPreloadEnginePluginModules");
 	check(EnginePluginModulesProperty != nullptr);
 
 	if (FApp::IsEngineInstalled())
@@ -17,6 +17,6 @@ ULiveCodingSettings::ULiveCodingSettings(const FObjectInitializer& ObjectInitial
 		EnginePluginModulesProperty->ClearPropertyFlags(CPF_Edit);
 	}
 
-	bIncludeProjectModules = true;
-	bIncludeProjectPluginModules = true;
+	bPreloadProjectModules = true;
+	bPreloadProjectPluginModules = true;
 }
diff --git a/Engine/Source/Developer/Windows/LiveCoding/Private/LiveCodingSettings.h b/Engine/Source/Developer/Windows/LiveCoding/Private/LiveCodingSettings.h
--- a/Engine/Source/Developer/Windows/LiveCoding/Private/LiveCodingSettings.h
+++ b/Engine/Source/Developer/Windows/LiveCoding/Private/LiveCodingSettings.h
@@ -27,22 +27,19 @@ public:
     ELiveCodingStartupMode Startup;
 
 	UPROPERTY(config, EditAnywhere, Category=Modules, Meta=(ConfigRestartRequired=true, EditCondition="bEnabled"))
-    bool bIncludeEngineModules;
+    bool bPreloadEngineModules;
 
     UPROPERTY(config, EditAnywhere, Category=Modules, Meta=(ConfigRestartRequired=true, EditCondition="bEnabled"))
-    bool bIncludeEnginePluginModules;
+    bool bPreloadEnginePluginModules;
 
     UPROPERTY(config, EditAnywhere, Category=Modules, Meta=(ConfigRestartRequired=true, EditCondition="bEnabled"))
-    bool bIncludeProjectModules;
+    bool bPreloadProjectModules;
 
     UPROPERTY(config, EditAnywhere, Category=Modules, Meta=(ConfigRestartRequired=true, EditCondition="bEnabled"))
-    bool bIncludeProjectPluginModules;
+    bool bPreloadProjectPluginModules;
 
     UPROPERTY(config, EditAnywhere, Category=Modules, Meta=(ConfigRestartRequired=true, EditCondition="bEnabled"))
-    TArray<FName> IncludeSpecificModules;
-
-    UPROPERTY(config, EditAnywhere, Category=Modules, Meta=(ConfigRestartRequired=true, EditCondition="bEnabled"))
-    TArray<FName> ExcludeSpecificModules;
+    TArray<FName> PreloadNamedModules;
 
 	ULiveCodingSettings(const FObjectInitializer& Initializer);
 };
diff --git a/Engine/Source/Developer/Windows/LiveCodingServer/LiveCodingServer.Build.cs b/Engine/Source/Developer/Windows/LiveCodingServer/LiveCodingServer.Build.cs
--- a/Engine/Source/Developer/Windows/LiveCodingServer/LiveCodingServer.Build.cs
+++ b/Engine/Source/Developer/Windows/LiveCodingServer/LiveCodingServer.Build.cs
@@ -27,6 +27,10 @@ public class LiveCodingServer : ModuleRules
 			RuntimeDependencies.Add("$(TargetOutputDir)/msdia140.dll", Path.Combine(DiaSdkDir, "bin", "amd64", "msdia140.dll"));
 		}
 
-		PrecompileForTargets = PrecompileTargetsType.None;
+		// Allow precompiling when generating project files so we can get intellisense
+		if(!Target.bGenerateProjectFiles)
+		{
+			PrecompileForTargets = PrecompileTargetsType.None;
+		}
 	}
 }
diff --git a/Engine/Source/Developer/Windows/LiveCodingServer/Private/External/LC_LiveModule.cpp b/Engine/Source/Developer/Windows/LiveCodingServer/Private/External/LC_LiveModule.cpp
--- a/Engine/Source/Developer/Windows/LiveCodingServer/Private/External/LC_LiveModule.cpp
+++ b/Engine/Source/Developer/Windows/LiveCodingServer/Private/External/LC_LiveModule.cpp
@@ -1334,7 +1334,7 @@ LiveModule::ErrorType::Enum LiveModule::Update(FileAttributeCache* fileCache, Di
 			LC_LOG_USER("File %S was modified or is new", modifiedOrNewObjFiles[i].c_str());
 		}
 
-		LC_LOG_USER("Building patch from %zu file(s) for Live++ module %S", modifiedOrNewObjFiles.size(), m_moduleName.c_str());
+		LC_LOG_USER("Building patch from %zu file(s) for Live Coding module %S", modifiedOrNewObjFiles.size(), m_moduleName.c_str());
 	}
 
 	// let the user know that we're about to compile
diff --git a/Engine/Source/Developer/Windows/LiveCodingServer/Private/External/LC_LiveProcess.cpp b/Engine/Source/Developer/Windows/LiveCodingServer/Private/External/LC_LiveProcess.cpp
--- a/Engine/Source/Developer/Windows/LiveCodingServer/Private/External/LC_LiveProcess.cpp
+++ b/Engine/Source/Developer/Windows/LiveCodingServer/Private/External/LC_LiveProcess.cpp
@@ -70,3 +70,41 @@ bool LiveProcess::TriedToLoadImage(const executable::Header& imageHeader) const
 {
 	return (m_imagesTriedToLoad.find(imageHeader) != m_imagesTriedToLoad.end());
 }
+
+// BEGIN EPIC MOD - Allow lazy-loading modules
+void LiveProcess::AddLazyLoadedModule(const std::wstring moduleName, Windows::HMODULE moduleBase)
+{
+	LazyLoadedModule module;
+	module.m_moduleBase = moduleBase;
+	module.m_loaded = false;
+	m_lazyLoadedModules.insert(std::make_pair(moduleName, module));
+}
+
+void LiveProcess::SetLazyLoadedModuleAsLoaded(const std::wstring moduleName)
+{
+	std::unordered_map<std::wstring, LazyLoadedModule>::iterator it = m_lazyLoadedModules.find(moduleName);
+	if (it != m_lazyLoadedModules.end())
+	{
+		it->second.m_loaded = true;
+	}
+}
+
+bool LiveProcess::IsPendingLazyLoadedModule(const std::wstring& moduleName) const
+{
+	std::unordered_map<std::wstring, LazyLoadedModule>::const_iterator iter = m_lazyLoadedModules.find(moduleName);
+	return iter != m_lazyLoadedModules.end() && !iter->second.m_loaded;
+}
+
+Windows::HMODULE LiveProcess::GetLazyLoadedModuleBase(const std::wstring& moduleName) const
+{
+	std::unordered_map<std::wstring, LazyLoadedModule>::const_iterator iter = m_lazyLoadedModules.find(moduleName);
+	if (iter == m_lazyLoadedModules.end())
+	{
+		return nullptr;
+	}
+	else
+	{
+		return iter->second.m_moduleBase;
+	}
+}
+// END EPIC MOD
diff --git a/Engine/Source/Developer/Windows/LiveCodingServer/Private/External/LC_LiveProcess.h b/Engine/Source/Developer/Windows/LiveCodingServer/Private/External/LC_LiveProcess.h
--- a/Engine/Source/Developer/Windows/LiveCodingServer/Private/External/LC_LiveProcess.h
+++ b/Engine/Source/Developer/Windows/LiveCodingServer/Private/External/LC_LiveProcess.h
@@ -64,6 +64,13 @@ public:
 	}
 	// END EPIC MOD
 
+	// BEGIN EPIC MOD - Allow lazy-loading modules
+	void AddLazyLoadedModule(const std::wstring moduleName, Windows::HMODULE moduleBase);
+	void SetLazyLoadedModuleAsLoaded(const std::wstring moduleName);
+	bool IsPendingLazyLoadedModule(const std::wstring& moduleName) const;
+	Windows::HMODULE GetLazyLoadedModuleBase(const std::wstring& moduleName) const;
+	// END EPIC MOD
+
 private:
 	process::Handle m_processHandle;
 	unsigned int m_processId;
@@ -74,6 +81,16 @@ private:
 	std::wstring m_buildArguments;
 	// END EPIC MOD
 
+	// BEGIN EPIC MOD - Allow lazy-loading modules
+	struct LazyLoadedModule
+	{
+		Windows::HMODULE m_moduleBase;
+		bool m_loaded;
+	};
+
+	types::unordered_map<std::wstring, LazyLoadedModule> m_lazyLoadedModules;
+	// END EPIC MOD
+
 	// loaded modules are not identified by their full path, but by their executable image header.
 	// we do this to ensure that the same executable loaded from a different path is not treated as
 	// a different executable.
diff --git a/Engine/Source/Developer/Windows/LiveCodingServer/Private/External/LC_ServerCommandThread.cpp b/Engine/Source/Developer/Windows/LiveCodingServer/Private/External/LC_ServerCommandThread.cpp
--- a/Engine/Source/Developer/Windows/LiveCodingServer/Private/External/LC_ServerCommandThread.cpp
+++ b/Engine/Source/Developer/Windows/LiveCodingServer/Private/External/LC_ServerCommandThread.cpp
@@ -22,6 +22,7 @@
 #include "LC_PrimitiveNames.h"
 #include "LC_AppSettings.h"
 #include "LC_Allocators.h"
+#include "LC_DuplexPipeClient.h"
 #include <mmsystem.h>
 
 // unreachable code
@@ -584,6 +585,96 @@ BOOL CALLBACK FocusApplicationWindows(HWND WindowHandle, LPARAM Lparam)
 }
 // END EPIC MOD
 
+// BEGIN EPIC MOD - Support for lazy-loading modules
+bool ServerCommandThread::FinishedLazyLoadingModulesAction::Execute(CommandType* command, const DuplexPipe* pipe, void* context)
+{ 
+	pipe->SendAck(); 
+	return false; 
+}
+
+struct ClientProxyThread
+{
+	struct ProxyGetModuleAction
+	{
+		typedef commands::GetModule CommandType;
+
+		static bool Execute(CommandType* command, const DuplexPipe* pipe, void* context)
+		{
+			pipe->SendAck();
+
+			LiveProcess* process = static_cast<LiveProcess*>(context);
+
+			commands::GetModuleInfo cmd;
+			cmd.moduleBase = process->GetLazyLoadedModuleBase(command->path);
+			cmd.processId = process->GetProcessId();
+			cmd.loadImports = command->loadImports;
+			cmd.taskContext = command->taskContext;
+			wcscpy_s(cmd.path, command->path);
+			pipe->SendCommandAndWaitForAck(cmd);
+
+			return true;
+		}
+	};
+
+	struct ProxyEnableModuleFinishedAction
+	{
+		typedef commands::EnableModuleFinished CommandType;
+
+		static bool Execute(CommandType* command, const DuplexPipe* pipe, void* context)
+		{
+			pipe->SendAck();
+			return false;
+		}
+	};
+
+	LiveProcess* m_process;
+	DuplexPipeClient* m_pipe;
+	std::vector<std::wstring> m_enableModules;
+	thread::Handle m_threadHandle;
+
+	ClientProxyThread(LiveProcess* process, DuplexPipeClient* pipe, const std::vector<std::wstring> enableModules)
+		: m_process(process)
+		, m_pipe(pipe)
+		, m_enableModules(enableModules)
+	{
+		m_threadHandle = thread::Create(64u * 1024u, &StaticEntryPoint, this);
+		thread::SetName("Live coding client proxy");
+	}
+
+	~ClientProxyThread()
+	{
+		thread::Join(m_threadHandle);
+		thread::Close(m_threadHandle);
+	}
+
+	static unsigned int __stdcall StaticEntryPoint(void* context)
+	{
+		static_cast<ClientProxyThread*>(context)->EntryPoint();
+		return 0;
+	}
+
+	void EntryPoint()
+	{
+		m_pipe->SendCommandAndWaitForAck(commands::EnableModuleBatchBegin());
+		for(const std::wstring& enableModule : m_enableModules)
+		{
+			commands::EnableModule enableModuleCommand;
+			enableModuleCommand.processId = m_process->GetProcessId();
+			wcscpy_s(enableModuleCommand.path, enableModule.c_str());
+			enableModuleCommand.token = nullptr;
+			m_pipe->SendCommandAndWaitForAck(enableModuleCommand);
+
+			CommandMap commandMap;
+			commandMap.RegisterAction<ProxyGetModuleAction>();
+			commandMap.RegisterAction<ProxyEnableModuleFinishedAction>();
+			commandMap.HandleCommands(m_pipe, m_process);
+		}
+		m_pipe->SendCommandAndWaitForAck(commands::EnableModuleBatchEnd());
+		m_pipe->SendCommandAndWaitForAck(commands::FinishedLazyLoadingModules());
+	}
+};
+// END EPIC MOD
+
 void ServerCommandThread::CompileChanges(bool didAllProcessesMakeProgress)
 {
 	// recompile files, if any
@@ -602,6 +693,7 @@ void ServerCommandThread::CompileChanges(bool didAllProcessesMakeProgress)
 	const ILiveCodingServer::FCompileDelegate& CompileDelegate = GLiveCodingServer->GetCompileDelegate();
 	if (CompileDelegate.IsBound())
 	{
+		// Get the list of arguments for building each target, and use the delegate to pass them to UBT
 		TArray<FString> Targets;
 		for (LiveProcess* liveProcess : m_liveProcesses)
 		{
@@ -617,6 +709,45 @@ void ServerCommandThread::CompileChanges(bool didAllProcessesMakeProgress)
 			return;
 		}
 
+		// Enable any lazy-loaded modules that we need
+		for (LiveProcess* liveProcess : m_liveProcesses)
+		{
+			types::vector<std::wstring> LoadModuleFileNames;
+			for(const TPair<FString, TArray<FString>>& Pair : ModuleToObjectFiles)
+			{
+				std::wstring ModuleFileName = file::NormalizePath(*Pair.Key);
+				if (liveProcess->IsPendingLazyLoadedModule(ModuleFileName))
+				{
+					LoadModuleFileNames.push_back(ModuleFileName);
+				}
+			}
+			if (LoadModuleFileNames.size() > 0)
+			{
+				const std::wstring PipeName = primitiveNames::Pipe(m_processGroupName + L"_ClientProxy");
+
+				DuplexPipeServer ServerPipe;
+				ServerPipe.Create(PipeName.c_str());
+
+				DuplexPipeClient ClientPipe;
+				ClientPipe.Connect(PipeName.c_str());
+
+				ClientProxyThread ClientThread(liveProcess, &ClientPipe, LoadModuleFileNames);
+
+				CommandMap commandMap;
+				commandMap.RegisterAction<EnableModuleBatchBeginAction>();
+				commandMap.RegisterAction<EnableModuleBatchEndAction>();
+				commandMap.RegisterAction<EnableModuleAction>();
+				commandMap.RegisterAction<FinishedLazyLoadingModulesAction>();
+				commandMap.HandleCommands(&ServerPipe, this);
+
+				for (const std::wstring& loadModuleFileName : LoadModuleFileNames)
+				{
+					liveProcess->SetLazyLoadedModuleAsLoaded(loadModuleFileName);
+				}
+			}
+		}
+
+		// Build up a list of all the modified object files in each module
 		types::unordered_set<std::wstring> ValidModuleFileNames;
 		for (const LiveModule* liveModule : m_liveModules)
 		{
@@ -1046,6 +1177,9 @@ unsigned int ServerCommandThread::CommandThread(DuplexPipeServer* pipe, Event* r
 	// BEGIN EPIC MOD - Adding SetBuildArguments command
 	commandMap.RegisterAction<SetBuildArgumentsAction>();
 	// END EPIC MOD
+	// BEGIN EPIC MOD - Support for lazy-loading modules
+	commandMap.RegisterAction<EnableLazyLoadedModuleAction>();
+	// END EPIC MOD
 	commandMap.RegisterAction<RegisterProcessAction>();
 	commandMap.RegisterAction<EnableModuleBatchBeginAction>();
 	commandMap.RegisterAction<EnableModuleBatchEndAction>();
@@ -1356,6 +1490,30 @@ bool ServerCommandThread::SetBuildArgumentsAction::Execute(CommandType* command,
 }
 // END EPIC MOD
 
+// BEGIN EPIC MOD - Support for lazy-loading modules
+bool ServerCommandThread::EnableLazyLoadedModuleAction::Execute(CommandType* command, const DuplexPipe* pipe, void* context)
+{
+	ServerCommandThread* commandThread = static_cast<ServerCommandThread*>(context);
+
+	// protect against accepting this command while compilation is already in progress
+	CriticalSection::ScopedLock lock(&commandThread->m_actionCS);
+
+	for (LiveProcess* process : commandThread->m_liveProcesses)
+	{
+		if (process->GetProcessId() == command->processId)
+		{
+			const std::wstring modulePath = file::NormalizePath(command->fileName);
+			process->AddLazyLoadedModule(modulePath, command->moduleBase);
+			LC_LOG_DEV("Registered module %S for lazy-loading", modulePath.c_str());
+		}
+	}
+
+	pipe->SendAck();
+
+	return true;
+}
+// END EPIC MOD
+
 bool ServerCommandThread::RegisterProcessAction::Execute(CommandType* command, const DuplexPipe* pipe, void* context)
 {
 	pipe->SendAck();
@@ -1453,6 +1611,9 @@ bool ServerCommandThread::EnableModuleBatchEndAction::Execute(CommandType*, cons
 	// EPIC REMOVED: g_theApp.GetMainFrame()->ResetStatusBarText();
 
 	// tell user we are ready
+	// BEGIN EPIC MOD - Support for lazy-loading modules
+	if (thread::GetId() != thread::GetId(commandThread->m_compileThread))
+	// END EPIC MOD
 	{
 		const int shortcut = appSettings::g_compileShortcut->GetValue();
 		const std::wstring& shortcutText = shortcut::ConvertShortcutToText(shortcut);
diff --git a/Engine/Source/Developer/Windows/LiveCodingServer/Private/External/LC_ServerCommandThread.h b/Engine/Source/Developer/Windows/LiveCodingServer/Private/External/LC_ServerCommandThread.h
--- a/Engine/Source/Developer/Windows/LiveCodingServer/Private/External/LC_ServerCommandThread.h
+++ b/Engine/Source/Developer/Windows/LiveCodingServer/Private/External/LC_ServerCommandThread.h
@@ -136,6 +136,20 @@ private:
 	};
 	// END EPIC MOD
 
+	// BEGIN EPIC MOD - Adding support for lazy-loading modules
+	struct EnableLazyLoadedModuleAction
+	{
+		typedef commands::EnableLazyLoadedModule CommandType;
+		static bool Execute(CommandType* command, const DuplexPipe* pipe, void* context);
+	};
+
+	struct FinishedLazyLoadingModulesAction
+	{
+		typedef commands::FinishedLazyLoadingModules CommandType;
+		static bool Execute(CommandType* command, const DuplexPipe* pipe, void* context);
+	};
+	// END EPIC MOD
+
 	struct RegisterProcessAction
 	{
 		typedef commands::RegisterProcess CommandType;
