LiveCoding: Fix unnecessary newlines on log messages.

diff --git a/Engine/Source/Developer/Windows/LiveCoding/Private/External/LC_Logging.cpp b/Engine/Source/Developer/Windows/LiveCoding/Private/External/LC_Logging.cpp
--- a/Engine/Source/Developer/Windows/LiveCoding/Private/External/LC_Logging.cpp
+++ b/Engine/Source/Developer/Windows/LiveCoding/Private/External/LC_Logging.cpp
@@ -30,16 +30,17 @@ namespace
 	// BEGIN EPIC MOD - Redirecting log output
 	static void DefaultOutputHandler(logging::Channel::Enum channel, logging::Type::Enum type, const wchar_t* const Message)
 	{
+		FString MessageWithoutNewline = FString(Message).TrimEnd();
 		switch (type)
 		{
 		case logging::Type::LOG_WARNING:
-			UE_LOG(LogLiveCoding, Warning, TEXT("%s"), Message);
+			UE_LOG(LogLiveCoding, Warning, TEXT("%s"), *MessageWithoutNewline);
 			break;
 		case logging::Type::LOG_ERROR:
-			UE_LOG(LogLiveCoding, Error, TEXT("%s"), Message);
+			UE_LOG(LogLiveCoding, Error, TEXT("%s"), *MessageWithoutNewline);
 			break;
 		default:
-			UE_LOG(LogLiveCoding, Display, TEXT("%s"), Message);
+			UE_LOG(LogLiveCoding, Display, TEXT("%s"), *MessageWithoutNewline);
 			break;
 		}
 	}
