Live Coding:  Fixes issue linking patch DLLs with Visual Studio Express, due to PATH environment variable not including path to mspdb140.dll.

Now stores the linker environment block in the live coding manifest, and adds it to the cache before executing the linker.

diff --git a/Engine/Source/Developer/Windows/LiveCodingServer/Private/External/LC_Compiler.cpp b/Engine/Source/Developer/Windows/LiveCodingServer/Private/External/LC_Compiler.cpp
--- a/Engine/Source/Developer/Windows/LiveCodingServer/Private/External/LC_Compiler.cpp
+++ b/Engine/Source/Developer/Windows/LiveCodingServer/Private/External/LC_Compiler.cpp
@@ -31,6 +31,16 @@ namespace
 
 		void Insert(const wchar_t* key, environment::Block* value)
 		{
+	        // BEGIN EPIC MOD - Allow passing environment block for linker
+			auto it = m_cache.find(key);
+			if (it != m_cache.end())
+			{
+				environment::DestroyBlock(it->second);
+				it->second = value;
+				return;
+			}
+			// END EPIC MOD
+
 			m_cache[key] = value;
 		}
 
@@ -230,3 +240,11 @@ const environment::Block* compiler::UpdateEnvironmentCache(const wchar_t* absolu
 
 	return CreateEnvironmentCacheEntry(absolutePathToCompilerExe);
 }
+
+// BEGIN EPIC MOD - Allow passing environment block for linker
+void compiler::AddEnvironmentToCache(const wchar_t* absolutePathToCompilerExe, environment::Block* block)
+{
+	CriticalSection::ScopedLock lock(&g_compilerCacheCS);
+	g_compilerEnvironmentCache.Insert(absolutePathToCompilerExe, block);
+}
+// END EPIC MOD
diff --git a/Engine/Source/Developer/Windows/LiveCodingServer/Private/External/LC_Compiler.h b/Engine/Source/Developer/Windows/LiveCodingServer/Private/External/LC_Compiler.h
--- a/Engine/Source/Developer/Windows/LiveCodingServer/Private/External/LC_Compiler.h
+++ b/Engine/Source/Developer/Windows/LiveCodingServer/Private/External/LC_Compiler.h
@@ -22,4 +22,8 @@ namespace compiler
 	// helper function that either creates a new entry in the cache if none exists yet,
 	// or returns the one found in the cache.
 	const environment::Block* UpdateEnvironmentCache(const wchar_t* absolutePathToCompilerExe);
+	
+	// BEGIN EPIC MOD - Allow overriding environment block for tools
+	void AddEnvironmentToCache(const wchar_t* absolutePathToCompilerExe, environment::Block* block);
+	// END EPIC MOD
 }
diff --git a/Engine/Source/Developer/Windows/LiveCodingServer/Private/External/LC_Environment.cpp b/Engine/Source/Developer/Windows/LiveCodingServer/Private/External/LC_Environment.cpp
--- a/Engine/Source/Developer/Windows/LiveCodingServer/Private/External/LC_Environment.cpp
+++ b/Engine/Source/Developer/Windows/LiveCodingServer/Private/External/LC_Environment.cpp
@@ -7,6 +7,9 @@
 #include "LC_AppSettings.h"
 #include "LC_Logging.h"
 
+// BEGIN EPIC MOD - Allow passing environment block for linker
+#include "Containers/StringConv.h"
+// END EPIC MOD
 
 namespace environment
 {
@@ -16,6 +19,25 @@ namespace environment
 		char* data;
 	};
 
+	// BEGIN EPIC MOD - Allow passing environment block for linker
+	Block* CreateBlockFromMap(const TMap<FString, FString>& Pairs)
+	{
+		std::vector<char> blockData;
+		for (const TPair<FString, FString>& Pair : Pairs)
+		{
+			auto AnsiString = StringCast<ANSICHAR>(*FString::Printf(TEXT("%s=%s"), *Pair.Key, *Pair.Value));
+			const char* AnsiStringData = AnsiString.Get();
+			blockData.insert(blockData.end(), AnsiStringData, strchr(AnsiStringData, 0) + 1);
+		}
+		blockData.push_back('\0');
+
+		Block* block = new Block;
+		block->size = blockData.size();
+		block->data = new char[blockData.size()];
+		memcpy(block->data, blockData.data(), blockData.size());
+		return block;
+	}
+	// END EPIC MOD
 
 	Block* CreateBlockFromFile(const wchar_t* path)
 	{
diff --git a/Engine/Source/Developer/Windows/LiveCodingServer/Private/External/LC_Environment.h b/Engine/Source/Developer/Windows/LiveCodingServer/Private/External/LC_Environment.h
--- a/Engine/Source/Developer/Windows/LiveCodingServer/Private/External/LC_Environment.h
+++ b/Engine/Source/Developer/Windows/LiveCodingServer/Private/External/LC_Environment.h
@@ -5,10 +5,18 @@
 #include "CoreTypes.h"
 #include <string>
 
+// BEGIN EPIC MOD - Allow passing environment block for linker
+#include "Containers/Map.h"
+// END EPIC MOD
+
 namespace environment
 {
 	struct Block;
 
+	// BEGIN EPIC MOD - Allow passing environment block for linker
+	Block* CreateBlockFromMap(const TMap<FString, FString>& Pairs);
+	// END EPIC MOD
+
 	Block* CreateBlockFromFile(const wchar_t* path);
 	void DestroyBlock(Block*& block);
 
diff --git a/Engine/Source/Developer/Windows/LiveCodingServer/Private/LiveCodingServer.cpp b/Engine/Source/Developer/Windows/LiveCodingServer/Private/LiveCodingServer.cpp
--- a/Engine/Source/Developer/Windows/LiveCodingServer/Private/LiveCodingServer.cpp
+++ b/Engine/Source/Developer/Windows/LiveCodingServer/Private/LiveCodingServer.cpp
@@ -6,6 +6,8 @@
 #include "External/LC_FileUtil.h"
 #include "External/LC_AppSettings.h"
 #include "External/LC_ServerCommandThread.h"
+#include "External/LC_Compiler.h"
+#include "External/LC_Environment.h"
 
 FLiveCodingServer* GLiveCodingServer = nullptr;
 
@@ -47,10 +49,16 @@ void FLiveCodingServer::Stop()
 	ProcessGroupName.clear();
 }
 
-void FLiveCodingServer::SetLinkerPath(const wchar_t* LinkerPath)
+void FLiveCodingServer::SetLinkerPath(const wchar_t* LinkerPath, const TMap<FString, FString>& LinkerEnvironment)
 {
 	appSettings::g_linkerPath->SetValueWithoutSaving(LinkerPath);
 	appSettings::UpdateLinkerPathCache();
+
+	if (LinkerEnvironment.Num() > 0)
+	{
+		environment::Block* block = environment::CreateBlockFromMap(LinkerEnvironment);
+		compiler::AddEnvironmentToCache(LinkerPath, block);
+	}
 }
 
 ILiveCodingServer::FBringToFrontDelegate& FLiveCodingServer::GetBringToFrontDelegate()
diff --git a/Engine/Source/Developer/Windows/LiveCodingServer/Private/LiveCodingServer.h b/Engine/Source/Developer/Windows/LiveCodingServer/Private/LiveCodingServer.h
--- a/Engine/Source/Developer/Windows/LiveCodingServer/Private/LiveCodingServer.h
+++ b/Engine/Source/Developer/Windows/LiveCodingServer/Private/LiveCodingServer.h
@@ -18,7 +18,7 @@ public:
 	virtual void Start(const wchar_t* ProcessGroupName) override;
 	virtual void Stop() override;
 
-	virtual void SetLinkerPath(const wchar_t* LinkerPath) override;
+	virtual void SetLinkerPath(const wchar_t* LinkerPath, const TMap<FString, FString>& LinkerEnvironment) override;
 
 	// ILiveCodingServer implementation
 	virtual FBringToFrontDelegate& GetBringToFrontDelegate() override final;
diff --git a/Engine/Source/Developer/Windows/LiveCodingServer/Public/ILiveCodingServer.h b/Engine/Source/Developer/Windows/LiveCodingServer/Public/ILiveCodingServer.h
--- a/Engine/Source/Developer/Windows/LiveCodingServer/Public/ILiveCodingServer.h
+++ b/Engine/Source/Developer/Windows/LiveCodingServer/Public/ILiveCodingServer.h
@@ -3,6 +3,7 @@
 #pragma once
 
 #include "CoreTypes.h"
+#include "Containers/Map.h"
 #include "Features/IModularFeature.h"
 #include "Delegates/Delegate.h"
 
@@ -28,7 +29,7 @@ public:
 	virtual void Start(const wchar_t* ProcessGroupName) = 0;
 	virtual void Stop() = 0;
 
-	virtual void SetLinkerPath(const wchar_t* LinkerPath) = 0;
+	virtual void SetLinkerPath(const wchar_t* LinkerPath, const TMap<FString, FString>& LinkerEnvironment) = 0;
 
 	DECLARE_DELEGATE(FBringToFrontDelegate);
 	virtual FBringToFrontDelegate& GetBringToFrontDelegate() = 0;
diff --git a/Engine/Source/Programs/LiveCodingConsole/Private/LiveCodingConsole.cpp b/Engine/Source/Programs/LiveCodingConsole/Private/LiveCodingConsole.cpp
--- a/Engine/Source/Programs/LiveCodingConsole/Private/LiveCodingConsole.cpp
+++ b/Engine/Source/Programs/LiveCodingConsole/Private/LiveCodingConsole.cpp
@@ -230,7 +230,7 @@ private:
 		}
 
 		// Override the linker path
-		Server.SetLinkerPath(*Manifest.LinkerPath);
+		Server.SetLinkerPath(*Manifest.LinkerPath, Manifest.LinkerEnvironment);
 
 		// Strip out all the files that haven't been modified
 		IFileManager& FileManager = IFileManager::Get();
diff --git a/Engine/Source/Programs/LiveCodingConsole/Private/LiveCodingManifest.cpp b/Engine/Source/Programs/LiveCodingConsole/Private/LiveCodingManifest.cpp
--- a/Engine/Source/Programs/LiveCodingConsole/Private/LiveCodingManifest.cpp
+++ b/Engine/Source/Programs/LiveCodingConsole/Private/LiveCodingManifest.cpp
@@ -40,6 +40,19 @@ bool FLiveCodingManifest::Parse(FJsonObject& Object, FString& OutFailReason)
 		return false;
 	}
 
+	const TSharedPtr<FJsonObject>* EnvironmentObject;
+	if (Object.TryGetObjectField(TEXT("LinkerEnvironment"), EnvironmentObject))
+	{
+		for(const TPair<FString, TSharedPtr<FJsonValue>>& Pair : EnvironmentObject->Get()->Values)
+		{
+			const FJsonValue* Value = Pair.Value.Get();
+			if (Value->Type == EJson::String)
+			{
+				LinkerEnvironment.Add(Pair.Key, Value->AsString());
+			}
+		}
+	}
+
 	const TArray<TSharedPtr<FJsonValue>>* ModulesArray;
 	if (!Object.TryGetArrayField(TEXT("Modules"), ModulesArray))
 	{
diff --git a/Engine/Source/Programs/LiveCodingConsole/Private/LiveCodingManifest.h b/Engine/Source/Programs/LiveCodingConsole/Private/LiveCodingManifest.h
--- a/Engine/Source/Programs/LiveCodingConsole/Private/LiveCodingManifest.h
+++ b/Engine/Source/Programs/LiveCodingConsole/Private/LiveCodingManifest.h
@@ -9,6 +9,7 @@
 struct FLiveCodingManifest
 {
 	FString LinkerPath;
+	TMap<FString, FString> LinkerEnvironment;
 	TMap<FString, TArray<FString>> BinaryToObjectFiles;
 
 	bool Read(const TCHAR* FileName, FString& OutFailReason);
diff --git a/Engine/Source/Programs/UnrealBuildTool/System/HotReload.cs b/Engine/Source/Programs/UnrealBuildTool/System/HotReload.cs
--- a/Engine/Source/Programs/UnrealBuildTool/System/HotReload.cs
+++ b/Engine/Source/Programs/UnrealBuildTool/System/HotReload.cs
@@ -751,6 +751,13 @@ namespace UnrealBuildTool
 					Writer.WriteValue("LinkerPath", LinkAction.CommandPath.FullName);
 				}
 
+				Writer.WriteObjectStart("LinkerEnvironment");
+				foreach (System.Collections.DictionaryEntry Entry in Environment.GetEnvironmentVariables())
+				{
+					Writer.WriteValue(Entry.Key.ToString(), Entry.Value.ToString());
+				}
+				Writer.WriteObjectEnd();
+
 				Writer.WriteArrayStart("Modules");
 				foreach(Action Action in Actions)
 				{
