Integrating changes from Live++ 1.4.1.

diff --git a/Engine/Source/Developer/Windows/LiveCoding/Private/External/LC_ClientCommandActions.cpp b/Engine/Source/Developer/Windows/LiveCoding/Private/External/LC_ClientCommandActions.cpp
--- a/Engine/Source/Developer/Windows/LiveCoding/Private/External/LC_ClientCommandActions.cpp
+++ b/Engine/Source/Developer/Windows/LiveCoding/Private/External/LC_ClientCommandActions.cpp
@@ -10,7 +10,7 @@
 #include "LC_Logging.h"
 
 
-bool actions::RegisterProcessFinished::Execute(CommandType* command, const DuplexPipe* pipe, void* context)
+bool actions::RegisterProcessFinished::Execute(const CommandType* command, const DuplexPipe* pipe, void* context)
 {
 	pipe->SendAck();
 
@@ -22,7 +22,7 @@ bool actions::RegisterProcessFinished::Execute(CommandType* command, const Duple
 }
 
 
-bool actions::EnableModuleFinished::Execute(CommandType* command, const DuplexPipe* pipe, void*)
+bool actions::EnableModuleFinished::Execute(const CommandType* command, const DuplexPipe* pipe, void*)
 {
 	Event* event = static_cast<Event*>(command->token);
 	event->Signal();
@@ -32,7 +32,7 @@ bool actions::EnableModuleFinished::Execute(CommandType* command, const DuplexPi
 }
 
 
-bool actions::DisableModuleFinished::Execute(CommandType* command, const DuplexPipe* pipe, void*)
+bool actions::DisableModuleFinished::Execute(const CommandType* command, const DuplexPipe* pipe, void*)
 {
 	Event* event = static_cast<Event*>(command->token);
 	event->Signal();
@@ -42,7 +42,7 @@ bool actions::DisableModuleFinished::Execute(CommandType* command, const DuplexP
 }
 
 
-bool actions::EnableAllModulesFinished::Execute(CommandType* command, const DuplexPipe* pipe, void*)
+bool actions::EnableAllModulesFinished::Execute(const CommandType* command, const DuplexPipe* pipe, void*)
 {
 	Event* event = static_cast<Event*>(command->token);
 	event->Signal();
@@ -52,7 +52,7 @@ bool actions::EnableAllModulesFinished::Execute(CommandType* command, const Dupl
 }
 
 
-bool actions::DisableAllModulesFinished::Execute(CommandType* command, const DuplexPipe* pipe, void*)
+bool actions::DisableAllModulesFinished::Execute(const CommandType* command, const DuplexPipe* pipe, void*)
 {
 	Event* event = static_cast<Event*>(command->token);
 	event->Signal();
@@ -62,7 +62,7 @@ bool actions::DisableAllModulesFinished::Execute(CommandType* command, const Dup
 }
 
 
-bool actions::EnterSyncPoint::Execute(CommandType*, const DuplexPipe* pipe, void*)
+bool actions::EnterSyncPoint::Execute(const CommandType*, const DuplexPipe* pipe, void*)
 {
 	syncPoint::Enter();
 	pipe->SendAck();
@@ -71,7 +71,7 @@ bool actions::EnterSyncPoint::Execute(CommandType*, const DuplexPipe* pipe, void
 }
 
 
-bool actions::LeaveSyncPoint::Execute(CommandType*, const DuplexPipe* pipe, void*)
+bool actions::LeaveSyncPoint::Execute(const CommandType*, const DuplexPipe* pipe, void*)
 {
 	syncPoint::Leave();
 	pipe->SendAck();
@@ -80,7 +80,7 @@ bool actions::LeaveSyncPoint::Execute(CommandType*, const DuplexPipe* pipe, void
 }
 
 
-bool actions::CallHooks::Execute(CommandType* command, const DuplexPipe* pipe, void*)
+bool actions::CallHooks::Execute(const CommandType* command, const DuplexPipe* pipe, void*)
 {
 	for (const hook::Function* hook = command->first; hook < command->last; ++hook)
 	{
@@ -98,7 +98,7 @@ bool actions::CallHooks::Execute(CommandType* command, const DuplexPipe* pipe, v
 }
 
 
-bool actions::GetModule::Execute(CommandType* command, const DuplexPipe* pipe, void*)
+bool actions::GetModule::Execute(const CommandType* command, const DuplexPipe* pipe, void*)
 {
 	HMODULE module = ::GetModuleHandleW(command->path);
 	if (module)
@@ -136,7 +136,7 @@ extern FChunkedFixedUObjectArray*& GObjectArrayForDebugVisualizers;
 // END EPIC MOD
 
 
-bool actions::LoadPatch::Execute(CommandType* command, const DuplexPipe* pipe, void*)
+bool actions::LoadPatch::Execute(const CommandType* command, const DuplexPipe* pipe, void*)
 {
 	// load library into this process
 	HMODULE module = ::LoadLibraryW(command->path);
@@ -163,7 +163,7 @@ bool actions::LoadPatch::Execute(CommandType* command, const DuplexPipe* pipe, v
 }
 
 
-bool actions::UnloadPatch::Execute(CommandType* command, const DuplexPipe* pipe, void*)
+bool actions::UnloadPatch::Execute(const CommandType* command, const DuplexPipe* pipe, void*)
 {
 	// unload library from this process
 	::FreeLibrary(command->module);
@@ -173,7 +173,7 @@ bool actions::UnloadPatch::Execute(CommandType* command, const DuplexPipe* pipe,
 }
 
 
-bool actions::CallEntryPoint::Execute(CommandType* command, const DuplexPipe* pipe, void*)
+bool actions::CallEntryPoint::Execute(const CommandType* command, const DuplexPipe* pipe, void*)
 {
 	executable::CallDllEntryPoint(command->moduleBase, command->entryPointRva);
 	pipe->SendAck();
@@ -182,7 +182,7 @@ bool actions::CallEntryPoint::Execute(CommandType* command, const DuplexPipe* pi
 }
 
 
-bool actions::LogOutput::Execute(CommandType* command, const DuplexPipe* pipe, void*)
+bool actions::LogOutput::Execute(const CommandType* command, const DuplexPipe* pipe, void*)
 {
 	logging::LogNoFormat<logging::Channel::USER>(command->buffer);
 	pipe->SendAck();
@@ -194,7 +194,7 @@ bool actions::LogOutput::Execute(CommandType* command, const DuplexPipe* pipe, v
 extern bool GIsCompileActive;
 // END EPIC MOD
 
-bool actions::CompilationFinished::Execute(CommandType*, const DuplexPipe* pipe, void*)
+bool actions::CompilationFinished::Execute(const CommandType*, const DuplexPipe* pipe, void*)
 {
 	pipe->SendAck();
 
@@ -207,7 +207,7 @@ bool actions::CompilationFinished::Execute(CommandType*, const DuplexPipe* pipe,
 }
 
 
-bool actions::HandleExceptionFinished::Execute(CommandType* command, const DuplexPipe* pipe, void* context)
+bool actions::HandleExceptionFinished::Execute(const CommandType* command, const DuplexPipe* pipe, void* context)
 {
 	ClientUserCommandThread::ExceptionResult* resultContext = static_cast<ClientUserCommandThread::ExceptionResult*>(context);
 	resultContext->returnAddress = command->returnAddress;
diff --git a/Engine/Source/Developer/Windows/LiveCoding/Private/External/LC_ClientCommandActions.h b/Engine/Source/Developer/Windows/LiveCoding/Private/External/LC_ClientCommandActions.h
--- a/Engine/Source/Developer/Windows/LiveCoding/Private/External/LC_ClientCommandActions.h
+++ b/Engine/Source/Developer/Windows/LiveCoding/Private/External/LC_ClientCommandActions.h
@@ -10,11 +10,11 @@ class DuplexPipe;
 
 namespace actions
 {
-	#define DECLARE_ACTION(_name)																	\
-		struct _name																				\
-		{																							\
-			typedef ::commands::_name CommandType;													\
-			static bool Execute(CommandType* command, const DuplexPipe* pipe, void* context);		\
+	#define DECLARE_ACTION(_name)																		\
+		struct _name																					\
+		{																								\
+			typedef ::commands::_name CommandType;														\
+			static bool Execute(const CommandType* command, const DuplexPipe* pipe, void* context);		\
 		}
 
 	DECLARE_ACTION(RegisterProcessFinished);
diff --git a/Engine/Source/Developer/Windows/LiveCoding/Private/External/LC_ClientStartupThread.cpp b/Engine/Source/Developer/Windows/LiveCoding/Private/External/LC_ClientStartupThread.cpp
--- a/Engine/Source/Developer/Windows/LiveCoding/Private/External/LC_ClientStartupThread.cpp
+++ b/Engine/Source/Developer/Windows/LiveCoding/Private/External/LC_ClientStartupThread.cpp
@@ -445,7 +445,7 @@ unsigned int ClientStartupThread::ThreadFunction(const std::wstring& processGrou
 	m_userCommandThread->Start(processGroupName, m_startEvent, m_pipeClientCS);
 
 	// register this process with Live++
-	m_pipeClient->SendCommandAndWaitForAck(commands::RegisterProcess { process::GetId(), process::GetBase(), commandThreadId });
+	m_pipeClient->SendCommandAndWaitForAck(commands::RegisterProcess { process::GetBase(), process::GetId(), commandThreadId });
 
 	// handle commands until registration is finished
 	{
diff --git a/Engine/Source/Developer/Windows/LiveCoding/Private/External/LC_ClientUserCommandThread.cpp b/Engine/Source/Developer/Windows/LiveCoding/Private/External/LC_ClientUserCommandThread.cpp
--- a/Engine/Source/Developer/Windows/LiveCoding/Private/External/LC_ClientUserCommandThread.cpp
+++ b/Engine/Source/Developer/Windows/LiveCoding/Private/External/LC_ClientUserCommandThread.cpp
@@ -47,6 +47,8 @@ namespace userCommands
 	{
 		EnableModuleCommand(void)
 			: BaseCommand(Scope::ENABLE_MODULES)
+			, token(nullptr)
+			, moduleName()
 		{
 		}
 
@@ -70,6 +72,8 @@ namespace userCommands
 	{
 		EnableAllModulesCommand(void)
 			: BaseCommand(Scope::ENABLE_MODULES)
+			, token(nullptr)
+			, moduleName()
 		{
 		}
 
@@ -93,6 +97,8 @@ namespace userCommands
 	{
 		DisableModuleCommand(void)
 			: BaseCommand(Scope::DISABLE_MODULES)
+			, token(nullptr)
+			, moduleName()
 		{
 		}
 
@@ -116,6 +122,8 @@ namespace userCommands
 	{
 		DisableAllModulesCommand(void)
 			: BaseCommand(Scope::DISABLE_MODULES)
+			, token(nullptr)
+			, moduleName()
 		{
 		}
 
@@ -262,6 +270,9 @@ namespace userCommands
 	{
 		BuildPatchCommand(void)
 			: BaseCommand(Scope::NONE)
+			, count(0u)
+			, moduleNames()
+			, objPaths()
 		{
 		}
 
@@ -292,6 +303,8 @@ namespace userCommands
 	{
 		ApplySettingBoolCommand(void)
 			: BaseCommand(Scope::NONE)
+			, settingName()
+			, value(0)
 		{
 		}
 
@@ -313,6 +326,8 @@ namespace userCommands
 	{
 		ApplySettingIntCommand(void)
 			: BaseCommand(Scope::NONE)
+			, settingName()
+			, value(0)
 		{
 		}
 
@@ -334,6 +349,8 @@ namespace userCommands
 	{
 		ApplySettingStringCommand(void)
 			: BaseCommand(Scope::NONE)
+			, settingName()
+			, value()
 		{
 		}
 
diff --git a/Engine/Source/Developer/Windows/LiveCoding/Private/External/LC_CriticalSection.h b/Engine/Source/Developer/Windows/LiveCoding/Private/External/LC_CriticalSection.h
--- a/Engine/Source/Developer/Windows/LiveCoding/Private/External/LC_CriticalSection.h
+++ b/Engine/Source/Developer/Windows/LiveCoding/Private/External/LC_CriticalSection.h
@@ -4,6 +4,7 @@
 
 #include "CoreTypes.h"
 #include "Windows/MinimalWindowsApi.h"
+#include "LC_Platform.h"
 
 class CriticalSection
 {
@@ -29,6 +30,11 @@ public:
 		}
 
 	private:
+		LC_DISABLE_COPY(ScopedLock);
+		LC_DISABLE_MOVE(ScopedLock);
+		LC_DISABLE_ASSIGNMENT(ScopedLock);
+		LC_DISABLE_MOVE_ASSIGNMENT(ScopedLock);
+
 		CriticalSection* m_cs;
 	};
 
diff --git a/Engine/Source/Developer/Windows/LiveCoding/Private/External/LC_Executable.cpp b/Engine/Source/Developer/Windows/LiveCoding/Private/External/LC_Executable.cpp
--- a/Engine/Source/Developer/Windows/LiveCoding/Private/External/LC_Executable.cpp
+++ b/Engine/Source/Developer/Windows/LiveCoding/Private/External/LC_Executable.cpp
@@ -26,14 +26,14 @@ namespace detail
 		void* base = image->base;
 
 		// PE image start with a DOS header
-		IMAGE_DOS_HEADER* dosHeader = static_cast<IMAGE_DOS_HEADER*>(base);
+		const IMAGE_DOS_HEADER* dosHeader = static_cast<IMAGE_DOS_HEADER*>(base);
 		if (dosHeader->e_magic != IMAGE_DOS_SIGNATURE)
 		{
 			LC_ERROR_USER("Image has unknown file format");
 			return nullptr;
 		}
 
-		IMAGE_NT_HEADERS* ntHeader = pointer::Offset<IMAGE_NT_HEADERS*>(dosHeader, dosHeader->e_lfanew);
+		const IMAGE_NT_HEADERS* ntHeader = pointer::Offset<const IMAGE_NT_HEADERS*>(dosHeader, dosHeader->e_lfanew);
 		if (ntHeader->Signature != IMAGE_NT_SIGNATURE)
 		{
 			LC_ERROR_USER("Invalid .exe file");
@@ -208,7 +208,7 @@ void executable::RebaseImage(Image* image, PreferredBase preferredBase)
 
 			// PE spec: Block size: The total number of bytes in the base relocation block, *including* the Page RVA and
 			// Block Size fields and the Type/Offset fields that follow.
-			DWORD numberOfEntriesInThisBlock = (blockSize - sizeof(IMAGE_BASE_RELOCATION)) / sizeof(WORD);
+			const DWORD numberOfEntriesInThisBlock = (blockSize - sizeof(IMAGE_BASE_RELOCATION)) / sizeof(WORD);
 			const WORD* entries = pointer::Offset<const WORD*>(baseRelocations, sizeof(IMAGE_BASE_RELOCATION));
 			for (DWORD i = 0u; i < numberOfEntriesInThisBlock; ++i)
 			{
diff --git a/Engine/Source/Developer/Windows/LiveCoding/Private/External/LC_InterprocessMutex.h b/Engine/Source/Developer/Windows/LiveCoding/Private/External/LC_InterprocessMutex.h
--- a/Engine/Source/Developer/Windows/LiveCoding/Private/External/LC_InterprocessMutex.h
+++ b/Engine/Source/Developer/Windows/LiveCoding/Private/External/LC_InterprocessMutex.h
@@ -28,6 +28,11 @@ public:
 		}
 
 	private:
+		LC_DISABLE_COPY(ScopedLock);
+		LC_DISABLE_MOVE(ScopedLock);
+		LC_DISABLE_ASSIGNMENT(ScopedLock);
+		LC_DISABLE_MOVE_ASSIGNMENT(ScopedLock);
+
 		InterprocessMutex* m_mutex;
 	};
 
diff --git a/Engine/Source/Developer/Windows/LiveCoding/Private/External/LC_Process.cpp b/Engine/Source/Developer/Windows/LiveCoding/Private/External/LC_Process.cpp
--- a/Engine/Source/Developer/Windows/LiveCoding/Private/External/LC_Process.cpp
+++ b/Engine/Source/Developer/Windows/LiveCoding/Private/External/LC_Process.cpp
@@ -634,11 +634,13 @@ namespace process
 			if (pointer::Offset<const void*>(scan, size) >= upperBound)
 			{
 				// outside of range to scan
+				LC_ERROR_DEV("Could not find memory range that fits 0x%X bytes with alignment 0x%X in range from 0x%p to 0x%p (scan: 0x%p)", size, alignment, lowerBound, upperBound, scan);
 				return nullptr;
 			}
 			else if (scan < lowerBound)
 			{
 				// outside of range (possible wrap-around)
+				LC_ERROR_DEV("Could not find memory range that fits 0x%X bytes with alignment 0x%X in range from 0x%p to 0x%p (scan: 0x%p)", size, alignment, lowerBound, upperBound, scan);
 				return nullptr;
 			}
 
@@ -737,6 +739,8 @@ namespace process
 			{
 				// buffer is too small, try again
 				::free(processSnapshot);
+				processSnapshot = nullptr;
+
 				bufferSize *= 2u;
 			}
 			else if (status < 0)
@@ -744,7 +748,7 @@ namespace process
 				// something went wrong
 				LC_ERROR_USER("Cannot enumerate threads in process (PID: %d)", processId);
 				::free(processSnapshot);
-				
+
 				return threadIds;
 			}
 		}
@@ -778,10 +782,13 @@ namespace process
 			}
 
 			// record all threads belonging to the given process
-			for (ULONG i = 0u; i < processInfo->uThreadCount; ++i)
+			if (processInfo)
 			{
-				const DWORD threadId = static_cast<DWORD>(reinterpret_cast<DWORD_PTR>(processInfo->Threads[i].ClientId.UniqueThread));
-				threadIds.push_back(threadId);
+				for (ULONG i = 0u; i < processInfo->uThreadCount; ++i)
+				{
+					const DWORD threadId = static_cast<DWORD>(reinterpret_cast<DWORD_PTR>(processInfo->Threads[i].ClientId.UniqueThread));
+					threadIds.push_back(threadId);
+				}
 			}
 		}
 
diff --git a/Engine/Source/Developer/Windows/LiveCodingServer/Private/External/LC_AppSettings.cpp b/Engine/Source/Developer/Windows/LiveCodingServer/Private/External/LC_AppSettings.cpp
--- a/Engine/Source/Developer/Windows/LiveCodingServer/Private/External/LC_AppSettings.cpp
+++ b/Engine/Source/Developer/Windows/LiveCodingServer/Private/External/LC_AppSettings.cpp
@@ -30,8 +30,9 @@ namespace
 
 	static std::wstring g_cachedCompilerPath;
 	static std::wstring g_cachedLinkerPath;
+	static types::vector<std::wstring> g_cachedAmalgamatedCppFileExtensions;
 
-	static std::wstring DeterminePath(SettingString* setting, const wchar_t* const type, const wchar_t* const vs2017Path, const wchar_t* const vs2015AndEarlierPath, const wchar_t* exeName)
+	static std::wstring DeterminePath(const SettingString* setting, const wchar_t* const type, const wchar_t* const vs2017Path, const wchar_t* const vs2015AndEarlierPath, const wchar_t* exeName)
 	{
 		// absolute paths can be used as they are, relative paths are supposed to be relative to the Live++ executable
 		std::wstring path;
@@ -453,6 +454,15 @@ void appSettings::Startup(const wchar_t* group)
 		L"Minimum number of .cpp files that must be included in an amalgamated/unity file before it is split",
 		3
 	);	
+
+	g_amalgamationCppFileExtensions = new SettingString
+	(
+		group,
+		L"amalgamation_cpp_file_extensions",
+		L"C/C++ file extensions",
+		L"File extensions treated as C/C++ files when splitting amalgamated/unity files",
+		L".cpp;.c;.cc;.c++;.cp;.cxx"
+	);
 }
 
 
@@ -502,6 +512,7 @@ void appSettings::Shutdown(void)
 
 	delete g_amalgamationSplitIntoSingleParts;
 	delete g_amalgamationSplitMinCppCount;
+	delete g_amalgamationCppFileExtensions;
 }
 
 
@@ -571,6 +582,12 @@ std::wstring appSettings::GetLinkerPath(void)
 }
 
 
+const types::vector<std::wstring>& appSettings::GetAmalgamatedCppFileExtensions(void)
+{
+	return g_cachedAmalgamatedCppFileExtensions;
+}
+
+
 void appSettings::UpdateCompilerPathCache(void)
 {
 	g_cachedCompilerPath = DeterminePath(g_compilerPath, L"compiler", VS2017_COMPILER_PATH, VS2015_AND_EARLIER_COMPILER_PATH, COMPILER_EXE);
@@ -590,6 +607,27 @@ void appSettings::UpdatePathCache(void)
 }
 
 
+void appSettings::UpdateAmalgamatedCppFileExtensions(void)
+{
+	g_cachedAmalgamatedCppFileExtensions.clear();
+	g_cachedAmalgamatedCppFileExtensions.reserve(16u);
+
+	const wchar_t DELIMITER = L';';
+	const std::wstring extensions = g_amalgamationCppFileExtensions->GetValue();
+	size_t start = extensions.find_first_not_of(DELIMITER);
+	size_t end = start;
+
+	while (start != std::string::npos)
+	{
+		// first find the next occurrence of the delimiter, then push the token into the vector,
+		// then skip all occurrences of the delimiter until we find the start of the next token
+		end = extensions.find(DELIMITER, start);
+		g_cachedAmalgamatedCppFileExtensions.push_back(extensions.substr(start, end - start));
+		start = extensions.find_first_not_of(DELIMITER, end);
+	}
+}
+
+
 void appSettings::ApplySettingBool(const char* const settingName, bool value)
 {
 	const unsigned int COUNT = 20u;
@@ -670,7 +708,7 @@ void appSettings::ApplySettingString(const char* const settingName, const wchar_
 {
 	// try string settings first
 	{
-		const unsigned int COUNT = 9u;
+		const unsigned int COUNT = 10u;
 		SettingString* settings[COUNT] =
 		{
 			g_playSoundOnSuccess,
@@ -681,7 +719,8 @@ void appSettings::ApplySettingString(const char* const settingName, const wchar_
 			g_linkerOptions,
 			g_continuousCompilationPath,
 			g_virtualDriveLetter,
-			g_virtualDrivePath
+			g_virtualDrivePath,
+			g_amalgamationCppFileExtensions
 		};
 
 		const SettingString* setting = ApplySetting(settings, COUNT, settingName, value);
@@ -697,12 +736,16 @@ void appSettings::ApplySettingString(const char* const settingName, const wchar_
 				UpdateLinkerPathCache();
 			}
 
+			// update cache when amalgamated C++ file extensions change
+			else if (setting == g_amalgamationCppFileExtensions)
+			{
+				UpdateAmalgamatedCppFileExtensions();
+			}
+
 			return;
 		}
 	}
 
-	const std::wstring wideSettingName(string::ToWideString(settingName));
-
 	// try proxies second
 	{
 		const unsigned int COUNT = 2u;
@@ -778,3 +821,4 @@ extern SettingBool* appSettings::g_installCompiledPatchesMultiProcess = nullptr;
 
 extern SettingBool* appSettings::g_amalgamationSplitIntoSingleParts = nullptr;
 extern SettingInt* appSettings::g_amalgamationSplitMinCppCount = nullptr;
+extern SettingString* appSettings::g_amalgamationCppFileExtensions = nullptr;
diff --git a/Engine/Source/Developer/Windows/LiveCodingServer/Private/External/LC_AppSettings.h b/Engine/Source/Developer/Windows/LiveCodingServer/Private/External/LC_AppSettings.h
--- a/Engine/Source/Developer/Windows/LiveCodingServer/Private/External/LC_AppSettings.h
+++ b/Engine/Source/Developer/Windows/LiveCodingServer/Private/External/LC_AppSettings.h
@@ -40,10 +40,15 @@ namespace appSettings
 	// returns the overridden linker path
 	std::wstring GetLinkerPath(void);
 
+	// returns an array of file extensions to treat as C++ files when splitting amalgamated files
+	const types::vector<std::wstring>& GetAmalgamatedCppFileExtensions(void);
+
 	void UpdateCompilerPathCache(void);
 	void UpdateLinkerPathCache(void);
 	void UpdatePathCache(void);
 
+	void UpdateAmalgamatedCppFileExtensions(void);
+
 	// apply a new value to any of the boolean settings
 	void ApplySettingBool(const char* const settingName, bool value);
 
@@ -79,14 +84,14 @@ namespace appSettings
 	extern SettingBool* g_enableDevLogCompilands;
 
 	// compiler
-	extern SettingString* g_compilerPath;					// DO NOT USE directly, use GetCompilerPath instead!
+	extern SettingString* g_compilerPath;						// DO NOT USE directly, use GetCompilerPath instead!
 	extern SettingBool* g_useCompilerOverrideAsFallback;
 	extern SettingBool* g_useCompilerEnvironment;
 	extern SettingString* g_compilerOptions;
 	extern SettingBool* g_compilerForcePchPdbs;
 
 	// linker
-	extern SettingString* g_linkerPath;						// DO NOT USE directly, use GetLinkerPath instead!
+	extern SettingString* g_linkerPath;							// DO NOT USE directly, use GetLinkerPath instead!
 	extern SettingBool* g_useLinkerOverrideAsFallback;
 	extern SettingBool* g_useLinkerEnvironment;
 	extern SettingString* g_linkerOptions;
@@ -107,4 +112,5 @@ namespace appSettings
 	// amalgamated/unity builds
 	extern SettingBool* g_amalgamationSplitIntoSingleParts;
 	extern SettingInt* g_amalgamationSplitMinCppCount;
+	extern SettingString* g_amalgamationCppFileExtensions;		// DO NOT USE directly, use GetAmalgamatedCppFileExtensions instead!
 }
diff --git a/Engine/Source/Developer/Windows/LiveCodingServer/Private/External/LC_Environment.cpp b/Engine/Source/Developer/Windows/LiveCodingServer/Private/External/LC_Environment.cpp
--- a/Engine/Source/Developer/Windows/LiveCodingServer/Private/External/LC_Environment.cpp
+++ b/Engine/Source/Developer/Windows/LiveCodingServer/Private/External/LC_Environment.cpp
@@ -61,10 +61,17 @@ namespace environment
 			const char* start = memory;
 
 			// search for carriage return
-			while (memory[0] != '\r')
+			while ((memory < memoryEnd) && (memory[0] != '\r'))
 			{
 				++memory;
 			}
+
+			if (memory >= memoryEnd)
+			{
+				// reached the end
+				break;
+			}
+
 			const char* end = memory;
 
 			std::string line(start, end);
@@ -83,7 +90,7 @@ namespace environment
 			}
 
 			// skip carriage return and new line
-			while ((memory[0] == '\r') || (memory[0] == '\n'))
+			while ((memory < memoryEnd) && ((memory[0] == '\r') || (memory[0] == '\n')))
 			{
 				++memory;
 			}
@@ -162,7 +169,11 @@ namespace environment
 		value.resize(MAX_PATH);
 
 		size_t bytes = 0u;
-		_wgetenv_s(&bytes, &value[0], value.size(), variable);
+		const errno_t error = _wgetenv_s(&bytes, &value[0], value.size(), variable);
+		if (error != 0)
+		{
+			LC_LOG_DEV("Could not retrieve environment variable %S (Error: %d)", variable, error);
+		}
 
 		return value;
 	}
diff --git a/Engine/Source/Developer/Windows/LiveCodingServer/Private/External/LC_FileUtil.cpp b/Engine/Source/Developer/Windows/LiveCodingServer/Private/External/LC_FileUtil.cpp
--- a/Engine/Source/Developer/Windows/LiveCodingServer/Private/External/LC_FileUtil.cpp
+++ b/Engine/Source/Developer/Windows/LiveCodingServer/Private/External/LC_FileUtil.cpp
@@ -73,12 +73,59 @@ namespace detail
 namespace
 {
 	static detail::NormalizedFilenameCache* g_normalizedFilenameCache = nullptr;
+
+	static file::DriveType::Enum g_driveTypeCache['z' - 'a' + 1u] = {};
 }
 
 
 void file::Startup(void)
 {
 	g_normalizedFilenameCache = new detail::NormalizedFilenameCache;
+
+	// fill cache of drive types
+	char root[4u] = { 'a', ':', '\\', '\0' };
+	for (char drive = 'a'; drive <= 'z'; ++drive)
+	{
+		const int index = drive - 'a';
+		
+		root[0] = drive;
+		const UINT driveType = ::GetDriveTypeA(root);
+
+		switch (driveType)
+		{
+			case DRIVE_UNKNOWN:
+				g_driveTypeCache[index] = DriveType::UNKNOWN;
+				break;
+
+			case DRIVE_NO_ROOT_DIR:
+				g_driveTypeCache[index] = DriveType::UNKNOWN;
+				break;
+
+			case DRIVE_REMOVABLE:
+				g_driveTypeCache[index] = DriveType::REMOVABLE;
+				break;
+
+			case DRIVE_FIXED:
+				g_driveTypeCache[index] = DriveType::FIXED;
+				break;
+
+			case DRIVE_REMOTE:
+				g_driveTypeCache[index] = DriveType::REMOTE;
+				break;
+
+			case DRIVE_CDROM:
+				g_driveTypeCache[index] = DriveType::OPTICAL;
+				break;
+
+			case DRIVE_RAMDISK:
+				g_driveTypeCache[index] = DriveType::RAMDISK;
+				break;
+
+			default:
+				g_driveTypeCache[index] = DriveType::UNKNOWN;
+				break;
+		}
+	}
 }
 
 
@@ -88,6 +135,37 @@ void file::Shutdown(void)
 }
 
 
+file::DriveType::Enum file::GetDriveType(char driveLetter)
+{
+	if ((driveLetter >= 'a') && (driveLetter <= 'z'))
+	{
+		return g_driveTypeCache[driveLetter - 'a'];
+	}
+	else if ((driveLetter >= 'A') && (driveLetter <= 'Z'))
+	{
+		return g_driveTypeCache[driveLetter - 'A'];
+	}
+
+	return DriveType::UNKNOWN;
+}
+
+
+file::DriveType::Enum file::GetDriveType(const wchar_t* path)
+{
+	const wchar_t driveLetter = path[0];
+	if ((driveLetter >= L'a') && (driveLetter <= L'z'))
+	{
+		return g_driveTypeCache[driveLetter - L'a'];
+	}
+	else if ((driveLetter >= L'A') && (driveLetter <= L'Z'))
+	{
+		return g_driveTypeCache[driveLetter - L'A'];
+	}
+
+	return DriveType::UNKNOWN;
+}
+
+
 file::Attributes file::GetAttributes(const wchar_t* path)
 {
 	const WIN32_FILE_ATTRIBUTE_DATA initialData = { INVALID_FILE_ATTRIBUTES };
@@ -170,7 +248,7 @@ bool file::IsRelativePath(const wchar_t* path)
 		return false;
 	}
 
-	return (::PathIsRelativeW(path) == Windows::TRUE);
+	return (::PathIsRelativeW(path) != Windows::FALSE);
 }
 
 
diff --git a/Engine/Source/Developer/Windows/LiveCodingServer/Private/External/LC_FileUtil.h b/Engine/Source/Developer/Windows/LiveCodingServer/Private/External/LC_FileUtil.h
--- a/Engine/Source/Developer/Windows/LiveCodingServer/Private/External/LC_FileUtil.h
+++ b/Engine/Source/Developer/Windows/LiveCodingServer/Private/External/LC_FileUtil.h
@@ -13,9 +13,25 @@ namespace file
 		WIN32_FILE_ATTRIBUTE_DATA data;
 	};
 
+	struct DriveType
+	{
+		enum Enum
+		{
+			UNKNOWN = 0,
+			REMOVABLE,
+			FIXED,
+			REMOTE,
+			OPTICAL,
+			RAMDISK
+		};
+	};
+
 	void Startup(void);
 	void Shutdown(void);
 
+	// Returns the type of a drive. Drive letters can be both lower- or upper-case
+	DriveType::Enum GetDriveType(char driveLetter);
+	DriveType::Enum GetDriveType(const wchar_t* path);
 
 	Attributes GetAttributes(const wchar_t* path);
 	uint64_t GetLastModificationTime(const Attributes& attributes);
diff --git a/Engine/Source/Developer/Windows/LiveCodingServer/Private/External/LC_GrowingMemoryBlock.cpp b/Engine/Source/Developer/Windows/LiveCodingServer/Private/External/LC_GrowingMemoryBlock.cpp
--- a/Engine/Source/Developer/Windows/LiveCodingServer/Private/External/LC_GrowingMemoryBlock.cpp
+++ b/Engine/Source/Developer/Windows/LiveCodingServer/Private/External/LC_GrowingMemoryBlock.cpp
@@ -35,12 +35,6 @@ bool GrowingMemoryBlock::Insert(const void* data, size_t size)
 	{
 		// grow to accommodate the new size, copy over the existing data, and delete the old storage
 		uint8_t* newData = new uint8_t[m_capacity];
-		if (!newData)
-		{
-			// could not allocate memory, bail out
-			return false;
-		}
-
 		memcpy(newData, m_data, m_size);
 		delete[] m_data;
 		m_data = newData;
diff --git a/Engine/Source/Developer/Windows/LiveCodingServer/Private/External/LC_LiveModule.cpp b/Engine/Source/Developer/Windows/LiveCodingServer/Private/External/LC_LiveModule.cpp
--- a/Engine/Source/Developer/Windows/LiveCodingServer/Private/External/LC_LiveModule.cpp
+++ b/Engine/Source/Developer/Windows/LiveCodingServer/Private/External/LC_LiveModule.cpp
@@ -167,20 +167,20 @@ namespace
 	{
 		// MSVC's link.exe is much more common, so treat this as our default
 		const std::wstring linkerPath = GetLinkerPath(linkerDb);
-		const std::wstring lowerCaseFilename = string::ToLower(file::GetFilename(linkerPath));
-		if (string::Contains(linkerPath.c_str(), L"lld"))
+		const std::wstring lowerCaseLinkerPath = string::ToLower(file::GetFilename(linkerPath));
+		if (string::Contains(lowerCaseLinkerPath.c_str(), L"lld"))
 		{
 			return coff::SymbolRemovalStrategy::LLD_COMPATIBLE;
 		}
-		else if (string::Contains(linkerPath.c_str(), L"lld-link"))
+		else if (string::Contains(lowerCaseLinkerPath.c_str(), L"lld-link"))
 		{
 			return coff::SymbolRemovalStrategy::LLD_COMPATIBLE;
 		}
-		else if (string::Contains(linkerPath.c_str(), L"ld.lld"))
+		else if (string::Contains(lowerCaseLinkerPath.c_str(), L"ld.lld"))
 		{
 			return coff::SymbolRemovalStrategy::LLD_COMPATIBLE;
 		}
-		else if (string::Contains(linkerPath.c_str(), L"ld64.lld"))
+		else if (string::Contains(lowerCaseLinkerPath.c_str(), L"ld64.lld"))
 		{
 			return coff::SymbolRemovalStrategy::LLD_COMPATIBLE;
 		}
@@ -521,7 +521,10 @@ namespace
 		const size_t MODULE_ALIGNMENT = 64u * 1024u;
 		void* preferredBase = process::ScanMemoryRange(processHandle, lowerBound, upperBound, patchSize, MODULE_ALIGNMENT);
 
-		return pointer::AsInteger<executable::PreferredBase>(preferredBase);
+		const executable::PreferredBase preferredImageBase = pointer::AsInteger<executable::PreferredBase>(preferredBase);
+		LC_LOG_DEV("Preferred base address for image: 0x%" PRIX64 " (PID: %d)", preferredImageBase, processId);
+
+		return preferredImageBase;
 	}
 #endif
 
@@ -680,9 +683,11 @@ namespace
 			const size_t count = hookData.data->processes.size();
 			for (size_t p = 0u; p < count; ++p)
 			{
-				const unsigned int pid = hookData.data->processes[p].processId;
-				void* moduleBase = hookData.data->processes[p].moduleBase;
-				const DuplexPipe* pipe = hookData.data->processes[p].pipe;
+				const ModuleCache::ProcessData& processData = hookData.data->processes[p];
+
+				const unsigned int pid = processData.processId;
+				void* moduleBase = processData.moduleBase;
+				const DuplexPipe* pipe = processData.pipe;
 
 				LC_LOG_USER("Calling compile start hooks (PID: %d)", pid);
 				pipe->SendCommandAndWaitForAck(commands::CallHooks { hook::MakeFunction(moduleBase, hookData.firstRva), hook::MakeFunction(moduleBase, hookData.lastRva) });
@@ -705,9 +710,11 @@ namespace
 			const size_t count = hookData.data->processes.size();
 			for (size_t p = 0u; p < count; ++p)
 			{
-				const unsigned int pid = hookData.data->processes[p].processId;
-				void* moduleBase = hookData.data->processes[p].moduleBase;
-				const DuplexPipe* pipe = hookData.data->processes[p].pipe;
+				const ModuleCache::ProcessData& processData = hookData.data->processes[p];
+
+				const unsigned int pid = processData.processId;
+				void* moduleBase = processData.moduleBase;
+				const DuplexPipe* pipe = processData.pipe;
 
 				LC_LOG_USER("Calling compile success hooks (PID: %d)", pid);				
 				pipe->SendCommandAndWaitForAck(commands::CallHooks { hook::MakeFunction(moduleBase, hookData.firstRva), hook::MakeFunction(moduleBase, hookData.lastRva) });
@@ -730,9 +737,11 @@ namespace
 			const size_t count = hookData.data->processes.size();
 			for (size_t p = 0u; p < count; ++p)
 			{
-				const unsigned int pid = hookData.data->processes[p].processId;
-				void* moduleBase = hookData.data->processes[p].moduleBase;
-				const DuplexPipe* pipe = hookData.data->processes[p].pipe;
+				const ModuleCache::ProcessData& processData = hookData.data->processes[p];
+
+				const unsigned int pid = processData.processId;
+				void* moduleBase = processData.moduleBase;
+				const DuplexPipe* pipe = processData.pipe;
 
 				LC_LOG_USER("Calling compile error hooks (PID: %d)", pid);
 				pipe->SendCommandAndWaitForAck(commands::CallHooks { hook::MakeFunction(moduleBase, hookData.firstRva), hook::MakeFunction(moduleBase, hookData.lastRva) });
@@ -1195,7 +1204,7 @@ LiveModule::ErrorType::Enum LiveModule::Update(FileAttributeCache* fileCache, Di
 {
 	telemetry::Scope updateScope("Update live module");
 
-	LC_LOG_DEV("\nLiveModule::Update -------------------------------------------\n");
+	LC_LOG_DEV("LiveModule Update: %S", m_moduleName.c_str());
 
 	GLiveCodingServer->GetStatusChangeDelegate().ExecuteIfBound(L"Checking modified files...");
 
@@ -3390,9 +3399,11 @@ LiveModule::ErrorType::Enum LiveModule::Update(FileAttributeCache* fileCache, Di
 				const size_t count = hookData.data->processes.size();
 				for (size_t p = 0u; p < count; ++p)
 				{
-					const unsigned int pid = hookData.data->processes[p].processId;
-					void* moduleBase = hookData.data->processes[p].moduleBase;
-					const DuplexPipe* pipe = hookData.data->processes[p].pipe;
+					const ModuleCache::ProcessData& hookProcessData = hookData.data->processes[p];
+
+					const unsigned int pid = hookProcessData.processId;
+					void* moduleBase = hookProcessData.moduleBase;
+					const DuplexPipe* pipe = hookProcessData.pipe;
 
 					LC_LOG_USER("Calling pre-patch hooks (PID: %d)", pid);
 					pipe->SendCommandAndWaitForAck(commands::CallHooks { hook::MakeFunction(moduleBase, hookData.firstRva), hook::MakeFunction(moduleBase, hookData.lastRva) });
@@ -3822,9 +3833,11 @@ LiveModule::ErrorType::Enum LiveModule::Update(FileAttributeCache* fileCache, Di
 				{
 					++functionsPatchedCount;
 
-					const unsigned int pid = originalData.data->processes[p].processId;
-					void* moduleBase = originalData.data->processes[p].moduleBase;
-					process::Handle processHandle = originalData.data->processes[p].processHandle;
+					const ModuleCache::ProcessData& hookProcessData = originalData.data->processes[p];
+
+					const unsigned int pid = hookProcessData.processId;
+					void* moduleBase = hookProcessData.moduleBase;
+					process::Handle processHandle = hookProcessData.processHandle;
 
 					char* originalAddress = pointer::Offset<char*>(moduleBase, originalData.symbol->rva);
 					char* patchAddress = pointer::Offset<char*>(loadedPatches[p], patchSymbol->rva);
@@ -4004,7 +4017,7 @@ bool LiveModule::InstallCompiledPatches(LiveProcess* liveProcess, void* original
 		return true;
 	}
 
-	LC_LOG_DEV("\nLiveModule::InstallCompiledPatches ---------------------------\n");
+	LC_LOG_DEV("LiveModule InstallCompiledPatches: %S", m_moduleName.c_str());
 
 	telemetry::Scope wholeScope("Installing patches");
 
@@ -4251,7 +4264,7 @@ bool LiveModule::LoadPatchInfoAction::Execute(CommandType* command, const Duplex
 void LiveModule::UpdateDirectoryCache(const ImmutableString& path, symbols::Dependency* dependency, DirectoryCache* cache)
 {
 	const std::wstring& directoryOnly = file::GetDirectory(string::ToWideString(path));
-	dependency->parentDirectory = cache->AddDirectory(directoryOnly.c_str());
+	dependency->parentDirectory = cache->AddDirectory(directoryOnly);
 }
 
 
diff --git a/Engine/Source/Developer/Windows/LiveCodingServer/Private/External/LC_Scheduler.cpp b/Engine/Source/Developer/Windows/LiveCodingServer/Private/External/LC_Scheduler.cpp
--- a/Engine/Source/Developer/Windows/LiveCodingServer/Private/External/LC_Scheduler.cpp
+++ b/Engine/Source/Developer/Windows/LiveCodingServer/Private/External/LC_Scheduler.cpp
@@ -43,7 +43,7 @@ namespace
 				result = ::GetLogicalProcessorInformation(buffer, &bytesNeeded);
 
 				// return the number of logical processors in case anything went wrong
-				if (result == Windows::FALSE)
+				if ((result == Windows::FALSE) || (buffer == nullptr))
 				{
 					::free(buffer);
 					return GetLogicalProcessorCount();
diff --git a/Engine/Source/Developer/Windows/LiveCodingServer/Private/External/LC_ServerCommandThread.cpp b/Engine/Source/Developer/Windows/LiveCodingServer/Private/External/LC_ServerCommandThread.cpp
--- a/Engine/Source/Developer/Windows/LiveCodingServer/Private/External/LC_ServerCommandThread.cpp
+++ b/Engine/Source/Developer/Windows/LiveCodingServer/Private/External/LC_ServerCommandThread.cpp
@@ -41,7 +41,11 @@ namespace
 	{
 		InitializeCOM(void)
 		{
-			::CoInitialize(NULL);
+			const HRESULT success = ::CoInitialize(NULL);
+			if (success != S_OK)
+			{
+				LC_LOG_DEV("Could not initialize COM. Error: %d", success);
+			}
 		}
 
 		~InitializeCOM(void)
@@ -125,6 +129,7 @@ ServerCommandThread::~ServerCommandThread(void)
 	// this is only called when Live++ is being torn down anyway, so we leave cleanup to the OS.
 	// otherwise we could run into races when trying to terminate the thread that might currently be doing
 	// some intensive work.
+	delete m_directoryCache;
 }
 
 
diff --git a/Engine/Source/Developer/Windows/LiveCodingServer/Private/External/LC_Settings.cpp b/Engine/Source/Developer/Windows/LiveCodingServer/Private/External/LC_Settings.cpp
--- a/Engine/Source/Developer/Windows/LiveCodingServer/Private/External/LC_Settings.cpp
+++ b/Engine/Source/Developer/Windows/LiveCodingServer/Private/External/LC_Settings.cpp
@@ -253,6 +253,7 @@ SettingString::SettingString(const wchar_t* group, const wchar_t* name, const wc
 	, m_shortDescription(shortDescription)
 	, m_description(description)
 	, m_value(initialValue)
+	, m_previousValue(initialValue)
 {
 	const std::wstring value = LoadSetting(m_group, m_name, initialValue);
 
@@ -273,6 +274,7 @@ void SettingString::SetValue(const wchar_t* value)
 
 void SettingString::SetValueWithoutSaving(const wchar_t* value)
 {
+	m_previousValue = m_value;
 	m_value = value;
 }
 
@@ -301,6 +303,12 @@ const wchar_t* SettingString::GetValue(void) const
 }
 
 
+const wchar_t* SettingString::GetPreviousValue(void) const
+{
+	return m_previousValue.c_str();
+}
+
+
 SettingIntProxy::SettingIntProxy(SettingInt* setting)
 	: Setting(Type::INTEGER_PROXY)
 	, m_setting(setting)
diff --git a/Engine/Source/Developer/Windows/LiveCodingServer/Private/External/LC_Settings.h b/Engine/Source/Developer/Windows/LiveCodingServer/Private/External/LC_Settings.h
--- a/Engine/Source/Developer/Windows/LiveCodingServer/Private/External/LC_Settings.h
+++ b/Engine/Source/Developer/Windows/LiveCodingServer/Private/External/LC_Settings.h
@@ -111,6 +111,7 @@ public:
 	const wchar_t* GetShortDescription(void) const;
 	const wchar_t* GetDescription(void) const;
 	const wchar_t* GetValue(void) const;
+	const wchar_t* GetPreviousValue(void) const;
 
 private:
 	const wchar_t* m_group;
@@ -118,6 +119,7 @@ private:
 	const wchar_t* m_shortDescription;
 	const wchar_t* m_description;
 	std::wstring m_value;
+	std::wstring m_previousValue;
 };
 
 
diff --git a/Engine/Source/Developer/Windows/LiveCodingServer/Private/External/LC_SymbolInfo.cpp b/Engine/Source/Developer/Windows/LiveCodingServer/Private/External/LC_SymbolInfo.cpp
new file mode 100644
--- /dev/null
+++ b/Engine/Source/Developer/Windows/LiveCodingServer/Private/External/LC_SymbolInfo.cpp
@@ -0,0 +1,14 @@
+// Copyright 2011-2019 Molecular Matters GmbH, all rights reserved.
+
+#include "LC_SymbolInfo.h"
+#include <stdio.h>
+
+
+SymbolInfo::SymbolInfo(const char* const function, const char* const filename, unsigned int line)
+	: m_function()
+	, m_filename()
+	, m_line(line)
+{
+	strcpy_s(m_function, function);
+	strcpy_s(m_filename, filename);
+}
diff --git a/Engine/Source/Developer/Windows/LiveCodingServer/Private/External/LC_SymbolInfo.h b/Engine/Source/Developer/Windows/LiveCodingServer/Private/External/LC_SymbolInfo.h
new file mode 100644
--- /dev/null
+++ b/Engine/Source/Developer/Windows/LiveCodingServer/Private/External/LC_SymbolInfo.h
@@ -0,0 +1,34 @@
+// Copyright 2011-2019 Molecular Matters GmbH, all rights reserved.
+
+#pragma once
+
+#include "CoreTypes.h"
+#include "LC_Platform.h"
+
+class SymbolInfo
+{
+public:
+	SymbolInfo(const char* const function, const char* const filename, unsigned int line);
+
+	inline const char* GetFunction(void) const
+	{
+		return m_function;
+	}
+
+	inline const char* GetFilename(void) const
+	{
+		return m_filename;
+	}
+
+	inline unsigned int GetLine(void) const
+	{
+		return m_line;
+	}
+
+private:
+	LC_DISABLE_ASSIGNMENT(SymbolInfo);
+
+	char m_function[512u];
+	char m_filename[WINDOWS_MAX_PATH];
+	unsigned int m_line;
+};
diff --git a/Engine/Source/Developer/Windows/LiveCodingServer/Private/External/LC_SymbolResolution.cpp b/Engine/Source/Developer/Windows/LiveCodingServer/Private/External/LC_SymbolResolution.cpp
new file mode 100644
--- /dev/null
+++ b/Engine/Source/Developer/Windows/LiveCodingServer/Private/External/LC_SymbolResolution.cpp
@@ -0,0 +1,59 @@
+// Copyright 2011-2019 Molecular Matters GmbH, all rights reserved.
+
+#include "LC_SymbolResolution.h"
+#include "LC_PointerUtil.h"
+#include "LC_Logging.h"
+#include "Windows/WindowsHWrapper.h"
+#include <dbghelp.h>
+
+namespace symbolResolution
+{
+	void Startup(void)
+	{
+		::SymSetOptions(SYMOPT_DEFERRED_LOADS | SYMOPT_LOAD_LINES | SYMOPT_UNDNAME | SYMOPT_INCLUDE_32BIT_MODULES);
+		const BOOL success = ::SymInitialize(::GetCurrentProcess(), "", true);
+		if (success == Windows::FALSE)
+		{
+			const DWORD error = ::GetLastError();
+			LC_ERROR_DEV("SymInitialize failed with error 0x%X", error);
+		}
+	}
+
+
+	void Shutdown(void)
+	{
+		const BOOL success = ::SymCleanup(::GetCurrentProcess());
+		if (success == Windows::FALSE)
+		{
+			const DWORD error = ::GetLastError();
+			LC_ERROR_DEV("SymCleanup failed with error 0x%X", error);
+		}
+	}
+
+
+	SymbolInfo ResolveSymbolsForAddress(const void* const address)
+	{
+		// retrieve function
+		DWORD64 displacement64 = 0;
+		ULONG64 buffer[(sizeof(SYMBOL_INFO) + MAX_SYM_NAME * sizeof(TCHAR) + sizeof(ULONG64) - 1) / sizeof(ULONG64)] = {};
+		SYMBOL_INFO* symbolInfo = reinterpret_cast<SYMBOL_INFO*>(buffer);
+		symbolInfo->SizeOfStruct = sizeof(SYMBOL_INFO);
+		symbolInfo->MaxNameLen = MAX_SYM_NAME;
+
+		const HANDLE process = ::GetCurrentProcess();
+		const DWORD64 intAddress = pointer::AsInteger<DWORD64>(address);
+
+		if (::SymFromAddr(process, intAddress, &displacement64, symbolInfo))
+		{
+			// retrieve filename and line number
+			DWORD displacement = 0;
+			IMAGEHLP_LINE64 line = {};
+			line.SizeOfStruct = sizeof(IMAGEHLP_LINE64);
+
+			if (::SymGetLineFromAddr64(process, intAddress, &displacement, &line))
+				return SymbolInfo(symbolInfo->Name, line.FileName, line.LineNumber);
+		}
+
+		return SymbolInfo("unknown", "unknown", 0);
+	}
+}
diff --git a/Engine/Source/Developer/Windows/LiveCodingServer/Private/External/LC_SymbolResolution.h b/Engine/Source/Developer/Windows/LiveCodingServer/Private/External/LC_SymbolResolution.h
new file mode 100644
--- /dev/null
+++ b/Engine/Source/Developer/Windows/LiveCodingServer/Private/External/LC_SymbolResolution.h
@@ -0,0 +1,15 @@
+// Copyright 2011-2019 Molecular Matters GmbH, all rights reserved.
+
+#pragma once
+
+#include "LC_SymbolInfo.h"
+
+
+namespace symbolResolution
+{
+	void Startup(void);
+	void Shutdown(void);
+
+	// Resolves the symbol information for a given address
+	SymbolInfo ResolveSymbolsForAddress(const void* const address);
+}
diff --git a/Engine/Source/Developer/Windows/LiveCodingServer/Private/External/LC_Symbols.cpp b/Engine/Source/Developer/Windows/LiveCodingServer/Private/External/LC_Symbols.cpp
--- a/Engine/Source/Developer/Windows/LiveCodingServer/Private/External/LC_Symbols.cpp
+++ b/Engine/Source/Developer/Windows/LiveCodingServer/Private/External/LC_Symbols.cpp
@@ -337,30 +337,15 @@ namespace
 
 	static bool IsCppOrCFile(const std::wstring& normalizedLowercaseFilename)
 	{
-		const std::wstring& extension = file::GetExtension(normalizedLowercaseFilename);
-		if (string::Matches(extension.c_str(), L".cpp"))
-		{
-			return true;
-		}
-		else if (string::Matches(extension.c_str(), L".c"))
-		{
-			return true;
-		}
-		else if (string::Matches(extension.c_str(), L".cc"))
-		{
-			return true;
-		}
-		else if (string::Matches(extension.c_str(), L".c++"))
-		{
-			return true;
-		}
-		else if (string::Matches(extension.c_str(), L".cp"))
-		{
-			return true;
-		}
-		else if (string::Matches(extension.c_str(), L".cxx"))
+		const std::wstring& filenameExtension = file::GetExtension(normalizedLowercaseFilename);
+		const types::vector<std::wstring>& amalgamatedCppFileExtensions = appSettings::GetAmalgamatedCppFileExtensions();
+
+		for (size_t i = 0u; i < amalgamatedCppFileExtensions.size(); ++i)
 		{
-			return true;
+			if (string::Matches(filenameExtension.c_str(), amalgamatedCppFileExtensions[i].c_str()))
+			{
+				return true;
+			}
 		}
 
 		return false;
@@ -418,6 +403,10 @@ namespace symbols
 		{
 			const types::vector<IDiaSymbol*>& publicSymbols = dia::GatherChildSymbols(provider->globalScope, SymTagPublicSymbol);
 			const size_t symbolCount = publicSymbols.size();
+			symbolDB->symbolsByName.reserve(symbolCount);
+			symbolDB->symbolsByRva.reserve(symbolCount);
+			symbolDB->patchableFunctionSymbols.reserve(symbolCount);
+
 			for (size_t i = 0u; i < symbolCount; ++i)
 			{
 				IDiaSymbol* publicSymbol = publicSymbols[i];
@@ -649,7 +638,8 @@ namespace symbols
 			const dia::SymbolName& diaCompilandPath = dia::GetSymbolName(diaSymbol);
 			std::wstring compilandPath(diaCompilandPath.GetString());
 			const std::wstring& uppercaseCompilandPath = string::ToUpper(compilandPath);
-			const bool isObjPath = string::Contains(uppercaseCompilandPath.c_str(), L".OBJ");
+
+			const bool isObjPath = string::Contains(uppercaseCompilandPath.c_str(), L".OBJ") || string::Contains(uppercaseCompilandPath.c_str(), L".O");
 			if (isObjPath)
 			{
 				// a valid compiland, gather more information.
@@ -752,9 +742,20 @@ namespace symbols
 					continue;
 				}
 
-				// only add compilands that exist on disk
+				// only add compilands that exist on disk.
+				// optimization: ignore files stored on optical drives because they cannot be changed anyway.
 				{
 					// test the compiland path first
+					if (file::GetDriveType(compilandPath.c_str()) == file::DriveType::OPTICAL)
+					{
+						if (generateLogs)
+						{
+							LC_LOG_DEV("Ignoring file %S on optical drive", compilandPath.c_str());
+						}
+
+						continue;
+					}
+
 					FileAttributeCache::Data cacheData = fileCache.UpdateCacheData(compilandPath);
 					if (!cacheData.exists)
 					{
@@ -773,6 +774,16 @@ namespace symbols
 							testPath += L"\\";
 							testPath += file::GetFilename(compilandPath);
 
+							if (file::GetDriveType(testPath.c_str()) == file::DriveType::OPTICAL)
+							{
+								if (generateLogs)
+								{
+									LC_LOG_DEV("Ignoring file %S on optical drive", testPath.c_str());
+								}
+
+								continue;
+							}
+
 							cacheData = fileCache.UpdateCacheData(testPath);
 							if (!cacheData.exists)
 							{
@@ -794,6 +805,17 @@ namespace symbols
 								testPath = optionsCache[2];
 								testPath += L"\\";
 								testPath += compilandPath;
+
+								if (file::GetDriveType(testPath.c_str()) == file::DriveType::OPTICAL)
+								{
+									if (generateLogs)
+									{
+										LC_LOG_DEV("Ignoring file %S on optical drive", testPath.c_str());
+									}
+
+									continue;
+								}
+
 								cacheData = fileCache.UpdateCacheData(testPath);
 								testFileExists = cacheData.exists;
 							}
@@ -899,6 +921,7 @@ namespace symbols
 				// remote filenames in distributed builds.
 				// if we find a file dependency matching the given source file, we take that one instead to get
 				// full absolute file paths.
+				// optimization: ignore all files on optical drives because they cannot be changed anyway
 				const std::wstring srcFileOnlyLowercase = string::ToLower(file::GetFilename(optionsCache[0]));
 
 				const ObjPath objPath(string::ToUtf8String(normalizedCompilandPath));
@@ -908,6 +931,16 @@ namespace symbols
 					// we are only interested in tracking .obj files. we will never be able to recompile files
 					// and we don't know anything about source files, dependencies, etc.
 					// but we still use our dependency tracking system by letting each .obj depend on itself.
+					if (file::GetDriveType(normalizedCompilandPath.c_str()) == file::DriveType::OPTICAL)
+					{
+						if (generateLogs)
+						{
+							LC_LOG_DEV("Ignoring file %S on optical drive", normalizedCompilandPath.c_str());
+						}
+
+						continue;
+					}
+
 					const FileAttributeCache::Data& cacheData = fileCache.UpdateCacheData(normalizedCompilandPath);
 					if (cacheData.exists)
 					{
@@ -992,26 +1025,36 @@ namespace symbols
 						const FileInfo& fileInfo = fileInfos[j];
 						const std::wstring& normalizedFilename = fileInfo.normalizedFilename;
 
-						const FileAttributeCache::Data& cacheData = fileCache.UpdateCacheData(normalizedFilename);
-						if (cacheData.exists)
+						if (file::GetDriveType(normalizedFilename.c_str()) == file::DriveType::OPTICAL)
 						{
 							if (generateLogs)
 							{
-								LC_LOG_DEV("Dependency %S", normalizedFilename.c_str());
+								LC_LOG_DEV("Ignoring file %S on optical drive", normalizedFilename.c_str());
 							}
+						}
+						else
+						{
+							const FileAttributeCache::Data& cacheData = fileCache.UpdateCacheData(normalizedFilename);
+							if (cacheData.exists)
+							{
+								if (generateLogs)
+								{
+									LC_LOG_DEV("Dependency %S", normalizedFilename.c_str());
+								}
 
-							const ImmutableString& sourceFilePath = string::ToUtf8String(normalizedFilename);
-							AddFileDependency(compilandDb, sourceFilePath, objPath, cacheData.lastModificationTime);
+								const ImmutableString& sourceFilePath = string::ToUtf8String(normalizedFilename);
+								AddFileDependency(compilandDb, sourceFilePath, objPath, cacheData.lastModificationTime);
 
-							if (splitAmalgamatedFiles)
+								if (splitAmalgamatedFiles)
+								{
+									compiland->sourceFiles->files.push_back(sourceFilePath);
+								}
+							}
+							else if (generateLogs)
 							{
-								compiland->sourceFiles->files.push_back(sourceFilePath);
+								LC_LOG_DEV("Missing dependency %S", normalizedFilename.c_str());
 							}
 						}
-						else if (generateLogs)
-						{
-							LC_LOG_DEV("Missing dependency %S", normalizedFilename.c_str());
-						}
 					}
 
 					compilandDb->compilands.insert(std::make_pair(objPath, compiland));
@@ -1042,19 +1085,29 @@ namespace symbols
 							{
 								LC_LOG_DEV("Main .cpp %S", normalizedFilename.c_str());
 
-								const FileAttributeCache::Data& cacheData = fileCache.UpdateCacheData(normalizedFilename);
-								if (cacheData.exists)
+								if (file::GetDriveType(normalizedFilename.c_str()) == file::DriveType::OPTICAL)
 								{
 									if (generateLogs)
 									{
-										LC_LOG_DEV("Dependency %S", normalizedFilename.c_str());
+										LC_LOG_DEV("Ignoring file %S on optical drive", normalizedFilename.c_str());
 									}
-
-									AddFileDependency(compilandDb, sourceFilePath, objPath, cacheData.lastModificationTime);
 								}
-								else if (generateLogs)
+								else
 								{
-									LC_LOG_DEV("Missing dependency %S", normalizedFilename.c_str());
+									const FileAttributeCache::Data& cacheData = fileCache.UpdateCacheData(normalizedFilename);
+									if (cacheData.exists)
+									{
+										if (generateLogs)
+										{
+											LC_LOG_DEV("Dependency %S", normalizedFilename.c_str());
+										}
+
+										AddFileDependency(compilandDb, sourceFilePath, objPath, cacheData.lastModificationTime);
+									}
+									else if (generateLogs)
+									{
+										LC_LOG_DEV("Missing dependency %S", normalizedFilename.c_str());
+									}
 								}
 							}
 							else
@@ -1064,95 +1117,115 @@ namespace symbols
 								// this source file will not trigger a build of the amalgamated file.
 								LC_LOG_DEV("Included .cpp %S", normalizedFilename.c_str());
 
-								const FileAttributeCache::Data& cacheData = fileCache.UpdateCacheData(normalizedFilename);
-								if (cacheData.exists)
+								if (file::GetDriveType(normalizedFilename.c_str()) == file::DriveType::OPTICAL)
 								{
 									if (generateLogs)
 									{
-										LC_LOG_DEV("Dependency %S", normalizedFilename.c_str());
+										LC_LOG_DEV("Ignoring file %S on optical drive", normalizedFilename.c_str());
 									}
+								}
+								else
+								{
+									const FileAttributeCache::Data& cacheData = fileCache.UpdateCacheData(normalizedFilename);
+									if (cacheData.exists)
+									{
+										if (generateLogs)
+										{
+											LC_LOG_DEV("Dependency %S", normalizedFilename.c_str());
+										}
 
-									// create new .obj path by appending this file name to the real .obj, e.g.
-									// Amalgamated.obj turns into Amalgamated.lpp_part.ASingleFile.obj.
-									const std::wstring newObjPart = amalgamation::CreateObjPart(file::NormalizePath(normalizedFilename.c_str()));
-									const std::wstring newObjPath = amalgamation::CreateObjPath(compilandPath, newObjPart);
-									const std::wstring normalizedNewObjPath = file::NormalizePath(newObjPath.c_str());
-
-									AddFileDependency(compilandDb, sourceFilePath, string::ToUtf8String(normalizedNewObjPath), cacheData.lastModificationTime);
+										// create new .obj path by appending this file name to the real .obj, e.g.
+										// Amalgamated.obj turns into Amalgamated.lpp_part.ASingleFile.obj.
+										const std::wstring newObjPart = amalgamation::CreateObjPart(file::NormalizePath(normalizedFilename.c_str()));
+										const std::wstring newObjPath = amalgamation::CreateObjPath(compilandPath, newObjPart);
+										const std::wstring normalizedNewObjPath = file::NormalizePath(newObjPath.c_str());
 
-									// create a new compiland matching this .obj.
-									// we could use different PDBs for different files when no PCHs are being used, but with
-									// our automatic multi-processor compilation this doesn't really gain anything performance-wise
-									// and just complicates things.
+										AddFileDependency(compilandDb, sourceFilePath, string::ToUtf8String(normalizedNewObjPath), cacheData.lastModificationTime);
 
-									// adapt command line to accommodate new .obj path
-									const std::wstring& newCommandLine = string::Replace(optionsCache[4], L".obj", newObjPart);
+										// create a new compiland matching this .obj.
+										// we could use different PDBs for different files when no PCHs are being used, but with
+										// our automatic multi-processor compilation this doesn't really gain anything performance-wise
+										// and just complicates things.
 
-									Compiland* newCompiland = LC_NEW(&g_compilandAllocator, Compiland)
-									{
-										string::ToUtf8String(newObjPath),
-										string::ToUtf8String(normalizedFilename),
-										envPdbPath,
-										envCompilerPath,
-										string::ToUtf8String(newCommandLine),
-										envCompilerWorkingDirectory,
-										objPath,										// .obj of the amalgamation
-										nullptr,										// file indices
-
-										// note that for the purpose of disambiguating symbols in COFF files,
-										// we treat these files as being the amalgamated file.
-										// symbols originally coming from amalgamated files need to have the same
-										// name as symbols from individual files.
-										uniqueId::Generate(normalizedCompilandPath),
-
-										// same for the DIA symbol index, for the same reason
-										static_cast<uint32_t>(i),
-
-										Compiland::Type::PART_OF_AMALGAMATION,			// type of file
-										isPartOfLibrary,								// isPartOfLibrary
-										false											// wasRecompiled
-									};
-
-									compilandDb->compilands.insert(std::make_pair(string::ToUtf8String(normalizedNewObjPath), newCompiland));
-
-									// try updating the amalgamated compiland for the given file and create a new one in case none exists yet
-									{
-										const auto& insertPair = compilandDb->amalgamatedCompilands.emplace(objPath, nullptr);
-										symbols::AmalgamatedCompiland*& amalgamatedCompiland = insertPair.first->second;
+										// adapt command line to accommodate new .obj path
+										const std::wstring& newCommandLine = string::Replace(optionsCache[4], L".obj", newObjPart);
 
-										if (insertPair.second)
+										Compiland* newCompiland = LC_NEW(&g_compilandAllocator, Compiland)
 										{
-											// insertion was successful, create a new amalgamated compiland
-											amalgamatedCompiland = LC_NEW(&g_amalgamatedCompilandAllocator, symbols::AmalgamatedCompiland);
-											amalgamatedCompiland->isSplit = false;
+											string::ToUtf8String(newObjPath),
+											string::ToUtf8String(normalizedFilename),
+											envPdbPath,
+											envCompilerPath,
+											string::ToUtf8String(newCommandLine),
+											envCompilerWorkingDirectory,
+											objPath,										// .obj of the amalgamation
+											nullptr,										// file indices
+
+											// note that for the purpose of disambiguating symbols in COFF files,
+											// we treat these files as being the amalgamated file.
+											// symbols originally coming from amalgamated files need to have the same
+											// name as symbols from individual files.
+											uniqueId::Generate(normalizedCompilandPath),
+
+											// same for the DIA symbol index, for the same reason
+											static_cast<uint32_t>(i),
+
+											Compiland::Type::PART_OF_AMALGAMATION,			// type of file
+											isPartOfLibrary,								// isPartOfLibrary
+											false											// wasRecompiled
+										};
+
+										compilandDb->compilands.insert(std::make_pair(string::ToUtf8String(normalizedNewObjPath), newCompiland));
+
+										// try updating the amalgamated compiland for the given file and create a new one in case none exists yet
+										{
+											const auto& insertPair = compilandDb->amalgamatedCompilands.emplace(objPath, nullptr);
+											symbols::AmalgamatedCompiland*& amalgamatedCompiland = insertPair.first->second;
+
+											if (insertPair.second)
+											{
+												// insertion was successful, create a new amalgamated compiland
+												amalgamatedCompiland = LC_NEW(&g_amalgamatedCompilandAllocator, symbols::AmalgamatedCompiland);
+												amalgamatedCompiland->isSplit = false;
+											}
+
+											// update entry
+											amalgamatedCompiland->singleParts.push_back(string::ToUtf8String(normalizedNewObjPath));
 										}
-
-										// update entry
-										amalgamatedCompiland->singleParts.push_back(string::ToUtf8String(normalizedNewObjPath));
 									}
-								}
-								else if (generateLogs)
-								{
-									LC_LOG_DEV("Missing dependency %S", normalizedFilename.c_str());
+									else if (generateLogs)
+									{
+										LC_LOG_DEV("Missing dependency %S", normalizedFilename.c_str());
+									}
 								}
 							}
 						}
 						else
 						{
 							// this is a header file. add it as regular dependency for the main amalgamated .obj file
-							const FileAttributeCache::Data& cacheData = fileCache.UpdateCacheData(normalizedFilename);
-							if (cacheData.exists)
+							if (file::GetDriveType(normalizedFilename.c_str()) == file::DriveType::OPTICAL)
 							{
 								if (generateLogs)
 								{
-									LC_LOG_DEV("Dependency %S", normalizedFilename.c_str());
+									LC_LOG_DEV("Ignoring file %S on optical drive", normalizedFilename.c_str());
 								}
-
-								AddFileDependency(compilandDb, sourceFilePath, objPath, cacheData.lastModificationTime);
 							}
-							else if (generateLogs)
+							else
 							{
-								LC_LOG_DEV("Missing dependency %S", normalizedFilename.c_str());
+								const FileAttributeCache::Data& cacheData = fileCache.UpdateCacheData(normalizedFilename);
+								if (cacheData.exists)
+								{
+									if (generateLogs)
+									{
+										LC_LOG_DEV("Dependency %S", normalizedFilename.c_str());
+									}
+
+									AddFileDependency(compilandDb, sourceFilePath, objPath, cacheData.lastModificationTime);
+								}
+								else if (generateLogs)
+								{
+									LC_LOG_DEV("Missing dependency %S", normalizedFilename.c_str());
+								}
 							}
 						}
 					}
@@ -1464,14 +1537,14 @@ namespace symbols
 		const symbols::ImageSection* firstSection = symbols::FindImageSectionByRVA(imageSectionDb, firstRva);
 		if (!firstSection)
 		{
-			LC_ERROR_DEV("Cannot find image section holding start of dynamic initializer range");
+			LC_ERROR_USER("Could not find image section containing dynamic initializers.\nThis will lead to constructors of global and static variables being called again for the next patch, likely leading to unexpected behaviour.");
 			return initializerDb;
 		}
 
 		const symbols::ImageSection* lastSection = symbols::FindImageSectionByRVA(imageSectionDb, lastRva);
 		if (!lastSection)
 		{
-			LC_ERROR_DEV("Cannot find image section holding end of dynamic initializer range");
+			LC_ERROR_USER("Could not find image section containing dynamic initializers.\nThis will lead to constructors of global and static variables being called again for the next patch, likely leading to unexpected behaviour.");
 			return initializerDb;
 		}
 
diff --git a/Engine/Source/Developer/Windows/LiveCodingServer/Private/External/LC_Symbols.h b/Engine/Source/Developer/Windows/LiveCodingServer/Private/External/LC_Symbols.h
--- a/Engine/Source/Developer/Windows/LiveCodingServer/Private/External/LC_Symbols.h
+++ b/Engine/Source/Developer/Windows/LiveCodingServer/Private/External/LC_Symbols.h
@@ -27,6 +27,7 @@ namespace symbols
 		uint32_t rva;
 
 		LC_DISABLE_ASSIGNMENT(Symbol);
+		LC_DISABLE_COPY(Symbol);
 	};
 
 	struct SymbolDB
@@ -88,6 +89,7 @@ namespace symbols
 		bool wasRecompiled;
 
 		LC_DISABLE_ASSIGNMENT(Compiland);
+		LC_DISABLE_COPY(Compiland);
 	};
 
 	struct AmalgamatedCompiland
diff --git a/Engine/Source/Developer/Windows/LiveCodingServer/Private/External/LC_Telemetry.cpp b/Engine/Source/Developer/Windows/LiveCodingServer/Private/External/LC_Telemetry.cpp
--- a/Engine/Source/Developer/Windows/LiveCodingServer/Private/External/LC_Telemetry.cpp
+++ b/Engine/Source/Developer/Windows/LiveCodingServer/Private/External/LC_Telemetry.cpp
@@ -31,12 +31,15 @@ namespace
 telemetry::Scope::Scope(const char* name)
 	: m_name(name)
 	, m_start(std::chrono::high_resolution_clock::now())
+	, m_cs()
 {
 }
 
 
 telemetry::Scope::~Scope(void)
 {
+	CriticalSection::ScopedLock lock(&m_cs);
+
 	if (m_name)
 	{
 		Print(m_name, m_start);
@@ -70,6 +73,8 @@ void telemetry::Scope::Restart(void)
 
 void telemetry::Scope::End(void)
 {
+	CriticalSection::ScopedLock lock(&m_cs);
+
 	Print(m_name, m_start);
 
 	// do not print again when going out of scope
diff --git a/Engine/Source/Developer/Windows/LiveCodingServer/Private/External/LC_Telemetry.h b/Engine/Source/Developer/Windows/LiveCodingServer/Private/External/LC_Telemetry.h
--- a/Engine/Source/Developer/Windows/LiveCodingServer/Private/External/LC_Telemetry.h
+++ b/Engine/Source/Developer/Windows/LiveCodingServer/Private/External/LC_Telemetry.h
@@ -4,6 +4,7 @@
 
 #include "CoreTypes.h"
 #include <chrono>
+#include "LC_CriticalSection.h"
 
 namespace telemetry
 {
@@ -23,6 +24,7 @@ namespace telemetry
 	private:
 		const char* m_name;
 		std::chrono::high_resolution_clock::time_point m_start;
+		CriticalSection m_cs;
 	};
 
 
