Merging early Live++ 1.4.3 changes.

diff --git a/Engine/Source/Developer/Windows/LiveCoding/Private/External/LC_ClientStartupThread.cpp b/Engine/Source/Developer/Windows/LiveCoding/Private/External/LC_ClientStartupThread.cpp
--- a/Engine/Source/Developer/Windows/LiveCoding/Private/External/LC_ClientStartupThread.cpp
+++ b/Engine/Source/Developer/Windows/LiveCoding/Private/External/LC_ClientStartupThread.cpp
@@ -176,14 +176,14 @@ void ClientStartupThread::TriggerRecompile(void)
 }
 
 
-void ClientStartupThread::BuildPatch(const wchar_t* moduleNames[], const wchar_t* objPaths[], unsigned int count)
+void ClientStartupThread::BuildPatch(const wchar_t* moduleNames[], const wchar_t* objPaths[], const wchar_t* amalgamatedObjPaths[], unsigned int count)
 {
 	// we cannot wait for commands in the user command thread as long as startup hasn't finished
 	Join();
 
 	if (m_userCommandThread)
 	{
-		m_userCommandThread->BuildPatch(moduleNames, objPaths, count);
+		m_userCommandThread->BuildPatch(moduleNames, objPaths, amalgamatedObjPaths, count);
 	}
 }
 
diff --git a/Engine/Source/Developer/Windows/LiveCoding/Private/External/LC_ClientStartupThread.h b/Engine/Source/Developer/Windows/LiveCoding/Private/External/LC_ClientStartupThread.h
--- a/Engine/Source/Developer/Windows/LiveCoding/Private/External/LC_ClientStartupThread.h
+++ b/Engine/Source/Developer/Windows/LiveCoding/Private/External/LC_ClientStartupThread.h
@@ -35,7 +35,7 @@ public:
 
 	void WaitForToken(void* token);
 	void TriggerRecompile(void);
-	void BuildPatch(const wchar_t* moduleNames[], const wchar_t* objPaths[], unsigned int count);
+	void BuildPatch(const wchar_t* moduleNames[], const wchar_t* objPaths[], const wchar_t* amalgamatedObjPaths[], unsigned int count);
 
 	void InstallExceptionHandler(void);
 
diff --git a/Engine/Source/Developer/Windows/LiveCoding/Private/External/LC_ClientUserCommandThread.cpp b/Engine/Source/Developer/Windows/LiveCoding/Private/External/LC_ClientUserCommandThread.cpp
--- a/Engine/Source/Developer/Windows/LiveCoding/Private/External/LC_ClientUserCommandThread.cpp
+++ b/Engine/Source/Developer/Windows/LiveCoding/Private/External/LC_ClientUserCommandThread.cpp
@@ -273,6 +273,7 @@ namespace userCommands
 			, count(0u)
 			, moduleNames()
 			, objPaths()
+			, amalgamatedObjPaths()
 		{
 		}
 
@@ -289,6 +290,7 @@ namespace userCommands
 				commands::BuildPatchPacket packet;
 				wcscpy_s(packet.moduleName, moduleNames[i].c_str());
 				wcscpy_s(packet.objPath, objPaths[i].c_str());
+				wcscpy_s(packet.amalgamatedObjPath, amalgamatedObjPaths[i].c_str());
 
 				pipe->SendCommandAndWaitForAck(packet);
 			}
@@ -297,6 +299,7 @@ namespace userCommands
 		unsigned int count;
 		std::vector<std::wstring> moduleNames;
 		std::vector<std::wstring> objPaths;
+		std::vector<std::wstring> amalgamatedObjPaths;
 	};
 
 	struct ApplySettingBoolCommand : public BaseCommand
@@ -527,17 +530,28 @@ void ClientUserCommandThread::TriggerRecompile(void)
 }
 
 
-void ClientUserCommandThread::BuildPatch(const wchar_t* moduleNames[], const wchar_t* objPaths[], unsigned int count)
+void ClientUserCommandThread::BuildPatch(const wchar_t* moduleNames[], const wchar_t* objPaths[], const wchar_t* amalgamatedObjPaths[], unsigned int count)
 {
 	userCommands::BuildPatchCommand* command = new userCommands::BuildPatchCommand;
 	command->count = count;
 	command->moduleNames.reserve(count);
 	command->objPaths.reserve(count);
+	command->amalgamatedObjPaths.reserve(count);
 
 	for (unsigned int i = 0u; i < count; ++i)
 	{
 		command->moduleNames.push_back(moduleNames[i]);
 		command->objPaths.push_back(objPaths[i]);
+
+		// the amalgamated object paths are optional
+		if (amalgamatedObjPaths && amalgamatedObjPaths[i])
+		{
+			command->amalgamatedObjPaths.push_back(amalgamatedObjPaths[i]);
+		}
+		else
+		{
+			command->amalgamatedObjPaths.push_back(std::wstring());
+		}
 	}
 	{
 		CriticalSection::ScopedLock lock(&g_userCommandQueueCS);
diff --git a/Engine/Source/Developer/Windows/LiveCoding/Private/External/LC_ClientUserCommandThread.h b/Engine/Source/Developer/Windows/LiveCoding/Private/External/LC_ClientUserCommandThread.h
--- a/Engine/Source/Developer/Windows/LiveCoding/Private/External/LC_ClientUserCommandThread.h
+++ b/Engine/Source/Developer/Windows/LiveCoding/Private/External/LC_ClientUserCommandThread.h
@@ -41,7 +41,7 @@ public:
 
 	void WaitForToken(void* token);
 	void TriggerRecompile(void);
-	void BuildPatch(const wchar_t* moduleNames[], const wchar_t* objPaths[], unsigned int count);
+	void BuildPatch(const wchar_t* moduleNames[], const wchar_t* objPaths[], const wchar_t* amalgamatedObjPaths[], unsigned int count);
 
 	void InstallExceptionHandler(void);
 	ExceptionResult HandleException(EXCEPTION_RECORD* exception, CONTEXT* context, unsigned int threadId);
diff --git a/Engine/Source/Developer/Windows/LiveCoding/Private/External/LC_Commands.h b/Engine/Source/Developer/Windows/LiveCoding/Private/External/LC_Commands.h
--- a/Engine/Source/Developer/Windows/LiveCoding/Private/External/LC_Commands.h
+++ b/Engine/Source/Developer/Windows/LiveCoding/Private/External/LC_Commands.h
@@ -256,6 +256,7 @@ namespace commands
 
 		wchar_t moduleName[MAX_PATH];
 		wchar_t objPath[MAX_PATH];
+		wchar_t amalgamatedObjPath[MAX_PATH];
 	};
 
 	// tell Live++ to handle an exception
diff --git a/Engine/Source/Developer/Windows/LiveCoding/Private/External/LC_EntryPoint.cpp b/Engine/Source/Developer/Windows/LiveCoding/Private/External/LC_EntryPoint.cpp
--- a/Engine/Source/Developer/Windows/LiveCoding/Private/External/LC_EntryPoint.cpp
+++ b/Engine/Source/Developer/Windows/LiveCoding/Private/External/LC_EntryPoint.cpp
@@ -125,9 +125,9 @@ LPP_DLL_API void __cdecl LppTriggerRecompile(void)
 }
 
 
-LPP_DLL_API void __cdecl LppBuildPatch(const wchar_t* moduleNames[], const wchar_t* objPaths[], unsigned int count)
+LPP_DLL_API void __cdecl LppBuildPatch(const wchar_t* moduleNames[], const wchar_t* objPaths[], const wchar_t* amalgamatedObjPaths[], unsigned int count)
 {
-	g_startupThread->BuildPatch(moduleNames, objPaths, count);
+	g_startupThread->BuildPatch(moduleNames, objPaths, amalgamatedObjPaths, count);
 }
 
 
diff --git a/Engine/Source/Developer/Windows/LiveCodingServer/Private/External/LC_LiveModule.cpp b/Engine/Source/Developer/Windows/LiveCodingServer/Private/External/LC_LiveModule.cpp
--- a/Engine/Source/Developer/Windows/LiveCodingServer/Private/External/LC_LiveModule.cpp
+++ b/Engine/Source/Developer/Windows/LiveCodingServer/Private/External/LC_LiveModule.cpp
@@ -366,12 +366,35 @@ namespace
 	}
 
 
-	// helper function that returns or generates the unique ID of an optional compiland
-	static inline uint32_t GetCompilandId(const symbols::Compiland* compiland, const wchar_t* const objPath)
+	// helper function that returns or generates the unique ID of an optional compiland.
+	// for files split off from amalgamated files, we need to use the original object path of the amalgamated file here,
+	// otherwise names of symbols would differ, leading to constructors of global instances being called again.
+	static inline uint32_t GetCompilandId(const symbols::Compiland* compiland, const wchar_t* const objPath, const types::vector<LiveModule::ModifiedObjFile>& modifiedObjFiles)
 	{
-		return compiland
-			? compiland->uniqueId									// compiland exists
-			: uniqueId::Generate(file::NormalizePath(objPath));		// new compiland, generate new unique ID
+		// try to find the given .obj path in the array of modified object files to check if there's an original amalgamated object path for it
+		for (size_t i = 0u; i < modifiedObjFiles.size(); ++i)
+		{
+			const LiveModule::ModifiedObjFile& objFile = modifiedObjFiles[i];
+
+			// don't bother checking strings if the amalgamated object path is empty anyway
+			if (!objFile.amalgamatedObjPath.empty())
+			{
+				if (string::Matches(objPath, objFile.objPath.c_str()))
+				{
+					return uniqueId::Generate(file::NormalizePath(objFile.amalgamatedObjPath.c_str()));
+				}
+			}
+		}
+
+		if (compiland)
+		{
+			// the compiland already exists
+			return compiland->uniqueId;
+		}
+		else
+		{
+			return uniqueId::Generate(file::NormalizePath(objPath));
+		}
 	}
 
 
@@ -441,7 +464,7 @@ namespace
 
 
 	template <typename T>
-	static types::vector<symbols::ObjPath> UpdateCoffCache(const T& compilands, CoffCache<coff::CoffDB>* coffCache, CacheUpdate::Enum updateType, coff::ReadFlags::Enum coffReadFlags)
+	static types::vector<symbols::ObjPath> UpdateCoffCache(const T& compilands, CoffCache<coff::CoffDB>* coffCache, CacheUpdate::Enum updateType, coff::ReadFlags::Enum coffReadFlags, const types::vector<LiveModule::ModifiedObjFile>& modifiedOrNewObjFiles)
 	{
 		LC_LOG_INDENT_DEV;
 
@@ -458,7 +481,7 @@ namespace
 			symbols::ObjPath objPath = it->first;
 			const std::wstring& wideObjPath = string::ToWideString(objPath);
 			const symbols::Compiland* compiland = it->second;
-			const uint32_t compilandUniqueId = GetCompilandId(compiland, wideObjPath.c_str());
+			const uint32_t compilandUniqueId = GetCompilandId(compiland, wideObjPath.c_str(), modifiedOrNewObjFiles);
 
 			const bool shouldUpdate = (updateType == CacheUpdate::NON_EXISTANT)
 				? (coffCache->Lookup(objPath) == nullptr)				// NON-EXISTANT: update cache only for files which don't have an entry yet
@@ -1200,7 +1223,7 @@ void LiveModule::UpdateDirectoryCache(DirectoryCache* cache)
 }
 
 
-LiveModule::ErrorType::Enum LiveModule::Update(FileAttributeCache* fileCache, DirectoryCache* directoryCache, UpdateType::Enum updateType, const std::vector<std::wstring>& modifiedOrNewObjFiles)
+LiveModule::ErrorType::Enum LiveModule::Update(FileAttributeCache* fileCache, DirectoryCache* directoryCache, UpdateType::Enum updateType, const types::vector<ModifiedObjFile>& modifiedOrNewObjFiles)
 {
 	telemetry::Scope updateScope("Update live module");
 
@@ -1340,7 +1363,7 @@ LiveModule::ErrorType::Enum LiveModule::Update(FileAttributeCache* fileCache, Di
 	{
 		for (size_t i = 0u; i < modifiedOrNewObjFiles.size(); ++i)
 		{
-			LC_LOG_USER("File %S was modified or is new", modifiedOrNewObjFiles[i].c_str());
+			LC_LOG_USER("File %S was modified or is new", modifiedOrNewObjFiles[i].objPath.c_str());
 		}
 
 		LC_LOG_USER("Building patch from %zu file(s) for Live Coding module %S", modifiedOrNewObjFiles.size(), m_moduleName.c_str());
@@ -1811,7 +1834,7 @@ LiveModule::ErrorType::Enum LiveModule::Update(FileAttributeCache* fileCache, Di
 			const size_t count = modifiedOrNewObjFiles.size();
 			for (size_t i = 0u; i < count; ++i)
 			{
-				const std::wstring& wideObjPath = modifiedOrNewObjFiles[i];
+				const std::wstring& wideObjPath = modifiedOrNewObjFiles[i].objPath;
 				const symbols::ObjPath& objPath = string::ToUtf8String(wideObjPath);
 				symbols::Compiland* compiland = symbols::FindCompiland(m_compilandDB, objPath);
 
@@ -1954,11 +1977,11 @@ LiveModule::ErrorType::Enum LiveModule::Update(FileAttributeCache* fileCache, Di
 		{
 			for (auto it = modifiedOrNewObjFiles.begin(); it != modifiedOrNewObjFiles.end(); ++it)
 			{
-				const symbols::ObjPath objPath(string::ToUtf8String(*it));
+				const symbols::ObjPath objPath(string::ToUtf8String(it->objPath));
 				const symbols::Compiland* compiland = symbols::FindCompiland(m_compilandDB, objPath);
 
 				// new compilands won't be found in the database, so there's no unique ID yet that we can use
-				const uint32_t compilandUniqueId = GetCompilandId(compiland, it->c_str());
+				const uint32_t compilandUniqueId = GetCompilandId(compiland, it->objPath.c_str(), modifiedOrNewObjFiles);
 
 				// this file was either modified or is new. in any case, the new .OBJ needs to be linked in, even
 				// though the file might be contained in a library.
@@ -2035,11 +2058,11 @@ LiveModule::ErrorType::Enum LiveModule::Update(FileAttributeCache* fileCache, Di
 						// there is no entry yet for this COFF in the cache.
 						// this means that this .obj was not recompiled (otherwise it would have an entry already),
 						// but has been pulled in for the first time due to unresolved symbols.
-						auto task = scheduler::CreateTask(taskRoot, [this, objPath, coffReadFlags]()
+						auto task = scheduler::CreateTask(taskRoot, [this, objPath, coffReadFlags, &modifiedOrNewObjFiles]()
 						{
 							const symbols::Compiland* compiland = symbols::FindCompiland(m_compilandDB, objPath);
 							const std::wstring& wideObjPath = string::ToWideString(objPath);
-							const uint32_t compilandUniqueId = GetCompilandId(compiland, wideObjPath.c_str());
+							const uint32_t compilandUniqueId = GetCompilandId(compiland, wideObjPath.c_str(), modifiedOrNewObjFiles);
 
 								LC_LOG_DEV("Need %s for the first time, updating COFF cache", objPath.c_str());
 
@@ -2103,7 +2126,7 @@ LiveModule::ErrorType::Enum LiveModule::Update(FileAttributeCache* fileCache, Di
 	}
 
 	// update the COFF cache for all compiled files
-	UpdateCoffCache(m_compiledCompilands, m_coffCache, CacheUpdate::ALL, coffReadFlags);
+	UpdateCoffCache(m_compiledCompilands, m_coffCache, CacheUpdate::ALL, coffReadFlags, modifiedOrNewObjFiles);
 
 	GLiveCodingServer->GetStatusChangeDelegate().ExecuteIfBound(L"Stripping COFFs...");
 
@@ -2132,7 +2155,7 @@ LiveModule::ErrorType::Enum LiveModule::Update(FileAttributeCache* fileCache, Di
 			symbols::Compiland* compiland = symbols::FindCompiland(m_compilandDB, objPath);
 			const std::wstring wideObjPath = string::ToWideString(objPath);
 
-			const uint32_t compilandUniqueId = GetCompilandId(compiland, wideObjPath.c_str());
+			const uint32_t compilandUniqueId = GetCompilandId(compiland, wideObjPath.c_str(), modifiedOrNewObjFiles);
 
 			coff::ObjFile* objFile = coff::OpenObj(wideObjPath.c_str());
 			if (objFile && objFile->memoryFile)
@@ -2660,7 +2683,7 @@ LiveModule::ErrorType::Enum LiveModule::Update(FileAttributeCache* fileCache, Di
 		const symbols::ObjPath& objPath = compilandIt->first;
 		const std::wstring& wideObjPath = string::ToWideString(objPath);
 		const symbols::Compiland* compiland = compilandIt->second;
-		const uint32_t compilandUniqueId = GetCompilandId(compiland, wideObjPath.c_str());
+		const uint32_t compilandUniqueId = GetCompilandId(compiland, wideObjPath.c_str(), modifiedOrNewObjFiles);
 
 		coff::ObjFile* coffFile = coff::OpenObj(wideObjPath.c_str());
 		if (coffFile && coffFile->memoryFile)
@@ -3247,7 +3270,7 @@ LiveModule::ErrorType::Enum LiveModule::Update(FileAttributeCache* fileCache, Di
 		// update the COFF cache for new patch compilands.
 		// there may be files for which we don't have a database yet, even though we updated the database for all compiled files.
 		// this can happen when a new .obj that is part of a library is linked in for the first time.
-		types::vector<symbols::ObjPath> updatedCoffs = UpdateCoffCache(patch_compilandDB->compilands, m_coffCache, CacheUpdate::NON_EXISTANT, coffReadFlags);
+		types::vector<symbols::ObjPath> updatedCoffs = UpdateCoffCache(patch_compilandDB->compilands, m_coffCache, CacheUpdate::NON_EXISTANT, coffReadFlags, modifiedOrNewObjFiles);
 
 
 		// similarly, reconstruct symbols and dynamic initializers for new .obj that have been pulled in for the first time.
diff --git a/Engine/Source/Developer/Windows/LiveCodingServer/Private/External/LC_LiveModule.h b/Engine/Source/Developer/Windows/LiveCodingServer/Private/External/LC_LiveModule.h
--- a/Engine/Source/Developer/Windows/LiveCodingServer/Private/External/LC_LiveModule.h
+++ b/Engine/Source/Developer/Windows/LiveCodingServer/Private/External/LC_LiveModule.h
@@ -51,6 +51,12 @@ public:
 		};
 	};
 
+	struct ModifiedObjFile
+	{
+		std::wstring objPath;
+		std::wstring amalgamatedObjPath;	// optional
+	};
+
 	struct CompileResult
 	{
 		unsigned int exitCode;
@@ -76,7 +82,7 @@ public:
 	// in DEFAULT mode, this checks for file modifications, compiles files automatically, builds a patch containing changes, and loads them into the host application.
 	// in EXTERNAL_BUILD_SYSTEM mode, this does not compile files but builds a patch containing modified .objs, loading the patch into the host application.
 	// optionally, an array of modified or new .objs can be given in this mode, which builds a patch containing these files, not checking for any other modifications.
-	ErrorType::Enum Update(FileAttributeCache* fileCache, DirectoryCache* directoryCache, UpdateType::Enum updateType, const std::vector<std::wstring>& modifiedOrNewObjFiles);
+	ErrorType::Enum Update(FileAttributeCache* fileCache, DirectoryCache* directoryCache, UpdateType::Enum updateType, const types::vector<ModifiedObjFile>& modifiedOrNewObjFiles);
 	bool InstallCompiledPatches(LiveProcess* liveProcess, void* originalModuleBase);
 
 	const std::wstring& GetModuleName(void) const;
diff --git a/Engine/Source/Developer/Windows/LiveCodingServer/Private/External/LC_PoolAllocator.cpp b/Engine/Source/Developer/Windows/LiveCodingServer/Private/External/LC_PoolAllocator.cpp
--- a/Engine/Source/Developer/Windows/LiveCodingServer/Private/External/LC_PoolAllocator.cpp
+++ b/Engine/Source/Developer/Windows/LiveCodingServer/Private/External/LC_PoolAllocator.cpp
@@ -38,7 +38,7 @@ namespace freeList
 		// the alignment will always satisfy restrictions while producing the least amount of wasted/unused memory.
 		// furthermore, we must always be able to store a Freelist* in each entry.
 		const size_t minimumSize = sizeof(Slot*) > originalElementSize ? sizeof(Slot*) : originalElementSize;
-		const size_t elementSize = RoundUpToMultiple(minimumSize, alignment);
+		const size_t elementSize = freelist::RoundUpToMultiple(minimumSize, alignment);
 
 		// in order to satisfy alignment restrictions, it suffices to offset the start of the free list only
 		as_void = memoryBlock;
diff --git a/Engine/Source/Developer/Windows/LiveCodingServer/Private/External/LC_ServerCommandThread.cpp b/Engine/Source/Developer/Windows/LiveCodingServer/Private/External/LC_ServerCommandThread.cpp
--- a/Engine/Source/Developer/Windows/LiveCodingServer/Private/External/LC_ServerCommandThread.cpp
+++ b/Engine/Source/Developer/Windows/LiveCodingServer/Private/External/LC_ServerCommandThread.cpp
@@ -771,11 +771,12 @@ void ServerCommandThread::CompileChanges(bool didAllProcessesMakeProgress)
 				return;
 			}
 
-			types::vector<std::wstring> ObjectFiles;
+			types::vector<LiveModule::ModifiedObjFile> ObjectFiles;
 			for(const FString& ObjectFile : Pair.Value)
 			{
-				std::wstring NormalizedObjectFile = file::NormalizePath(*ObjectFile);
-				ObjectFiles.push_back(std::move(NormalizedObjectFile));
+				LiveModule::ModifiedObjFile ModifiedObjFile;
+				ModifiedObjFile.objPath = file::NormalizePath(*ObjectFile);
+				ObjectFiles.push_back(std::move(ModifiedObjFile));
 			}
 
 			m_liveModuleToModifiedOrNewObjFiles.insert(std::make_pair(ModuleFileName, std::move(ObjectFiles)));
@@ -828,14 +829,14 @@ void ServerCommandThread::CompileChanges(bool didAllProcessesMakeProgress)
 			else
 			{
 				// build a patch with the given list of .objs for this module
-				const std::vector<std::wstring>& objFiles = objFilesIt->second;
+				const types::vector<LiveModule::ModifiedObjFile>& objFiles = objFilesIt->second;
 				moduleUpdateError = liveModule->Update(&fileCache, m_directoryCache, updateType, objFiles);
 			}
 		}
 		else
 		{
 			// no optional .objs were given, update all live modules regularly
-			std::vector<std::wstring> emptyObjs;
+			types::vector<LiveModule::ModifiedObjFile> emptyObjs;
 			moduleUpdateError = liveModule->Update(&fileCache, m_directoryCache, updateType, emptyObjs);
 		}
 
@@ -1328,7 +1329,10 @@ bool ServerCommandThread::BuildPatchAction::Execute(CommandType* command, const
 
 		pipe->SendAck();
 
-		commandThread->m_liveModuleToModifiedOrNewObjFiles[packetCommand.moduleName].push_back(packetCommand.objPath);
+		LiveModule::ModifiedObjFile modifiedObjFile;
+		modifiedObjFile.objPath = packetCommand.objPath;
+		modifiedObjFile.amalgamatedObjPath = packetCommand.amalgamatedObjPath;
+		commandThread->m_liveModuleToModifiedOrNewObjFiles[packetCommand.moduleName].push_back(modifiedObjFile);
 	}
 
 	commandThread->m_manualRecompileTriggered = true;
diff --git a/Engine/Source/Developer/Windows/LiveCodingServer/Private/External/LC_ServerCommandThread.h b/Engine/Source/Developer/Windows/LiveCodingServer/Private/External/LC_ServerCommandThread.h
--- a/Engine/Source/Developer/Windows/LiveCodingServer/Private/External/LC_ServerCommandThread.h
+++ b/Engine/Source/Developer/Windows/LiveCodingServer/Private/External/LC_ServerCommandThread.h
@@ -13,6 +13,7 @@
 #include "LC_Executable.h"
 #include "LC_RunMode.h"
 #include "LC_Types.h"
+#include "LC_LiveModule.h"
 
 
 class MainFrame;
@@ -261,5 +262,5 @@ private:
 
 	// for triggering recompiles using the API
 	bool m_manualRecompileTriggered;
-	types::unordered_map<std::wstring, types::vector<std::wstring>> m_liveModuleToModifiedOrNewObjFiles;
+	types::unordered_map<std::wstring, types::vector<LiveModule::ModifiedObjFile>> m_liveModuleToModifiedOrNewObjFiles;
 };
