LiveCoding: Prevent Live Coding starting automatically when in unattended mode.

diff --git a/Engine/Source/Developer/Windows/LiveCoding/Private/LiveCodingModule.cpp b/Engine/Source/Developer/Windows/LiveCoding/Private/LiveCodingModule.cpp
--- a/Engine/Source/Developer/Windows/LiveCoding/Private/LiveCodingModule.cpp
+++ b/Engine/Source/Developer/Windows/LiveCoding/Private/LiveCodingModule.cpp
@@ -68,7 +68,7 @@ void FLiveCodingModule::StartupModule()
 	extern void Startup(Windows::HINSTANCE hInstance);
 	Startup(hInstance);
 
-	if (Settings->bEnabled)
+	if (Settings->bEnabled && !FApp::IsUnattended())
 	{
 		if(Settings->Startup == ELiveCodingStartupMode::Automatic)
 		{
