Fix non-unity
#rnx

diff --git a/Engine/Source/Developer/Windows/LiveCoding/Private/External/LC_InterprocessMutex.h b/Engine/Source/Developer/Windows/LiveCoding/Private/External/LC_InterprocessMutex.h
--- a/Engine/Source/Developer/Windows/LiveCoding/Private/External/LC_InterprocessMutex.h
+++ b/Engine/Source/Developer/Windows/LiveCoding/Private/External/LC_InterprocessMutex.h
@@ -3,6 +3,7 @@
 #pragma once
 
 #include "Windows/MinimalWindowsAPI.h"
+#include "LC_Platform.h"
 
 class InterprocessMutex
 {
