Fix missing include paths in generated project files in binary builds, due to exception while generating project files due to missing module.

diff --git a/Engine/Source/Developer/Windows/LiveCodingServer/LiveCodingServer.Build.cs b/Engine/Source/Developer/Windows/LiveCodingServer/LiveCodingServer.Build.cs
--- a/Engine/Source/Developer/Windows/LiveCodingServer/LiveCodingServer.Build.cs
+++ b/Engine/Source/Developer/Windows/LiveCodingServer/LiveCodingServer.Build.cs
@@ -8,10 +8,11 @@ public class LiveCodingServer : ModuleRules
 	public LiveCodingServer(ReadOnlyTargetRules Target) : base(Target)
 	{
 		PrivateDependencyModuleNames.Add("Core");
-		PrivateDependencyModuleNames.Add("Distorm");
 		PrivateDependencyModuleNames.Add("LiveCoding");
 
-		if (Target.Platform == UnrealTargetPlatform.Win32 || Target.Platform == UnrealTargetPlatform.Win64)
+        AddEngineThirdPartyPrivateStaticDependencies(Target, "Distorm");
+
+        if (Target.Platform == UnrealTargetPlatform.Win32 || Target.Platform == UnrealTargetPlatform.Win64)
 		{
 			PrivateIncludePaths.Add("Developer/Windows/LiveCoding/Private");
 			PrivateIncludePaths.Add("Developer/Windows/LiveCoding/Private/External");
