Live Coding: Fix crashes when patching adaptive non-unity files in game modules containing static global variables.

Live++ reads object files at startup for game modules, and assigns unique ids to each compiland (used to disambiguate static variables). When compiling the patch, these compilands are modified to reference a unique id for the unity blob, causing the variables to be reconstructed.

Solution is to generate a JSON file to each output directory containing object files containing the mapping, and to use that to assign compiland ids at startup.

diff --git a/Engine/Source/Developer/Windows/LiveCodingServer/LiveCodingServer.Build.cs b/Engine/Source/Developer/Windows/LiveCodingServer/LiveCodingServer.Build.cs
--- a/Engine/Source/Developer/Windows/LiveCodingServer/LiveCodingServer.Build.cs
+++ b/Engine/Source/Developer/Windows/LiveCodingServer/LiveCodingServer.Build.cs
@@ -7,8 +7,9 @@ public class LiveCodingServer : ModuleRules
 {
 	public LiveCodingServer(ReadOnlyTargetRules Target) : base(Target)
 	{
-		PrivateDependencyModuleNames.Add("Core");
-		PrivateDependencyModuleNames.Add("LiveCoding");
+        PrivateDependencyModuleNames.Add("Core");
+        PrivateDependencyModuleNames.Add("Json");
+        PrivateDependencyModuleNames.Add("LiveCoding");
 
         AddEngineThirdPartyPrivateStaticDependencies(Target, "Distorm");
 
diff --git a/Engine/Source/Developer/Windows/LiveCodingServer/Private/External/LC_LiveModule.cpp b/Engine/Source/Developer/Windows/LiveCodingServer/Private/External/LC_LiveModule.cpp
--- a/Engine/Source/Developer/Windows/LiveCodingServer/Private/External/LC_LiveModule.cpp
+++ b/Engine/Source/Developer/Windows/LiveCodingServer/Private/External/LC_LiveModule.cpp
@@ -39,6 +39,13 @@
 #include "Misc/Paths.h"
 // END EPIC MOD
 
+// BEGIN EPIC MOD - Allow mapping from object files to their unity object file
+#include "Misc/FileHelper.h"
+#include "Serialization/JsonReader.h"
+#include "Serialization/JsonSerializer.h"
+#include "Dom/JsonObject.h"
+// END EPIC MOD
+
 namespace
 {
 	// common linker options:
@@ -365,12 +372,114 @@ namespace
 		return LiveModule::CompileResult { exitCode, true };
 	}
 
+	// BEGIN EPIC MOD - Allow mapping from object files to their unity object file
+	static bool ReadLiveCodingInfo(const wchar_t* manifestFile, types::StringMap<uint32_t>& objFileToCompilandId)
+	{
+		// Read the file to a string
+		FString FileContents;
+		if (!FFileHelper::LoadFileToString(FileContents, manifestFile))
+		{
+			return false;
+		}
+
+		// Deserialize a JSON object from the string
+		TSharedPtr< FJsonObject > Object;
+		TSharedRef< TJsonReader<> > Reader = TJsonReaderFactory<>::Create(FileContents);
+		if (!FJsonSerializer::Deserialize(Reader, Object) || !Object.IsValid())
+		{
+			return false;
+		}
+
+		const TSharedPtr<FJsonObject>* FilesObject;
+		if (!Object->TryGetObjectField(TEXT("RemapUnityFiles"), FilesObject))
+		{
+			return false;
+		}
+
+		std::wstring BaseDir = file::GetDirectory(manifestFile);
+		for (const TPair<FString, TSharedPtr<FJsonValue>>& Pair : FilesObject->Get()->Values)
+		{
+			std::wstring UnityObjectFile = file::NormalizePath((BaseDir + L"\\" + *Pair.Key).c_str());
+			uint32_t UnityCompilandId = uniqueId::Generate(UnityObjectFile);
+			objFileToCompilandId.insert(std::make_pair(string::ToUtf8String(UnityObjectFile), UnityCompilandId));
+
+			const FJsonValue* Value = Pair.Value.Get();
+			if (Value->Type != EJson::Array)
+			{
+				return false;
+			}
+
+			const TArray<TSharedPtr<FJsonValue>>& SourceFileValues = Value->AsArray();
+			for (const TSharedPtr<FJsonValue>& SourceFileValue : SourceFileValues)
+			{
+				if (SourceFileValue->Type != EJson::String)
+				{
+					return false;
+				}
+
+				std::wstring WideObjectFile = file::NormalizePath((BaseDir + L"\\" + *SourceFileValue->AsString()).c_str());
+				ImmutableString ObjectFile = string::ToUtf8String(WideObjectFile);
+				if (objFileToCompilandId.find(ObjectFile) == objFileToCompilandId.end())
+				{
+					objFileToCompilandId.insert(std::make_pair(ObjectFile, UnityCompilandId));
+				}
+			}
+		}
+
+		return true;
+	}
+
+	static void UpdateCompilandCache(types::StringMap<symbols::Compiland*>& compilands, types::StringMap<uint32_t>& objFileToCompilandId)
+	{
+		types::unordered_set<std::wstring> Directories;
+		for (std::pair<const ImmutableString, symbols::Compiland*>& Pair : compilands)
+		{
+			symbols::Compiland* Compiland = Pair.second;
+			if (Compiland->amalgamatedUniqueId == ~(uint32_t)0)
+			{
+				const std::wstring& wideObjPath = string::ToWideString(Pair.first);
+				Directories.insert(file::GetDirectory(wideObjPath));
+			}
+		}
+		for (const std::wstring& Directory : Directories)
+		{
+			std::wstring ManifestFile = Directory + L"\\LiveCodingInfo.json";
+			if (file::DoesExist(file::GetAttributes(ManifestFile.c_str())))
+			{
+				ReadLiveCodingInfo(ManifestFile.c_str(), objFileToCompilandId);
+			}
+		}
+		for (std::pair<const ImmutableString, symbols::Compiland*>& Pair : compilands)
+		{
+			symbols::Compiland* Compiland = Pair.second;
+			if (Compiland->amalgamatedUniqueId == ~(uint32_t)0)
+			{
+				types::StringMap<uint32_t>::const_iterator Iter = objFileToCompilandId.find(Pair.first);
+				if (Iter == objFileToCompilandId.end())
+				{
+					LC_WARNING_DEV("Unable to get amalgamated id for %s", Pair.first.c_str());
+				}
+				else
+				{
+					Compiland->amalgamatedUniqueId = Iter->second;
+				}
+			}
+		}
+	}
+	// END EPIC MOD
 
 	// helper function that returns or generates the unique ID of an optional compiland.
 	// for files split off from amalgamated files, we need to use the original object path of the amalgamated file here,
 	// otherwise names of symbols would differ, leading to constructors of global instances being called again.
 	static inline uint32_t GetCompilandId(const symbols::Compiland* compiland, const wchar_t* const objPath, const types::vector<LiveModule::ModifiedObjFile>& modifiedObjFiles)
 	{
+		// BEGIN EPIC MOD - Allow mapping from object files to their unity object file
+		if (compiland && compiland->amalgamatedUniqueId != ~(uint32_t)0)
+		{
+			return compiland->amalgamatedUniqueId;
+		}
+		// END EPIC MOD
+
 		// try to find the given .obj path in the array of modified object files to check if there's an original amalgamated object path for it
 		for (size_t i = 0u; i < modifiedObjFiles.size(); ++i)
 		{
@@ -1028,6 +1137,10 @@ void LiveModule::Load(symbols::Provider* provider, symbols::DiaCompilandDB* diaC
 	{
 		LC_LOG_DEV("Caching all .objs on Load() due to external build system being used");
 
+		// BEGIN EPIC MOD - Allow mapping from object files to their unity object file
+		UpdateCompilandCache(m_compilandDB->compilands, m_objFileToCompilandId);
+		// END EPIC MOD
+
 		// the user wants to use an external build system. in this case, we only track .objs for changes and never
 		// compile anything ourselves. we cannot load .objs lazily in this case, so we have to do that right now.
 		struct GatherResult
@@ -1057,10 +1170,21 @@ void LiveModule::Load(symbols::Provider* provider, symbols::DiaCompilandDB* diaC
 
 			// do the loading and gathering concurrently
 			auto task = scheduler::CreateTask(gatherTaskRoot, [objPath, compiland]()
-			{
+				{
 				const std::wstring& wideObjPath = string::ToWideString(objPath);
 				coff::ObjFile* objFile = coff::OpenObj(wideObjPath.c_str());
-				coff::CoffDB* database = coff::GatherDatabase(objFile, compiland->uniqueId, coff::ReadFlags::NONE);
+				// BEGIN EPIC MOD - Allow mapping from object files to their unity object file
+				uint32_t uniqueId;
+			    if (compiland->amalgamatedUniqueId != ~(uint32_t)0)
+			    {
+				    uniqueId = compiland->amalgamatedUniqueId;
+			    }
+			    else
+			    {
+				    uniqueId = compiland->uniqueId;
+			    }
+				coff::CoffDB* database = coff::GatherDatabase(objFile, uniqueId, coff::ReadFlags::NONE);
+				// END EPIC MOD
 				coff::CloseObj(objFile);
 
 				return GatherResult { database, objPath };
@@ -1236,129 +1360,129 @@ LiveModule::ErrorType::Enum LiveModule::Update(FileAttributeCache* fileCache, Di
 	// only check for modifications if no files have been handed to us
 	if (modifiedOrNewObjFiles.size() == 0u)
 	{
-	// check all files whether they changed
-	for (auto compilandIt = m_compilandDB->dependencies.begin(); compilandIt != m_compilandDB->dependencies.end(); ++compilandIt)
-	{
-		symbols::Dependency* dependency = compilandIt->second;
-		if (!dependency->parentDirectory->hadChange)
-		{
-			// no need to check this compiland, the parent directory didn't notice a change
-			continue;
-		}
-
-		const std::wstring filePath = string::ToWideString(compilandIt->first);
-		const types::vector<symbols::ObjPath>& objPaths = dependency->objPaths;
-
-		const FileAttributeCache::Data& cacheData = fileCache->UpdateCacheData(filePath);
-		const uint64_t currentTime = cacheData.lastModificationTime;
-		if (currentTime != dependency->lastModification)
+		// check all files whether they changed
+		for (auto compilandIt = m_compilandDB->dependencies.begin(); compilandIt != m_compilandDB->dependencies.end(); ++compilandIt)
 		{
-			dependency->lastModification = currentTime;
+			symbols::Dependency* dependency = compilandIt->second;
+			if (!dependency->parentDirectory->hadChange)
 			{
-				const std::wstring prettyPath = file::NormalizePathWithoutLinks(filePath.c_str());
-				LC_LOG_USER("File %S was modified", prettyPath.c_str());
+				// no need to check this compiland, the parent directory didn't notice a change
+				continue;
 			}
 
-			// AMALGAMATION
-			if (appSettings::g_amalgamationSplitIntoSingleParts->GetValue())
+			const std::wstring filePath = string::ToWideString(compilandIt->first);
+			const types::vector<symbols::ObjPath>& objPaths = dependency->objPaths;
+
+			const FileAttributeCache::Data& cacheData = fileCache->UpdateCacheData(filePath);
+			const uint64_t currentTime = cacheData.lastModificationTime;
+			if (currentTime != dependency->lastModification)
 			{
-				// look at each file individually and determine what to do
-				for (auto it : objPaths)
+				dependency->lastModification = currentTime;
 				{
-					symbols::Compiland* compiland = symbols::FindCompiland(m_compilandDB, it);
-					if (compiland)
+					const std::wstring prettyPath = file::NormalizePathWithoutLinks(filePath.c_str());
+					LC_LOG_USER("File %S was modified", prettyPath.c_str());
+				}
+
+				// AMALGAMATION
+				if (appSettings::g_amalgamationSplitIntoSingleParts->GetValue())
+				{
+					// look at each file individually and determine what to do
+					for (auto it : objPaths)
 					{
-						if (symbols::IsAmalgamation(compiland))
+						symbols::Compiland* compiland = symbols::FindCompiland(m_compilandDB, it);
+						if (compiland)
 						{
-							// split amalgamated file
-							symbols::AmalgamatedCompiland* amalgamatedCompiland = symbols::FindAmalgamatedCompiland(m_compilandDB, it);
-							if (amalgamatedCompiland)
+							if (symbols::IsAmalgamation(compiland))
 							{
-								// the amalgamated compiland needs to be split into its single parts.
-								// add all compilands that are part of the amalgamation for compilation.
-								// we always split in this case to trigger recompiles when included headers change.
-								LC_LOG_USER("Splitting amalgamated/unity file %s", it.c_str());
-
-								if (!amalgamatedCompiland->isSplit)
+								// split amalgamated file
+								symbols::AmalgamatedCompiland* amalgamatedCompiland = symbols::FindAmalgamatedCompiland(m_compilandDB, it);
+								if (amalgamatedCompiland)
 								{
-									// this is the first time the amalgamation is split into single files
-									forceAmalgamationPartsLinkage = true;
-								}
+									// the amalgamated compiland needs to be split into its single parts.
+									// add all compilands that are part of the amalgamation for compilation.
+									// we always split in this case to trigger recompiles when included headers change.
+									LC_LOG_USER("Splitting amalgamated/unity file %s", it.c_str());
 
-								m_modifiedFiles.insert(amalgamatedCompiland->singleParts.begin(), amalgamatedCompiland->singleParts.end());
-								amalgamatedCompiland->isSplit = true;
+									if (!amalgamatedCompiland->isSplit)
+									{
+										// this is the first time the amalgamation is split into single files
+										forceAmalgamationPartsLinkage = true;
+									}
+
+									m_modifiedFiles.insert(amalgamatedCompiland->singleParts.begin(), amalgamatedCompiland->singleParts.end());
+									amalgamatedCompiland->isSplit = true;
+								}
 							}
-						}
-						else if (symbols::IsPartOfAmalgamation(compiland))
-						{
-							// this file is part of an amalgamation.
-							// if the amalgamation needs to be split, do that now.
-							// in any case, this file needs to be recompiled.
-							m_modifiedFiles.insert(it);
-
-							// find the amalgamated compiland this file belongs to
-							const ImmutableString& amalgamatedObjPath = compiland->amalgamationPath;
-							symbols::AmalgamatedCompiland* amalgamatedCompiland = symbols::FindAmalgamatedCompiland(m_compilandDB, amalgamatedObjPath);
-							if (amalgamatedCompiland)
+							else if (symbols::IsPartOfAmalgamation(compiland))
 							{
-								if (!amalgamatedCompiland->isSplit)
+								// this file is part of an amalgamation.
+								// if the amalgamation needs to be split, do that now.
+								// in any case, this file needs to be recompiled.
+								m_modifiedFiles.insert(it);
+
+								// find the amalgamated compiland this file belongs to
+								const ImmutableString& amalgamatedObjPath = compiland->amalgamationPath;
+								symbols::AmalgamatedCompiland* amalgamatedCompiland = symbols::FindAmalgamatedCompiland(m_compilandDB, amalgamatedObjPath);
+								if (amalgamatedCompiland)
 								{
-									// this is the first time the amalgamation is split into single files
-									forceAmalgamationPartsLinkage = true;
+									if (!amalgamatedCompiland->isSplit)
+									{
+										// this is the first time the amalgamation is split into single files
+										forceAmalgamationPartsLinkage = true;
 
-									// the amalgamated compiland needs to be split into its single parts.
-									// add all compilands that are part of the amalgamation for compilation, and mark the
-									// amalgamated compiland as being split.
-									LC_LOG_USER("Splitting amalgamated/unity file %s", amalgamatedObjPath.c_str());
+										// the amalgamated compiland needs to be split into its single parts.
+										// add all compilands that are part of the amalgamation for compilation, and mark the
+										// amalgamated compiland as being split.
+										LC_LOG_USER("Splitting amalgamated/unity file %s", amalgamatedObjPath.c_str());
 
-									m_modifiedFiles.insert(amalgamatedCompiland->singleParts.begin(), amalgamatedCompiland->singleParts.end());
-									amalgamatedCompiland->isSplit = true;
+										m_modifiedFiles.insert(amalgamatedCompiland->singleParts.begin(), amalgamatedCompiland->singleParts.end());
+										amalgamatedCompiland->isSplit = true;
+									}
 								}
 							}
-						}
-						else
-						{
-							m_modifiedFiles.insert(it);
+							else
+							{
+								m_modifiedFiles.insert(it);
+							}
 						}
 					}
 				}
-			}
-			else
-			{
-				// don't need to do anything fancy, just add all affected .objs
-				m_modifiedFiles.insert(objPaths.begin(), objPaths.end());
+				else
+				{
+					// don't need to do anything fancy, just add all affected .objs
+					m_modifiedFiles.insert(objPaths.begin(), objPaths.end());
+				}
 			}
 		}
-	}
 
-	if (m_runMode == RunMode::DEFAULT)
-	{
-		if (m_modifiedFiles.size() == 0u)
+		if (m_runMode == RunMode::DEFAULT)
 		{
-			if (m_compiledCompilands.size() == 0u)
+			if (m_modifiedFiles.size() == 0u)
 			{
-				// no change detected in this module
-				return ErrorType::NO_CHANGE;
+				if (m_compiledCompilands.size() == 0u)
+				{
+					// no change detected in this module
+					return ErrorType::NO_CHANGE;
+				}
+				else
+				{
+					// there are still compiled files that haven't been linked
+				}
 			}
 			else
 			{
-				// there are still compiled files that haven't been linked
+				LC_LOG_USER("Detected %zu file(s) to be compiled for Live++ module %S", m_modifiedFiles.size(), m_moduleName.c_str());
 			}
 		}
-		else
+		else if (m_runMode == RunMode::EXTERNAL_BUILD_SYSTEM)
 		{
-			LC_LOG_USER("Detected %zu file(s) to be compiled for Live++ module %S", m_modifiedFiles.size(), m_moduleName.c_str());
-		}
-	}
-	else if (m_runMode == RunMode::EXTERNAL_BUILD_SYSTEM)
-	{
-		if (m_modifiedFiles.size() == 0u)
-		{
-			// no changed .obj detected in this module
-			return ErrorType::NO_CHANGE;
+			if (m_modifiedFiles.size() == 0u)
+			{
+				// no changed .obj detected in this module
+				return ErrorType::NO_CHANGE;
+			}
 		}
 	}
-	}
 	else
 	{
 		for (size_t i = 0u; i < modifiedOrNewObjFiles.size(); ++i)
@@ -1816,16 +1940,16 @@ LiveModule::ErrorType::Enum LiveModule::Update(FileAttributeCache* fileCache, Di
 	{
 		if (modifiedOrNewObjFiles.size() == 0u)
 		{
-		// files were compiled by an external build system, we just have to mark them appropriately
-		const size_t count = availableModifiedFiles.size();
-		for (size_t i = 0u; i < count; ++i)
-		{
-			const symbols::ObjPath& objPath = availableModifiedFiles[i].objPath;
-			symbols::Compiland* compiland = availableModifiedFiles[i].compiland;
+			// files were compiled by an external build system, we just have to mark them appropriately
+			const size_t count = availableModifiedFiles.size();
+			for (size_t i = 0u; i < count; ++i)
+			{
+				const symbols::ObjPath& objPath = availableModifiedFiles[i].objPath;
+				symbols::Compiland* compiland = availableModifiedFiles[i].compiland;
 
-			m_compiledCompilands.emplace(objPath, compiland);
-			symbols::MarkCompilandAsRecompiled(compiland);
-		}
+				m_compiledCompilands.emplace(objPath, compiland);
+				symbols::MarkCompilandAsRecompiled(compiland);
+			}
 		}
 		else
 		{
@@ -1953,26 +2077,26 @@ LiveModule::ErrorType::Enum LiveModule::Update(FileAttributeCache* fileCache, Di
 
 					Helper::UpdateExternalSymbolsAndNeededFiles(objPath, compiland, compiland->uniqueId, coffReadFlags, m_pchSymbolToCompilandName, externalSymbols, neededCompilands);
 				}
-			else
-			{
-				// this file has not changed, so consult the cache for external symbols
-				auto cacheIt = m_externalSymbolsPerCompilandCache.find(objPath);
-				if (cacheIt != m_externalSymbolsPerCompilandCache.end())
+				else
 				{
-					const size_t symbolCount = cacheIt->second.size();
-					for (size_t i = 0u; i < symbolCount; ++i)
+					// this file has not changed, so consult the cache for external symbols
+					auto cacheIt = m_externalSymbolsPerCompilandCache.find(objPath);
+					if (cacheIt != m_externalSymbolsPerCompilandCache.end())
+					{
+						const size_t symbolCount = cacheIt->second.size();
+						for (size_t i = 0u; i < symbolCount; ++i)
+						{
+							const ImmutableString& symbolName = cacheIt->second[i]->name;
+							externalSymbols.emplace(symbolName, CompilandInfo { objPath, compiland });
+						}
+					}
+					else
 					{
-						const ImmutableString& symbolName = cacheIt->second[i]->name;
-						externalSymbols.emplace(symbolName, CompilandInfo { objPath, compiland });
+						// this compiland does not store any external symbol
 					}
 				}
-				else
-				{
-					// this compiland does not store any external symbol
-				}
 			}
 		}
-	}
 		else
 		{
 			for (auto it = modifiedOrNewObjFiles.begin(); it != modifiedOrNewObjFiles.end(); ++it)
@@ -2125,6 +2249,10 @@ LiveModule::ErrorType::Enum LiveModule::Update(FileAttributeCache* fileCache, Di
 		}
 	}
 
+	// BEGIN EPIC MOD - Allow mapping from object files to their unity object file
+	UpdateCompilandCache(m_compiledCompilands, m_objFileToCompilandId);
+	// END EPIC MOD
+
 	// update the COFF cache for all compiled files
 	UpdateCoffCache(m_compiledCompilands, m_coffCache, CacheUpdate::ALL, coffReadFlags, modifiedOrNewObjFiles);
 
diff --git a/Engine/Source/Developer/Windows/LiveCodingServer/Private/External/LC_LiveModule.h b/Engine/Source/Developer/Windows/LiveCodingServer/Private/External/LC_LiveModule.h
--- a/Engine/Source/Developer/Windows/LiveCodingServer/Private/External/LC_LiveModule.h
+++ b/Engine/Source/Developer/Windows/LiveCodingServer/Private/External/LC_LiveModule.h
@@ -140,4 +140,8 @@ private:
 
 	// all patches loaded so far along with recorded data how to load them into other processes
 	types::vector<ModulePatch*> m_compiledModulePatches;
+
+	// BEGIN EPIC MOD - Allow mapping from object files to their unity object file
+	types::StringMap<uint32_t> m_objFileToCompilandId;
+	// END EPIC MOD
 };
diff --git a/Engine/Source/Developer/Windows/LiveCodingServer/Private/External/LC_ServerCommandThread.cpp b/Engine/Source/Developer/Windows/LiveCodingServer/Private/External/LC_ServerCommandThread.cpp
--- a/Engine/Source/Developer/Windows/LiveCodingServer/Private/External/LC_ServerCommandThread.cpp
+++ b/Engine/Source/Developer/Windows/LiveCodingServer/Private/External/LC_ServerCommandThread.cpp
@@ -708,8 +708,7 @@ void ServerCommandThread::CompileChanges(bool didAllProcessesMakeProgress)
 		GLiveCodingServer->GetStatusChangeDelegate().ExecuteIfBound(L"Compiling changes for live coding...");
 
 		TMap<FString, TArray<FString>> ModuleToObjectFiles;
-		TMap<FString, FString> ObjectFileToUnityObjectFile;
-		if (!CompileDelegate.Execute(Targets, ModuleToObjectFiles, ObjectFileToUnityObjectFile))
+		if (!CompileDelegate.Execute(Targets, ModuleToObjectFiles))
 		{
 			GLiveCodingServer->GetCompileFinishedDelegate().ExecuteIfBound(ELiveCodingResult::Error, L"Compilation error.");
 			return;
@@ -777,13 +776,6 @@ void ServerCommandThread::CompileChanges(bool didAllProcessesMakeProgress)
 			{
 				LiveModule::ModifiedObjFile ModifiedObjFile;
 				ModifiedObjFile.objPath = file::NormalizePath(*ObjectFile);
-
-				const FString* UnityObjectFile = ObjectFileToUnityObjectFile.Find(ObjectFile);
-				if (UnityObjectFile != nullptr)
-				{
-					ModifiedObjFile.amalgamatedObjPath = file::NormalizePath(**UnityObjectFile);
-				}
-
 				ObjectFiles.push_back(std::move(ModifiedObjFile));
 			}
 
diff --git a/Engine/Source/Developer/Windows/LiveCodingServer/Private/External/LC_Symbols.h b/Engine/Source/Developer/Windows/LiveCodingServer/Private/External/LC_Symbols.h
--- a/Engine/Source/Developer/Windows/LiveCodingServer/Private/External/LC_Symbols.h
+++ b/Engine/Source/Developer/Windows/LiveCodingServer/Private/External/LC_Symbols.h
@@ -88,6 +88,10 @@ namespace symbols
 		bool isPartOfLibrary;
 		bool wasRecompiled;
 
+		// BEGIN EPIC MOD - Allow mapping from object files to their unity object file
+		uint32_t amalgamatedUniqueId = ~(uint32_t)0;
+		// END EPIC MOD
+
 		LC_DISABLE_ASSIGNMENT(Compiland);
 		LC_DISABLE_COPY(Compiland);
 	};
diff --git a/Engine/Source/Developer/Windows/LiveCodingServer/Public/ILiveCodingServer.h b/Engine/Source/Developer/Windows/LiveCodingServer/Public/ILiveCodingServer.h
--- a/Engine/Source/Developer/Windows/LiveCodingServer/Public/ILiveCodingServer.h
+++ b/Engine/Source/Developer/Windows/LiveCodingServer/Public/ILiveCodingServer.h
@@ -44,8 +44,7 @@ public:
 	virtual FLogOutputDelegate& GetLogOutputDelegate() = 0;
 
 	typedef TMap<FString, TArray<FString>> FModuleToObjectFiles;
-	typedef TMap<FString, FString> FObjectFileToUnityObjectFile;
-	DECLARE_DELEGATE_RetVal_ThreeParams(bool, FCompileDelegate, const TArray<FString>&, FModuleToObjectFiles&, FObjectFileToUnityObjectFile&)
+	DECLARE_DELEGATE_RetVal_TwoParams(bool, FCompileDelegate, const TArray<FString>&, FModuleToObjectFiles&)
 	virtual FCompileDelegate& GetCompileDelegate() = 0;
 
 	DECLARE_DELEGATE(FCompileStartedDelegate);
diff --git a/Engine/Source/Programs/LiveCodingConsole/Private/LiveCodingConsole.cpp b/Engine/Source/Programs/LiveCodingConsole/Private/LiveCodingConsole.cpp
--- a/Engine/Source/Programs/LiveCodingConsole/Private/LiveCodingConsole.cpp
+++ b/Engine/Source/Programs/LiveCodingConsole/Private/LiveCodingConsole.cpp
@@ -179,7 +179,7 @@ private:
 		LogWidget->AppendLine(GetLogColor(Verbosity), MoveTemp(Text));
 	}
 
-	bool CompilePatch(const TArray<FString>& Targets, TMap<FString, TArray<FString>>& ModuleToObjectFiles, TMap<FString, FString>& ObjectFileToUnityObjectFile)
+	bool CompilePatch(const TArray<FString>& Targets, TMap<FString, TArray<FString>>& ModuleToObjectFiles)
 	{
 		// Update the compile start time. This gets copied into the last patch time once a patch has been confirmed to have been applied.
 		NextPatchStartTime = FDateTime::UtcNow();
@@ -250,9 +250,6 @@ private:
 				}
 			}
 		}
-
-		// Save off the mapping to unity file
-		ObjectFileToUnityObjectFile = MoveTemp(Manifest.ObjectFileToUnityObjectFile);
 		return true;
 	}
 
diff --git a/Engine/Source/Programs/LiveCodingConsole/Private/LiveCodingManifest.cpp b/Engine/Source/Programs/LiveCodingConsole/Private/LiveCodingManifest.cpp
--- a/Engine/Source/Programs/LiveCodingConsole/Private/LiveCodingManifest.cpp
+++ b/Engine/Source/Programs/LiveCodingConsole/Private/LiveCodingManifest.cpp
@@ -96,36 +96,5 @@ bool FLiveCodingManifest::Parse(FJsonObject& Object, FString& OutFailReason)
 		}
 	}
 
-	const TArray<TSharedPtr<FJsonValue>>* AdaptiveUnityFilesArray;
-	if (Object.TryGetArrayField(TEXT("AdaptiveUnityFiles"), AdaptiveUnityFilesArray))
-	{
-		for (const TSharedPtr<FJsonValue>& AdaptiveUnityFileValue : *AdaptiveUnityFilesArray)
-		{
-			if (AdaptiveUnityFileValue->Type != EJson::Object)
-			{
-				OutFailReason = TEXT("invalid module object");
-				return false;
-			}
-
-			const FJsonObject& AdaptiveUnityFileObject = *AdaptiveUnityFileValue->AsObject();
-
-			FString ObjectFile;
-			if (!AdaptiveUnityFileObject.TryGetStringField(TEXT("ObjectFile"), ObjectFile))
-			{
-				OutFailReason = TEXT("missing module 'ObjectFile' field");
-				return false;
-			}
-
-			FString UnityObjectFile;
-			if (!AdaptiveUnityFileObject.TryGetStringField(TEXT("UnityObjectFile"), UnityObjectFile))
-			{
-				OutFailReason = TEXT("missing module 'UnityObjectFile' field");
-				return false;
-			}
-
-			ObjectFileToUnityObjectFile.Emplace(MoveTemp(ObjectFile), MoveTemp(UnityObjectFile));
-		}
-	}
-
 	return true;
 }
diff --git a/Engine/Source/Programs/LiveCodingConsole/Private/LiveCodingManifest.h b/Engine/Source/Programs/LiveCodingConsole/Private/LiveCodingManifest.h
--- a/Engine/Source/Programs/LiveCodingConsole/Private/LiveCodingManifest.h
+++ b/Engine/Source/Programs/LiveCodingConsole/Private/LiveCodingManifest.h
@@ -11,7 +11,6 @@ struct FLiveCodingManifest
 	FString LinkerPath;
 	TMap<FString, FString> LinkerEnvironment;
 	TMap<FString, TArray<FString>> BinaryToObjectFiles;
-	TMap<FString, FString> ObjectFileToUnityObjectFile;
 
 	bool Read(const TCHAR* FileName, FString& OutFailReason);
 	bool Parse(FJsonObject& Object, FString& OutFailReason);
diff --git a/Engine/Source/Programs/UnrealBuildTool/Configuration/UEBuildModuleCPP.cs b/Engine/Source/Programs/UnrealBuildTool/Configuration/UEBuildModuleCPP.cs
--- a/Engine/Source/Programs/UnrealBuildTool/Configuration/UEBuildModuleCPP.cs
+++ b/Engine/Source/Programs/UnrealBuildTool/Configuration/UEBuildModuleCPP.cs
@@ -415,11 +415,14 @@ namespace UnrealBuildTool
 			// Write all the definitions to a separate file
 			CreateHeaderForDefinitions(CompileEnvironment, IntermediateDirectory, null);
 
+			// Mapping of source file to unity file. We output this to intermediate directories for other tools (eg. live coding) to use.
+			Dictionary<FileItem, FileItem> SourceFileToUnityFile = new Dictionary<FileItem, FileItem>();
+
 			// Compile CPP files
 			List<FileItem> CPPFilesToCompile = InputFiles.CPPFiles;
 			if (bModuleUsesUnityBuild)
 			{
-				CPPFilesToCompile = Unity.GenerateUnityCPPs(Target, CPPFilesToCompile, CompileEnvironment, WorkingSet, Rules.ShortName ?? Name, IntermediateDirectory, Makefile);
+				CPPFilesToCompile = Unity.GenerateUnityCPPs(Target, CPPFilesToCompile, CompileEnvironment, WorkingSet, Rules.ShortName ?? Name, IntermediateDirectory, Makefile, SourceFileToUnityFile);
 				LinkInputFiles.AddRange(CompileUnityFilesWithToolChain(Target, ToolChain, CompileEnvironment, ModuleCompileEnvironment, CPPFilesToCompile, Makefile.Actions).ObjectFiles);
 			}
 			else
@@ -466,7 +469,7 @@ namespace UnrealBuildTool
 
 					if (bModuleUsesUnityBuild)
 					{
-						GeneratedFileItems = Unity.GenerateUnityCPPs(Target, GeneratedFileItems, GeneratedCPPCompileEnvironment, WorkingSet, (Rules.ShortName ?? Name) + ".gen", IntermediateDirectory, Makefile);
+						GeneratedFileItems = Unity.GenerateUnityCPPs(Target, GeneratedFileItems, GeneratedCPPCompileEnvironment, WorkingSet, (Rules.ShortName ?? Name) + ".gen", IntermediateDirectory, Makefile, SourceFileToUnityFile);
 						LinkInputFiles.AddRange(CompileUnityFilesWithToolChain(Target, ToolChain, GeneratedCPPCompileEnvironment, ModuleCompileEnvironment, GeneratedFileItems, Makefile.Actions).ObjectFiles);
 					}
 					else
@@ -518,6 +521,28 @@ namespace UnrealBuildTool
 				Manifest.WriteIfModified(PrecompiledManifestLocation);
 			}
 
+			// Write a mapping of unity object file to standalone object file for live coding
+			if(Rules.Target.bWithLiveCoding)
+			{
+				FileReference UnityManifestFile = FileReference.Combine(IntermediateDirectory, "LiveCodingInfo.json");
+				using (JsonWriter Writer = new JsonWriter(UnityManifestFile))
+				{
+					Writer.WriteObjectStart();
+					Writer.WriteObjectStart("RemapUnityFiles");
+					foreach (IGrouping<FileItem, KeyValuePair<FileItem, FileItem>> UnityGroup in SourceFileToUnityFile.GroupBy(x => x.Value))
+					{
+						Writer.WriteArrayStart(UnityGroup.Key.Location.GetFileName() + ".obj");
+						foreach (FileItem SourceFile in UnityGroup.Select(x => x.Key))
+						{
+							Writer.WriteValue(SourceFile.Location.GetFileName() + ".obj");
+						}
+						Writer.WriteArrayEnd();
+					}
+					Writer.WriteObjectEnd();
+					Writer.WriteObjectEnd();
+				}
+			}
+
 			return LinkInputFiles;
 		}
 
diff --git a/Engine/Source/Programs/UnrealBuildTool/Modes/BuildMode.cs b/Engine/Source/Programs/UnrealBuildTool/Modes/BuildMode.cs
--- a/Engine/Source/Programs/UnrealBuildTool/Modes/BuildMode.cs
+++ b/Engine/Source/Programs/UnrealBuildTool/Modes/BuildMode.cs
@@ -232,7 +232,7 @@ namespace UnrealBuildTool
 				if(TargetDescriptor.LiveCodingManifest != null)
 				{
 					TargetMakefile Makefile = Makefiles[TargetIdx];
-					HotReload.WriteLiveCodeManifest(TargetDescriptor.LiveCodingManifest, Makefile.SourceFileToUnityFile, Makefile.Actions);
+					HotReload.WriteLiveCodingManifest(TargetDescriptor.LiveCodingManifest, Makefile.Actions);
 				}
 			}
 
diff --git a/Engine/Source/Programs/UnrealBuildTool/System/HotReload.cs b/Engine/Source/Programs/UnrealBuildTool/System/HotReload.cs
--- a/Engine/Source/Programs/UnrealBuildTool/System/HotReload.cs
+++ b/Engine/Source/Programs/UnrealBuildTool/System/HotReload.cs
@@ -727,9 +727,8 @@ namespace UnrealBuildTool
 		/// Writes a manifest containing all the information needed to create a live coding patch
 		/// </summary>
 		/// <param name="ManifestFile">File to write to</param>
-		/// <param name="SourceFileToUnityFile">Mapping of source file to the unity file that would contain it</param>
 		/// <param name="Actions">List of actions that are part of the graph</param>
-		public static void WriteLiveCodeManifest(FileReference ManifestFile, Dictionary<FileItem, FileItem> SourceFileToUnityFile, List<Action> Actions)
+		public static void WriteLiveCodingManifest(FileReference ManifestFile, List<Action> Actions)
 		{
 			// Find all the output object files
 			HashSet<FileItem> ObjectFiles = new HashSet<FileItem>();
@@ -741,30 +740,6 @@ namespace UnrealBuildTool
 				}
 			}
 
-			// Find any action that has a non-unity cpp prerequisite, and map it to the appropriate 
-			Dictionary<FileItem, FileItem> ObjectFileToUnityObjectFile = new Dictionary<FileItem, FileItem>();
-			foreach(Action Action in Actions)
-			{
-				if (Action.ActionType == ActionType.Compile)
-				{
-					foreach(FileItem PrerequisiteItem in Action.PrerequisiteItems)
-					{
-						FileItem UnityFile;
-						if(SourceFileToUnityFile.TryGetValue(PrerequisiteItem, out UnityFile))
-						{
-							foreach(FileItem ProducedItem in Action.ProducedItems)
-							{
-								if(ProducedItem.HasExtension(".obj"))
-								{
-									FileItem UnityObjFile = FileItem.GetItemByFileReference(new FileReference(UnityFile.Location.FullName + ".obj"));
-									ObjectFileToUnityObjectFile[ProducedItem] = UnityObjFile;
-								}
-							}
-						}
-					}
-				}
-			}
-
 			// Write the output manifest
 			using (JsonWriter Writer = new JsonWriter(ManifestFile))
 			{
@@ -810,26 +785,6 @@ namespace UnrealBuildTool
 				}
 				Writer.WriteArrayEnd();
 
-				Writer.WriteArrayStart("AdaptiveUnityFiles");
-				foreach (Action Action in Actions)
-				{
-					if (Action.ActionType == ActionType.Link && Action.ProducedItems.Any(x => x.HasExtension(".exe") || x.HasExtension(".dll")))
-					{
-						foreach (FileItem InputFile in Action.PrerequisiteItems)
-						{
-							FileItem UnityFile;
-							if (ObjectFileToUnityObjectFile.TryGetValue(InputFile, out UnityFile))
-							{
-								Writer.WriteObjectStart();
-								Writer.WriteValue("ObjectFile", InputFile.AbsolutePath);
-								Writer.WriteValue("UnityObjectFile", UnityFile.AbsolutePath);
-								Writer.WriteObjectEnd();
-							}
-						}
-					}
-				}
-				Writer.WriteArrayEnd();
-
 				Writer.WriteObjectEnd();
 			}
 		}
diff --git a/Engine/Source/Programs/UnrealBuildTool/System/TargetMakefile.cs b/Engine/Source/Programs/UnrealBuildTool/System/TargetMakefile.cs
--- a/Engine/Source/Programs/UnrealBuildTool/System/TargetMakefile.cs
+++ b/Engine/Source/Programs/UnrealBuildTool/System/TargetMakefile.cs
@@ -20,7 +20,7 @@ namespace UnrealBuildTool
 		/// <summary>
 		/// The version number to write
 		/// </summary>
-		public const int CurrentVersion = 12;
+		public const int CurrentVersion = 13;
 
 		/// <summary>
 		/// The time at which the makefile was created
@@ -113,11 +113,6 @@ namespace UnrealBuildTool
 		/// </summary>
 		public HashSet<FileItem> CandidatesForWorkingSet = new HashSet<FileItem>();
 
-		/// <summary>
-		/// Mapping from source file to the unity file that contains it. Used for Live Coding. Note that this mapping includes files that are part of the working set, even if they are now a separate compiland.
-		/// </summary>
-		public Dictionary<FileItem, FileItem> SourceFileToUnityFile = new Dictionary<FileItem, FileItem>();
-
 		/// <summary>
 		/// Maps each target to a list of UObject module info structures
 		/// </summary>
@@ -195,7 +190,6 @@ namespace UnrealBuildTool
 			DirectoryToSourceFiles = Reader.ReadDictionary(() => Reader.ReadDirectoryItem(), () => Reader.ReadArray(() => Reader.ReadFileItem()));
 			WorkingSet = Reader.ReadHashSet(() => Reader.ReadFileItem());
 			CandidatesForWorkingSet = Reader.ReadHashSet(() => Reader.ReadFileItem());
-			SourceFileToUnityFile = Reader.ReadDictionary(() => Reader.ReadFileItem(), () => Reader.ReadFileItem());
 			UObjectModules = Reader.ReadList(() => new UHTModuleInfo(Reader));
 			UObjectModuleHeaders = Reader.ReadList(() => new UHTModuleHeaderInfo(Reader));
 			PluginFiles = Reader.ReadHashSet(() => Reader.ReadFileItem());
@@ -226,7 +220,6 @@ namespace UnrealBuildTool
 			Writer.WriteDictionary(DirectoryToSourceFiles, k => Writer.WriteDirectoryItem(k), v => Writer.WriteArray(v, e => Writer.WriteFileItem(e)));
 			Writer.WriteHashSet(WorkingSet, x => Writer.WriteFileItem(x));
 			Writer.WriteHashSet(CandidatesForWorkingSet, x => Writer.WriteFileItem(x));
-			Writer.WriteDictionary(SourceFileToUnityFile, k => Writer.WriteFileItem(k), v => Writer.WriteFileItem(v));
 			Writer.WriteList(UObjectModules, e => e.Write(Writer));
 			Writer.WriteList(UObjectModuleHeaders, x => x.Write(Writer));
 			Writer.WriteHashSet(PluginFiles, x => Writer.WriteFileItem(x));
diff --git a/Engine/Source/Programs/UnrealBuildTool/System/Unity.cs b/Engine/Source/Programs/UnrealBuildTool/System/Unity.cs
--- a/Engine/Source/Programs/UnrealBuildTool/System/Unity.cs
+++ b/Engine/Source/Programs/UnrealBuildTool/System/Unity.cs
@@ -158,6 +158,7 @@ namespace UnrealBuildTool
 		/// <param name="BaseName">Base name to use for the Unity files</param>
 		/// <param name="IntermediateDirectory">Intermediate directory for unity cpp files</param>
 		/// <param name="Makefile">The makefile being built</param>
+		/// <param name="SourceFileToUnityFile">Receives a mapping of source file to unity file</param>
 		/// <returns>The "unity" C++ files.</returns>
 		public static List<FileItem> GenerateUnityCPPs(
 			ReadOnlyTargetRules Target,
@@ -166,7 +167,8 @@ namespace UnrealBuildTool
 			ISourceFileWorkingSet WorkingSet,
 			string BaseName,
 			DirectoryReference IntermediateDirectory,
-			TargetMakefile Makefile
+			TargetMakefile Makefile,
+			Dictionary<FileItem, FileItem> SourceFileToUnityFile
 			)
 		{
 			List<FileItem> NewCPPFiles = new List<FileItem>();
@@ -341,11 +343,11 @@ namespace UnrealBuildTool
 				// Store the mapping of source files to unity files in the makefile
 				foreach(FileItem SourceFile in UnityFile.Files)
 				{
-					Makefile.SourceFileToUnityFile[SourceFile] = UnityCPPFile;
+					SourceFileToUnityFile[SourceFile] = UnityCPPFile;
 				}
 				foreach (FileItem SourceFile in UnityFile.VirtualFiles)
 				{
-					Makefile.SourceFileToUnityFile[SourceFile] = UnityCPPFile;
+					SourceFileToUnityFile[SourceFile] = UnityCPPFile;
 				}
 			}
 
