Update to latest version of Live Coding. Fixes issues with symbols not being stripped if they can throw an exception during construction.

diff --git a/Engine/Source/Developer/Windows/LiveCoding/LiveCoding.Build.cs b/Engine/Source/Developer/Windows/LiveCoding/LiveCoding.Build.cs
--- a/Engine/Source/Developer/Windows/LiveCoding/LiveCoding.Build.cs
+++ b/Engine/Source/Developer/Windows/LiveCoding/LiveCoding.Build.cs
@@ -10,5 +10,7 @@ public class LiveCoding : ModuleRules
 		PrivateDependencyModuleNames.Add("Core");
 		PrivateDependencyModuleNames.Add("CoreUObject");
 		PrivateDependencyModuleNames.Add("Settings");
+
+		PublicAdditionalLibraries.Add(Path.Combine(ModuleDirectory, "Private", "External", "LC_JumpToSelf.lib"));
 	}
 }
diff --git a/Engine/Source/Developer/Windows/LiveCoding/Private/External/LC_ClientCommandActions.cpp b/Engine/Source/Developer/Windows/LiveCoding/Private/External/LC_ClientCommandActions.cpp
--- a/Engine/Source/Developer/Windows/LiveCoding/Private/External/LC_ClientCommandActions.cpp
+++ b/Engine/Source/Developer/Windows/LiveCoding/Private/External/LC_ClientCommandActions.cpp
@@ -10,7 +10,7 @@
 #include "LC_Logging.h"
 
 
-bool actions::RegisterProcessFinished::Execute(const CommandType* command, const DuplexPipe* pipe, void* context)
+bool actions::RegisterProcessFinished::Execute(const CommandType* command, const DuplexPipe* pipe, void* context, const void*, size_t)
 {
 	pipe->SendAck();
 
@@ -22,7 +22,7 @@ bool actions::RegisterProcessFinished::Execute(const CommandType* command, const
 }
 
 
-bool actions::EnableModuleFinished::Execute(const CommandType* command, const DuplexPipe* pipe, void*)
+bool actions::EnableModulesFinished::Execute(const CommandType* command, const DuplexPipe* pipe, void*, const void*, size_t)
 {
 	Event* event = static_cast<Event*>(command->token);
 	event->Signal();
@@ -32,7 +32,7 @@ bool actions::EnableModuleFinished::Execute(const CommandType* command, const Du
 }
 
 
-bool actions::DisableModuleFinished::Execute(const CommandType* command, const DuplexPipe* pipe, void*)
+bool actions::DisableModulesFinished::Execute(const CommandType* command, const DuplexPipe* pipe, void*, const void*, size_t)
 {
 	Event* event = static_cast<Event*>(command->token);
 	event->Signal();
@@ -42,27 +42,7 @@ bool actions::DisableModuleFinished::Execute(const CommandType* command, const D
 }
 
 
-bool actions::EnableAllModulesFinished::Execute(const CommandType* command, const DuplexPipe* pipe, void*)
-{
-	Event* event = static_cast<Event*>(command->token);
-	event->Signal();
-	pipe->SendAck();
-
-	return false;
-}
-
-
-bool actions::DisableAllModulesFinished::Execute(const CommandType* command, const DuplexPipe* pipe, void*)
-{
-	Event* event = static_cast<Event*>(command->token);
-	event->Signal();
-	pipe->SendAck();
-
-	return false;
-}
-
-
-bool actions::EnterSyncPoint::Execute(const CommandType*, const DuplexPipe* pipe, void*)
+bool actions::EnterSyncPoint::Execute(const CommandType*, const DuplexPipe* pipe, void*, const void*, size_t)
 {
 	syncPoint::Enter();
 	pipe->SendAck();
@@ -71,7 +51,7 @@ bool actions::EnterSyncPoint::Execute(const CommandType*, const DuplexPipe* pipe
 }
 
 
-bool actions::LeaveSyncPoint::Execute(const CommandType*, const DuplexPipe* pipe, void*)
+bool actions::LeaveSyncPoint::Execute(const CommandType*, const DuplexPipe* pipe, void*, const void*, size_t)
 {
 	syncPoint::Leave();
 	pipe->SendAck();
@@ -80,7 +60,7 @@ bool actions::LeaveSyncPoint::Execute(const CommandType*, const DuplexPipe* pipe
 }
 
 
-bool actions::CallHooks::Execute(const CommandType* command, const DuplexPipe* pipe, void*)
+bool actions::CallHooks::Execute(const CommandType* command, const DuplexPipe* pipe, void*, const void*, size_t)
 {
 	for (const hook::Function* hook = command->first; hook < command->last; ++hook)
 	{
@@ -98,35 +78,6 @@ bool actions::CallHooks::Execute(const CommandType* command, const DuplexPipe* p
 }
 
 
-bool actions::GetModule::Execute(const CommandType* command, const DuplexPipe* pipe, void*)
-{
-	HMODULE module = ::GetModuleHandleW(command->path);
-	if (module)
-	{
-		wchar_t fullPath[MAX_PATH];
-		::GetModuleFileNameW(module, fullPath, MAX_PATH);
-
-		// send back command with module info
-		{
-			commands::GetModuleInfo cmd = { module, process::GetId(), command->loadImports, command->taskContext };
-			wcscpy_s(cmd.path, fullPath);
-
-			pipe->SendAck();
-
-			pipe->SendCommandAndWaitForAck(cmd);
-		}
-	}
-	else
-	{
-		pipe->SendAck();
-
-		pipe->SendCommandAndWaitForAck(commands::GetModuleInfo { nullptr, process::GetId(), command->loadImports, command->taskContext });
-	}
-
-	return true;
-}
-
-
 // BEGIN EPIC MOD - Support for UE4 debug visualizers
 struct FNameEntry;
 extern FNameEntry*** GFNameTableForDebuggerVisualizers_MT;
@@ -136,7 +87,7 @@ extern FChunkedFixedUObjectArray*& GObjectArrayForDebugVisualizers;
 // END EPIC MOD
 
 
-bool actions::LoadPatch::Execute(const CommandType* command, const DuplexPipe* pipe, void*)
+bool actions::LoadPatch::Execute(const CommandType* command, const DuplexPipe* pipe, void*, const void*, size_t)
 {
 	// load library into this process
 	HMODULE module = ::LoadLibraryW(command->path);
@@ -157,13 +108,13 @@ bool actions::LoadPatch::Execute(const CommandType* command, const DuplexPipe* p
 	pipe->SendAck();
 
 	// send back command with module info
-	pipe->SendCommandAndWaitForAck(commands::LoadPatchInfo { module });
+	pipe->SendCommandAndWaitForAck(commands::LoadPatchInfo { module }, nullptr, 0u);
 
 	return true;
 }
 
 
-bool actions::UnloadPatch::Execute(const CommandType* command, const DuplexPipe* pipe, void*)
+bool actions::UnloadPatch::Execute(const CommandType* command, const DuplexPipe* pipe, void*, const void*, size_t)
 {
 	// unload library from this process
 	::FreeLibrary(command->module);
@@ -173,7 +124,7 @@ bool actions::UnloadPatch::Execute(const CommandType* command, const DuplexPipe*
 }
 
 
-bool actions::CallEntryPoint::Execute(const CommandType* command, const DuplexPipe* pipe, void*)
+bool actions::CallEntryPoint::Execute(const CommandType* command, const DuplexPipe* pipe, void*, const void*, size_t)
 {
 	executable::CallDllEntryPoint(command->moduleBase, command->entryPointRva);
 	pipe->SendAck();
@@ -182,9 +133,9 @@ bool actions::CallEntryPoint::Execute(const CommandType* command, const DuplexPi
 }
 
 
-bool actions::LogOutput::Execute(const CommandType* command, const DuplexPipe* pipe, void*)
+bool actions::LogOutput::Execute(const CommandType*, const DuplexPipe* pipe, void*, const void* payload, size_t)
 {
-	logging::LogNoFormat<logging::Channel::USER>(command->buffer);
+	logging::LogNoFormat<logging::Channel::USER>(static_cast<const wchar_t*>(payload));
 	pipe->SendAck();
 
 	return true;
@@ -194,7 +145,7 @@ bool actions::LogOutput::Execute(const CommandType* command, const DuplexPipe* p
 extern bool GIsCompileActive;
 // END EPIC MOD
 
-bool actions::CompilationFinished::Execute(const CommandType*, const DuplexPipe* pipe, void*)
+bool actions::CompilationFinished::Execute(const CommandType*, const DuplexPipe* pipe, void*, const void*, size_t)
 {
 	pipe->SendAck();
 
@@ -207,7 +158,7 @@ bool actions::CompilationFinished::Execute(const CommandType*, const DuplexPipe*
 }
 
 
-bool actions::HandleExceptionFinished::Execute(const CommandType* command, const DuplexPipe* pipe, void* context)
+bool actions::HandleExceptionFinished::Execute(const CommandType* command, const DuplexPipe* pipe, void* context, const void*, size_t)
 {
 	ClientUserCommandThread::ExceptionResult* resultContext = static_cast<ClientUserCommandThread::ExceptionResult*>(context);
 	resultContext->returnAddress = command->returnAddress;
diff --git a/Engine/Source/Developer/Windows/LiveCoding/Private/External/LC_ClientCommandActions.h b/Engine/Source/Developer/Windows/LiveCoding/Private/External/LC_ClientCommandActions.h
--- a/Engine/Source/Developer/Windows/LiveCoding/Private/External/LC_ClientCommandActions.h
+++ b/Engine/Source/Developer/Windows/LiveCoding/Private/External/LC_ClientCommandActions.h
@@ -10,22 +10,19 @@ class DuplexPipe;
 
 namespace actions
 {
-	#define DECLARE_ACTION(_name)																		\
-		struct _name																					\
-		{																								\
-			typedef ::commands::_name CommandType;														\
-			static bool Execute(const CommandType* command, const DuplexPipe* pipe, void* context);		\
+	#define DECLARE_ACTION(_name)																													\
+		struct _name																																\
+		{																																			\
+			typedef ::commands::_name CommandType;																									\
+			static bool Execute(const CommandType* command, const DuplexPipe* pipe, void* context, const void* payload, size_t payloadSize);		\
 		}
 
 	DECLARE_ACTION(RegisterProcessFinished);
-	DECLARE_ACTION(EnableModuleFinished);
-	DECLARE_ACTION(DisableModuleFinished);
-	DECLARE_ACTION(EnableAllModulesFinished);
-	DECLARE_ACTION(DisableAllModulesFinished);
+	DECLARE_ACTION(EnableModulesFinished);
+	DECLARE_ACTION(DisableModulesFinished);
 	DECLARE_ACTION(EnterSyncPoint);
 	DECLARE_ACTION(LeaveSyncPoint);
 	DECLARE_ACTION(CallHooks);
-	DECLARE_ACTION(GetModule);
 	DECLARE_ACTION(LoadPatch);
 	DECLARE_ACTION(UnloadPatch);
 	DECLARE_ACTION(CallEntryPoint);
diff --git a/Engine/Source/Developer/Windows/LiveCoding/Private/External/LC_ClientCommandThread.cpp b/Engine/Source/Developer/Windows/LiveCoding/Private/External/LC_ClientCommandThread.cpp
--- a/Engine/Source/Developer/Windows/LiveCoding/Private/External/LC_ClientCommandThread.cpp
+++ b/Engine/Source/Developer/Windows/LiveCoding/Private/External/LC_ClientCommandThread.cpp
@@ -25,14 +25,7 @@ ClientCommandThread::~ClientCommandThread(void)
 unsigned int ClientCommandThread::Start(const std::wstring& processGroupName, Event* compilationEvent, Event* waitForStartEvent, CriticalSection* pipeAccessCS)
 {
 	// spawn a thread that communicates with the server
-	ThreadContext* context = new ThreadContext;
-	context->thisInstance = this;
-	context->processGroupName = processGroupName;
-	context->compilationEvent = compilationEvent;
-	context->waitForStartEvent = waitForStartEvent;
-	context->pipeAccessCS = pipeAccessCS;
-
-	m_thread = thread::Create(128u * 1024u, &ThreadProxy, context);
+	m_thread = thread::Create("Live coding commands", 128u * 1024u, &ClientCommandThread::ThreadFunction, this, processGroupName, compilationEvent, waitForStartEvent, pipeAccessCS);
 
 	return thread::GetId(m_thread);
 }
@@ -48,19 +41,6 @@ void ClientCommandThread::Join(void)
 }
 
 
-unsigned int __stdcall ClientCommandThread::ThreadProxy(void* context)
-{
-	thread::SetName("Live coding commands");
-
-	ThreadContext* realContext = static_cast<ThreadContext*>(context);
-	const unsigned int exitCode = realContext->thisInstance->ThreadFunction(realContext->processGroupName, realContext->compilationEvent, realContext->waitForStartEvent, realContext->pipeAccessCS);
-
-	delete realContext;
-
-	return exitCode;
-}
-
-
 unsigned int ClientCommandThread::ThreadFunction(const std::wstring& processGroupName, Event* compilationEvent, Event* waitForStartEvent, CriticalSection* pipeAccessCS)
 {
 	waitForStartEvent->Wait();
@@ -101,7 +81,7 @@ unsigned int ClientCommandThread::ThreadFunction(const std::wstring& processGrou
 		// we need to make sure that other threads talking through the pipe don't use it at the same time.
 		CriticalSection::ScopedLock lock(pipeAccessCS);
 
-		m_pipe->SendCommandAndWaitForAck(commands::ReadyForCompilation {});
+		m_pipe->SendCommandAndWaitForAck(commands::ReadyForCompilation {}, nullptr, 0u);
 
 		commandMap.HandleCommands(m_pipe, nullptr);
 	}
diff --git a/Engine/Source/Developer/Windows/LiveCoding/Private/External/LC_ClientCommandThread.h b/Engine/Source/Developer/Windows/LiveCoding/Private/External/LC_ClientCommandThread.h
--- a/Engine/Source/Developer/Windows/LiveCoding/Private/External/LC_ClientCommandThread.h
+++ b/Engine/Source/Developer/Windows/LiveCoding/Private/External/LC_ClientCommandThread.h
@@ -26,16 +26,6 @@ public:
 	void Join(void);
 
 private:
-	struct ThreadContext
-	{
-		ClientCommandThread* thisInstance;
-		std::wstring processGroupName;
-		Event* compilationEvent;
-		Event* waitForStartEvent;
-		CriticalSection* pipeAccessCS;
-	};
-
-	static unsigned int __stdcall ThreadProxy(void* context);
 	unsigned int ThreadFunction(const std::wstring& processGroupName, Event* compilationEvent, Event* waitForStartEvent, CriticalSection* pipeAccessCS);
 
 	thread::Handle m_thread;
diff --git a/Engine/Source/Developer/Windows/LiveCoding/Private/External/LC_ClientStartupThread.cpp b/Engine/Source/Developer/Windows/LiveCoding/Private/External/LC_ClientStartupThread.cpp
--- a/Engine/Source/Developer/Windows/LiveCoding/Private/External/LC_ClientStartupThread.cpp
+++ b/Engine/Source/Developer/Windows/LiveCoding/Private/External/LC_ClientStartupThread.cpp
@@ -17,6 +17,9 @@
 #include "Misc/ConfigCacheIni.h"
 #include "Misc/App.h"
 
+// JumpToSelf is an extern function coming from assembler source
+extern void JumpToSelf(void);
+
 namespace
 {
 	template <typename T>
@@ -60,7 +63,7 @@ ClientStartupThread::~ClientStartupThread(void)
 		// give the server a chance to deal with disconnected clients
 		if (m_pipeClient->IsValid())
 		{
-			m_pipeClient->SendCommandAndWaitForAck(commands::DisconnectClient{});
+			m_pipeClient->SendCommandAndWaitForAck(commands::DisconnectClient {}, nullptr, 0u);
 		}
 		m_pipeClient->Close();
 	}
@@ -108,13 +111,11 @@ ClientStartupThread::~ClientStartupThread(void)
 
 void ClientStartupThread::Start(const char* const groupName, RunMode::Enum runMode)
 {
-	// spawn a thread that does all the initialization work
-	ThreadContext* context = new ThreadContext;
-	context->thisInstance = this;
-	context->processGroupName = string::ToWideString(groupName);
-	context->runMode = runMode;
-
-	m_thread = thread::Create(128u * 1024u, &ThreadProxy, context);
+	// spawn a thread that does all the initialization work.
+	// in the context of mutexes, jobs, named shared memory, etc. object names behave similar to
+	// file names and are not allowed to contain certain characters.
+	std::wstring safeProcessGroupName = string::MakeSafeName(string::ToWideString(groupName));
+	m_thread = thread::Create("Live coding startup", 128u * 1024u, &ClientStartupThread::ThreadFunction, this, safeProcessGroupName, runMode);
 }
 
 
@@ -130,31 +131,63 @@ void ClientStartupThread::Join(void)
 
 void* ClientStartupThread::EnableModule(const wchar_t* const nameOfExeOrDll)
 {
-	return m_userCommandThread->EnableModule(nameOfExeOrDll);
+	// wait for the startup thread to finish initialization
+	Join();
+
+	if (m_userCommandThread)
+	{
+		return m_userCommandThread->EnableModule(nameOfExeOrDll);
+	}
+
+	return nullptr;
 }
 
 
 void* ClientStartupThread::EnableAllModules(const wchar_t* const nameOfExeOrDll)
 {
-	return m_userCommandThread->EnableAllModules(nameOfExeOrDll);
+	// wait for the startup thread to finish initialization
+	Join();
+
+	if (m_userCommandThread)
+	{
+		return m_userCommandThread->EnableAllModules(nameOfExeOrDll);
+	}
+
+	return nullptr;
 }
 
 
 void* ClientStartupThread::DisableModule(const wchar_t* const nameOfExeOrDll)
 {
-	return m_userCommandThread->DisableModule(nameOfExeOrDll);
+	// wait for the startup thread to finish initialization
+	Join();
+
+	if (m_userCommandThread)
+	{
+		return m_userCommandThread->DisableModule(nameOfExeOrDll);
+	}
+
+	return nullptr;
 }
 
 
 void* ClientStartupThread::DisableAllModules(const wchar_t* const nameOfExeOrDll)
 {
-	return m_userCommandThread->DisableAllModules(nameOfExeOrDll);
+	// wait for the startup thread to finish initialization
+	Join();
+
+	if (m_userCommandThread)
+	{
+		return m_userCommandThread->DisableAllModules(nameOfExeOrDll);
+	}
+
+	return nullptr;
 }
 
 
 void ClientStartupThread::WaitForToken(void* token)
 {
-	// we cannot wait for commands in the user command thread as long as startup hasn't finished
+	// wait for the startup thread to finish initialization
 	Join();
 
 	if (m_userCommandThread)
@@ -166,7 +199,7 @@ void ClientStartupThread::WaitForToken(void* token)
 
 void ClientStartupThread::TriggerRecompile(void)
 {
-	// we cannot wait for commands in the user command thread as long as startup hasn't finished
+	// wait for the startup thread to finish initialization
 	Join();
 
 	if (m_userCommandThread)
@@ -178,7 +211,7 @@ void ClientStartupThread::TriggerRecompile(void)
 
 void ClientStartupThread::BuildPatch(const wchar_t* moduleNames[], const wchar_t* objPaths[], const wchar_t* amalgamatedObjPaths[], unsigned int count)
 {
-	// we cannot wait for commands in the user command thread as long as startup hasn't finished
+	// wait for the startup thread to finish initialization
 	Join();
 
 	if (m_userCommandThread)
@@ -190,7 +223,7 @@ void ClientStartupThread::BuildPatch(const wchar_t* moduleNames[], const wchar_t
 
 void ClientStartupThread::InstallExceptionHandler(void)
 {
-	// we cannot wait for commands in the user command thread as long as startup hasn't finished
+	// wait for the startup thread to finish initialization
 	Join();
 
 	if (m_userCommandThread)
@@ -270,7 +303,7 @@ void ClientStartupThread::EnableLazyLoadedModule(const wchar_t* fileName, Window
 
 void ClientStartupThread::ApplySettingBool(const char* const settingName, int value)
 {
-	// we cannot wait for commands in the user command thread as long as startup hasn't finished
+	// wait for the startup thread to finish initialization
 	Join();
 
 	if (m_userCommandThread)
@@ -282,7 +315,7 @@ void ClientStartupThread::ApplySettingBool(const char* const settingName, int va
 
 void ClientStartupThread::ApplySettingInt(const char* const settingName, int value)
 {
-	// we cannot wait for commands in the user command thread as long as startup hasn't finished
+	// wait for the startup thread to finish initialization
 	Join();
 
 	if (m_userCommandThread)
@@ -294,7 +327,7 @@ void ClientStartupThread::ApplySettingInt(const char* const settingName, int val
 
 void ClientStartupThread::ApplySettingString(const char* const settingName, const wchar_t* const value)
 {
-	// we cannot wait for commands in the user command thread as long as startup hasn't finished
+	// wait for the startup thread to finish initialization
 	Join();
 
 	if (m_userCommandThread)
@@ -304,23 +337,6 @@ void ClientStartupThread::ApplySettingString(const char* const settingName, cons
 }
 
 
-unsigned int __stdcall ClientStartupThread::ThreadProxy(void* context)
-{
-	thread::SetName("Live coding startup");
-
-	ThreadContext* realContext = static_cast<ThreadContext*>(context);
-
-	// in the context of mutexes, jobs, named shared memory, etc. object names behave similar to
-	// file names and are not allowed to contain certain characters.
-	const std::wstring& safeProcessGroupName = string::MakeSafeName(realContext->processGroupName);
-	const unsigned int exitCode = realContext->thisInstance->ThreadFunction(safeProcessGroupName, realContext->runMode);
-
-	delete realContext;
-
-	return exitCode;
-}
-
-
 unsigned int ClientStartupThread::ThreadFunction(const std::wstring& processGroupName, RunMode::Enum runMode)
 {
 	// configure all child processes associated with the job to terminate when the parent terminates.
@@ -337,7 +353,7 @@ unsigned int ClientStartupThread::ThreadFunction(const std::wstring& processGrou
 	// the first one will spawn the Live++ process, all others will connect to the same process.
 	{
 		InterprocessMutex initProcessMutex(primitiveNames::StartupMutex(processGroupName).c_str());
-		initProcessMutex.Lock();
+		InterprocessMutex::ScopedLock mutexLock(&initProcessMutex);
 
 		m_sharedMemory = new NamedSharedMemory(primitiveNames::StartupNamedSharedMemory(processGroupName).c_str());
 		if (m_sharedMemory->IsOwnedByCallingProcess())
@@ -392,8 +408,6 @@ unsigned int ClientStartupThread::ThreadFunction(const std::wstring& processGrou
 				::AssignProcessToJobObject(m_job, m_processHandle);
 			}
 		}
-
-		initProcessMutex.Unlock();
 	}
 
 	if (!m_processHandle)
@@ -445,7 +459,7 @@ unsigned int ClientStartupThread::ThreadFunction(const std::wstring& processGrou
 	m_userCommandThread->Start(processGroupName, m_startEvent, m_pipeClientCS);
 
 	// register this process with Live++
-	m_pipeClient->SendCommandAndWaitForAck(commands::RegisterProcess { process::GetBase(), process::GetId(), commandThreadId });
+	m_pipeClient->SendCommandAndWaitForAck(commands::RegisterProcess { process::GetBase(), process::GetId(), commandThreadId, &JumpToSelf }, nullptr, 0u);
 
 	// handle commands until registration is finished
 	{
diff --git a/Engine/Source/Developer/Windows/LiveCoding/Private/External/LC_ClientStartupThread.h b/Engine/Source/Developer/Windows/LiveCoding/Private/External/LC_ClientStartupThread.h
--- a/Engine/Source/Developer/Windows/LiveCoding/Private/External/LC_ClientStartupThread.h
+++ b/Engine/Source/Developer/Windows/LiveCoding/Private/External/LC_ClientStartupThread.h
@@ -64,14 +64,6 @@ public:
 	// END EPIC MOD
 
 private:
-	struct ThreadContext
-	{
-		ClientStartupThread* thisInstance;
-		std::wstring processGroupName;
-		RunMode::Enum runMode;
-	};
-
-	static unsigned int __stdcall ThreadProxy(void* context);
 	unsigned int ThreadFunction(const std::wstring& groupName, RunMode::Enum runMode);
 
 	HINSTANCE m_instance;
diff --git a/Engine/Source/Developer/Windows/LiveCoding/Private/External/LC_ClientUserCommandThread.cpp b/Engine/Source/Developer/Windows/LiveCoding/Private/External/LC_ClientUserCommandThread.cpp
--- a/Engine/Source/Developer/Windows/LiveCoding/Private/External/LC_ClientUserCommandThread.cpp
+++ b/Engine/Source/Developer/Windows/LiveCoding/Private/External/LC_ClientUserCommandThread.cpp
@@ -8,404 +8,150 @@
 #include "LC_Process.h"
 #include "LC_CriticalSection.h"
 #include "LC_StringUtil.h"
+#include "LC_Executable.h"
+#include "LC_MemoryStream.h"
+#include "LC_Logging.h"
 #include <deque>
+#include <unordered_set>
 
 
-namespace userCommands
+namespace
 {
-	struct Scope
-	{
-		enum Enum
-		{
-			NONE,
-			ENABLE_MODULES,
-			DISABLE_MODULES,
-		};
-	};
-
-	struct BaseCommand
-	{
-		explicit BaseCommand(Scope::Enum scope)
-			: m_scope(scope)
-		{
-		}
-
-		virtual ~BaseCommand(void) {}
-
-		virtual void Execute(DuplexPipe* pipe) = 0;
-		
-		inline Scope::Enum GetScope(void) const
-		{
-			return m_scope;
-		}
-
-	private:
-		Scope::Enum m_scope;
-	};
-
-	struct EnableModuleCommand : public BaseCommand
-	{
-		EnableModuleCommand(void)
-			: BaseCommand(Scope::ENABLE_MODULES)
-			, token(nullptr)
-			, moduleName()
-		{
-		}
-
-		virtual ~EnableModuleCommand(void) {}
-
-		virtual void Execute(DuplexPipe* pipe) override
-		{
-			commands::EnableModule serverCommand;
-			serverCommand.processId = process::GetId();
-			wcscpy_s(serverCommand.path, moduleName.c_str());
-			serverCommand.token = token;
-
-			pipe->SendCommandAndWaitForAck(serverCommand);
-		}
-
-		Event* token;
-		std::wstring moduleName;
-	};
-
-	struct EnableAllModulesCommand : public BaseCommand
-	{
-		EnableAllModulesCommand(void)
-			: BaseCommand(Scope::ENABLE_MODULES)
-			, token(nullptr)
-			, moduleName()
-		{
-		}
-
-		virtual ~EnableAllModulesCommand(void) {}
-
-		virtual void Execute(DuplexPipe* pipe) override
-		{
-			commands::EnableAllModules serverCommand;
-			serverCommand.processId = process::GetId();
-			wcscpy_s(serverCommand.path, moduleName.c_str());
-			serverCommand.token = token;
-
-			pipe->SendCommandAndWaitForAck(serverCommand);
-		}
-
-		Event* token;
-		std::wstring moduleName;
-	};
-
-	struct DisableModuleCommand : public BaseCommand
-	{
-		DisableModuleCommand(void)
-			: BaseCommand(Scope::DISABLE_MODULES)
-			, token(nullptr)
-			, moduleName()
-		{
-		}
-
-		virtual ~DisableModuleCommand(void) {}
-
-		virtual void Execute(DuplexPipe* pipe) override
-		{
-			commands::DisableModule serverCommand;
-			serverCommand.processId = process::GetId();
-			wcscpy_s(serverCommand.path, moduleName.c_str());
-			serverCommand.token = token;
-
-			pipe->SendCommandAndWaitForAck(serverCommand);
-		}
-
-		Event* token;
-		std::wstring moduleName;
-	};
-
-	struct DisableAllModulesCommand : public BaseCommand
-	{
-		DisableAllModulesCommand(void)
-			: BaseCommand(Scope::DISABLE_MODULES)
-			, token(nullptr)
-			, moduleName()
-		{
-		}
-
-		virtual ~DisableAllModulesCommand(void) {}
-
-		virtual void Execute(DuplexPipe* pipe) override
-		{
-			commands::DisableAllModules serverCommand;
-			serverCommand.processId = process::GetId();
-			wcscpy_s(serverCommand.path, moduleName.c_str());
-			serverCommand.token = token;
-
-			pipe->SendCommandAndWaitForAck(serverCommand);
-		}
-
-		Event* token;
-		std::wstring moduleName;
-	};
-
-	struct TriggerRecompileCommand : public BaseCommand
+	template <typename T>
+	class ProxyCommand : public ClientUserCommandThread::BaseCommand
 	{
-		TriggerRecompileCommand(void)
-			: BaseCommand(Scope::NONE)
+	public:
+		ProxyCommand(bool expectResponse, size_t payloadSize)
+			: BaseCommand(expectResponse)
+			, m_command()
+			, m_payload(payloadSize)
 		{
 		}
 
-		virtual ~TriggerRecompileCommand(void) {}
-
 		virtual void Execute(DuplexPipe* pipe) override
 		{
-			commands::TriggerRecompile serverCommand;
-			pipe->SendCommandAndWaitForAck(serverCommand);
-		}
-	};
-
-	// BEGIN EPIC MOD - Adding ShowConsole command
-	struct ShowConsoleCommand : public BaseCommand
-	{
-		ShowConsoleCommand(void)
-			: BaseCommand(Scope::NONE)
-		{
+			pipe->SendCommandAndWaitForAck(m_command, m_payload.GetData(), m_payload.GetSize());
 		}
 
-		virtual ~ShowConsoleCommand(void) {}
+		T m_command;
+		memoryStream::Writer m_payload;
 
-		virtual void Execute(DuplexPipe* pipe) override
-		{
-			commands::ShowConsole serverCommand;
-			pipe->SendCommandAndWaitForAck(serverCommand);
-		}
+		LC_DISABLE_COPY(ProxyCommand);
+		LC_DISABLE_MOVE(ProxyCommand);
+		LC_DISABLE_ASSIGNMENT(ProxyCommand);
+		LC_DISABLE_MOVE_ASSIGNMENT(ProxyCommand);
 	};
-	// END EPIC MOD
 
-	// BEGIN EPIC MOD - Adding SetVisible command
-	struct SetVisibleCommand : public BaseCommand
+	// gathers module data for the given module, its import modules, the import's import modules, and so forth
+	static std::vector<commands::ModuleData> GatherImportModuleData(HMODULE mainModule)
 	{
-		SetVisibleCommand(void)
-			: BaseCommand(Scope::NONE)
-		{
-		}
-
-		virtual ~SetVisibleCommand(void) {}
+		std::vector<commands::ModuleData> moduleDatas;
+		moduleDatas.reserve(1024u);
 
-		virtual void Execute(DuplexPipe* pipe) override
-		{
-			commands::SetVisible serverCommand;
-			serverCommand.visible = visible;
-			pipe->SendCommandAndWaitForAck(serverCommand);
-		}
+		std::unordered_set<std::wstring> loadedModules;
+		loadedModules.reserve(1024u);
 
-		bool visible;
-	};
-	// END EPIC MOD
+		std::vector<HMODULE> modules;
+		modules.reserve(1024u);
 
-	// BEGIN EPIC MOD - Adding SetActive command
-	struct SetActiveCommand : public BaseCommand
-	{
-		SetActiveCommand(void)
-			: BaseCommand(Scope::NONE)
+		modules.push_back(mainModule);
+		while (!modules.empty())
 		{
-		}
+			const HMODULE module = modules.back();
+			modules.pop_back();
 
-		virtual ~SetActiveCommand(void) {}
+			// get the absolute path of the module.
+			// this automatically takes care of API sets used by Windows 7 and later. in a nutshell, these API sets
+			// allow redirection of an API DLL to an underlying OS DLL, e.g. api-ms-win-core-apiquery-l1-1-0.dll redirects
+			// to ntdll.dll.
+			wchar_t fullPath[MAX_PATH];
+			::GetModuleFileNameW(module, fullPath, MAX_PATH);
 
-		virtual void Execute(DuplexPipe* pipe) override
-		{
-			commands::SetActive serverCommand;
-			serverCommand.active = active;
-			pipe->SendCommandAndWaitForAck(serverCommand);
-		}
-
-		bool active;
-	};
-	// END EPIC MOD
-
-	// BEGIN EPIC MOD - Adding SetActive command
-	struct SetBuildArgumentsCommand : public BaseCommand
-	{
-		SetBuildArgumentsCommand(void)
-			: BaseCommand(Scope::NONE)
-		{
-		}
-
-		virtual ~SetBuildArgumentsCommand(void) {}
-
-		virtual void Execute(DuplexPipe* pipe) override
-		{
-			commands::SetBuildArguments serverCommand;
-			serverCommand.processId = process::GetId();
-			wcscpy_s(serverCommand.arguments, arguments.c_str());
-			pipe->SendCommandAndWaitForAck(serverCommand);
-		}
-
-		std::wstring arguments;
-	};
-	// END EPIC MOD
-
-	// BEGIN EPIC MOD - Support for lazy-loading modules
-	struct EnableLazyLoadedModuleCommand : public BaseCommand
-	{
-		EnableLazyLoadedModuleCommand(void)
-			: BaseCommand(Scope::NONE)
-		{
-		}
-
-		virtual ~EnableLazyLoadedModuleCommand(void) {}
-
-		virtual void Execute(DuplexPipe* pipe) override
-		{
-			commands::EnableLazyLoadedModule serverCommand;
-			serverCommand.processId = process::GetId();
-			wcscpy_s(serverCommand.fileName, fileName.c_str());
-			serverCommand.moduleBase = moduleBase;
-			pipe->SendCommandAndWaitForAck(serverCommand);
-		}
-
-		std::wstring fileName;
-		Windows::HMODULE moduleBase;
-	};
-	// END EPIC MOD
-
-
-	struct BuildPatchCommand : public BaseCommand
-	{
-		BuildPatchCommand(void)
-			: BaseCommand(Scope::NONE)
-			, count(0u)
-			, moduleNames()
-			, objPaths()
-			, amalgamatedObjPaths()
-		{
-		}
-
-		virtual ~BuildPatchCommand(void) {}
-
-		virtual void Execute(DuplexPipe* pipe) override
-		{
-			commands::BuildPatch serverCommand;
-			serverCommand.count = count;
-			pipe->SendCommandAndWaitForAck(serverCommand);
-
-			for (unsigned int i = 0u; i < count; ++i)
+			// did we load the imports of this module already?
+			auto findIt = loadedModules.find(fullPath);
+			if (findIt == loadedModules.end())
 			{
-				commands::BuildPatchPacket packet;
-				wcscpy_s(packet.moduleName, moduleNames[i].c_str());
-				wcscpy_s(packet.objPath, objPaths[i].c_str());
-				wcscpy_s(packet.amalgamatedObjPath, amalgamatedObjPaths[i].c_str());
-
-				pipe->SendCommandAndWaitForAck(packet);
-			}
-		}
+				loadedModules.insert(fullPath);
 
-		unsigned int count;
-		std::vector<std::wstring> moduleNames;
-		std::vector<std::wstring> objPaths;
-		std::vector<std::wstring> amalgamatedObjPaths;
-	};
-
-	struct ApplySettingBoolCommand : public BaseCommand
-	{
-		ApplySettingBoolCommand(void)
-			: BaseCommand(Scope::NONE)
-			, settingName()
-			, value(0)
-		{
-		}
-
-		virtual ~ApplySettingBoolCommand(void) {}
-
-		virtual void Execute(DuplexPipe* pipe) override
-		{
-			commands::ApplySettingBool serverCommand;
-			strcpy_s(serverCommand.settingName, settingName.c_str());
-			serverCommand.settingValue = value;
-			pipe->SendCommandAndWaitForAck(serverCommand);
-		}
-
-		std::string settingName;
-		int value;
-	};
-
-	struct ApplySettingIntCommand : public BaseCommand
-	{
-		ApplySettingIntCommand(void)
-			: BaseCommand(Scope::NONE)
-			, settingName()
-			, value(0)
-		{
-		}
-
-		virtual ~ApplySettingIntCommand(void) {}
-
-		virtual void Execute(DuplexPipe* pipe) override
-		{
-			commands::ApplySettingInt serverCommand;
-			strcpy_s(serverCommand.settingName, settingName.c_str());
-			serverCommand.settingValue = value;
-			pipe->SendCommandAndWaitForAck(serverCommand);
-		}
-
-		std::string settingName;
-		int value;
-	};
-
-	struct ApplySettingStringCommand : public BaseCommand
-	{
-		ApplySettingStringCommand(void)
-			: BaseCommand(Scope::NONE)
-			, settingName()
-			, value()
-		{
-		}
-
-		virtual ~ApplySettingStringCommand(void) {}
+				// add data for this module
+				{
+					commands::ModuleData moduleData = {};
+					moduleData.base = module;
+					wcscpy_s(moduleData.path, fullPath);
+					moduleDatas.emplace_back(moduleData);
+				}
 
-		virtual void Execute(DuplexPipe* pipe) override
-		{
-			commands::ApplySettingString serverCommand;
-			strcpy_s(serverCommand.settingName, settingName.c_str());
-			wcscpy_s(serverCommand.settingValue, value.c_str());
-			pipe->SendCommandAndWaitForAck(serverCommand);
+				executable::Image* image = executable::OpenImage(fullPath, file::OpenMode::READ_ONLY);
+				if (image)
+				{
+					executable::ImageSectionDB* imageSections = executable::GatherImageSectionDB(image);
+					if (imageSections)
+					{
+						executable::ImportModuleDB* importModules = executable::GatherImportModuleDB(image, imageSections);
+						if (importModules)
+						{
+							for (size_t i = 0; i < importModules->modules.size(); ++i)
+							{
+								const char* importModulePath = importModules->modules[i].path;
+
+								// only add the import module if it is loaded into the process
+								HMODULE importModule = ::GetModuleHandleA(importModulePath);
+								if (importModule)
+								{
+									modules.push_back(importModule);
+								}
+								else
+								{
+									LC_ERROR_USER("Cannot enable module %s because it is not loaded by this process.", importModulePath);
+								}
+							}
+
+							executable::DestroyImportModuleDB(importModules);
+						}
+
+						executable::DestroyImageSectionDB(imageSections);
+					}
+
+					executable::CloseImage(image);
+				}
+			}
 		}
 
-		std::string settingName;
-		std::wstring value;
-	};
+		return moduleDatas;
+	}
 }
 
 
-namespace
+ClientUserCommandThread::BaseCommand::BaseCommand(bool expectResponse)
+	: m_expectResponse(expectResponse)
 {
-	// queue for working on commands received by user code
-	static std::deque<userCommands::BaseCommand*> g_userCommandQueue;
-	static CriticalSection g_userCommandQueueCS;
 }
 
-// BEGIN EPIC MOD - Allow manually batching commands
-void BeginCommandBatch()
+
+ClientUserCommandThread::BaseCommand::~BaseCommand(void)
 {
-	g_userCommandQueueCS.Enter();
 }
 
-void EndCommandBatch()
+
+bool ClientUserCommandThread::BaseCommand::ExpectsResponse(void) const
 {
-	g_userCommandQueueCS.Leave();
+	return m_expectResponse;
 }
-// END EPIC MOD
-
 
 ClientUserCommandThread::ClientUserCommandThread(DuplexPipeClient* pipeClient, DuplexPipeClient* exceptionPipeClient)
 	: m_thread(INVALID_HANDLE_VALUE)
 	, m_processGroupName()
 	, m_pipe(pipeClient)
 	, m_exceptionPipe(exceptionPipeClient)
-	, m_itemInQueueEvent(new Event(nullptr, Event::Type::MANUAL_RESET))
+	, m_userCommandQueue()
+	, m_userCommandQueueCS()
+	, m_userCommandQueueSema(0, 65535u)
 {
 }
 
 
 ClientUserCommandThread::~ClientUserCommandThread(void)
 {
-	delete m_itemInQueueEvent;
 }
 
 
@@ -414,12 +160,7 @@ unsigned int ClientUserCommandThread::Start(const std::wstring& processGroupName
 	m_processGroupName = processGroupName;
 
 	// spawn a thread that does the work
-	ThreadContext* context = new ThreadContext;
-	context->thisInstance = this;
-	context->waitForStartEvent = waitForStartEvent;
-	context->pipeAccessCS = pipeAccessCS;
-
-	m_thread = thread::Create(128u * 1024u, &ThreadProxy, context);
+	m_thread = thread::Create("Live coding user commands", 128u * 1024u, &ClientUserCommandThread::ThreadFunction, this, waitForStartEvent, pipeAccessCS);
 
 	return thread::GetId(m_thread);
 }
@@ -437,69 +178,110 @@ void ClientUserCommandThread::Join(void)
 
 void* ClientUserCommandThread::EnableModule(const wchar_t* const nameOfExeOrDll)
 {
-	userCommands::EnableModuleCommand* command = new userCommands::EnableModuleCommand;
-	command->token = new Event(nullptr, Event::Type::AUTO_RESET);
-	command->moduleName = nameOfExeOrDll;
+	HMODULE module = ::GetModuleHandleW(nameOfExeOrDll);
+	if (!module)
 	{
-		CriticalSection::ScopedLock lock(&g_userCommandQueueCS);
-		g_userCommandQueue.push_front(command);
+		LC_ERROR_USER("Cannot enable module %S because it is not loaded by this process.", nameOfExeOrDll);
+		return nullptr;
 	}
 
-	// signal to the thread that a new item is in the queue
-	m_itemInQueueEvent->Signal();
+	ProxyCommand<commands::EnableModules>* proxy = new ProxyCommand<commands::EnableModules>(true, sizeof(commands::ModuleData)*1u);
+	proxy->m_command.processId = process::GetId();
+	proxy->m_command.moduleCount = 1u;
+	proxy->m_command.token = new Event(nullptr, Event::Type::AUTO_RESET);
 
-	return command->token;
-}
+	commands::ModuleData moduleData = {};
+	moduleData.base = module;
+	::GetModuleFileNameW(module, moduleData.path, MAX_PATH);
+
+	proxy->m_payload.Write(moduleData);
 
+	PushUserCommand(proxy);
+
+	return proxy->m_command.token;
+}
 
 void* ClientUserCommandThread::EnableAllModules(const wchar_t* const nameOfExeOrDll)
 {
-	userCommands::EnableAllModulesCommand* command = new userCommands::EnableAllModulesCommand;
-	command->token = new Event(nullptr, Event::Type::AUTO_RESET);
-	command->moduleName = nameOfExeOrDll;
+	HMODULE module = ::GetModuleHandleW(nameOfExeOrDll);
+	if (!module)
 	{
-		CriticalSection::ScopedLock lock(&g_userCommandQueueCS);
-		g_userCommandQueue.push_front(command);
+		LC_ERROR_USER("Cannot enable module %S because it is not loaded by this process.", nameOfExeOrDll);
+		return nullptr;
 	}
 
-	// signal to the thread that a new item is in the queue
-	m_itemInQueueEvent->Signal();
+	const std::vector<commands::ModuleData>& allModuleData = GatherImportModuleData(module);
+	const size_t moduleCount = allModuleData.size();
+
+	ProxyCommand<commands::EnableModules>* proxy = new ProxyCommand<commands::EnableModules>(true, sizeof(commands::ModuleData) * moduleCount);
+	proxy->m_command.processId = process::GetId();
+	proxy->m_command.moduleCount = static_cast<unsigned int>(moduleCount);
+	proxy->m_command.token = new Event(nullptr, Event::Type::AUTO_RESET);
 
-	return command->token;
+	for (size_t i = 0u; i < moduleCount; ++i)
+	{
+		const commands::ModuleData& moduleData = allModuleData[i];
+		proxy->m_payload.Write(moduleData);
+	}
+
+	PushUserCommand(proxy);
+
+	return proxy->m_command.token;
 }
 
 
 void* ClientUserCommandThread::DisableModule(const wchar_t* const nameOfExeOrDll)
 {
-	userCommands::DisableModuleCommand* command = new userCommands::DisableModuleCommand;
-	command->token = new Event(nullptr, Event::Type::AUTO_RESET);
-	command->moduleName = nameOfExeOrDll;
+	HMODULE module = ::GetModuleHandleW(nameOfExeOrDll);
+	if (!module)
 	{
-		CriticalSection::ScopedLock lock(&g_userCommandQueueCS);
-		g_userCommandQueue.push_front(command);
+		LC_ERROR_USER("Cannot disable module %S because it is not loaded by this process.", nameOfExeOrDll);
+		return nullptr;
 	}
 
-	// signal to the thread that a new item is in the queue
-	m_itemInQueueEvent->Signal();
+	ProxyCommand<commands::DisableModules>* proxy = new ProxyCommand<commands::DisableModules>(true, sizeof(commands::ModuleData) * 1u);
+	proxy->m_command.processId = process::GetId();
+	proxy->m_command.moduleCount = 1u;
+	proxy->m_command.token = new Event(nullptr, Event::Type::AUTO_RESET);
+
+	commands::ModuleData moduleData = {};
+	moduleData.base = module;
+	::GetModuleFileNameW(module, moduleData.path, MAX_PATH);
+
+	proxy->m_payload.Write(moduleData);
 
-	return command->token;
+	PushUserCommand(proxy);
+
+	return proxy->m_command.token;
 }
 
 
 void* ClientUserCommandThread::DisableAllModules(const wchar_t* const nameOfExeOrDll)
 {
-	userCommands::DisableAllModulesCommand* command = new userCommands::DisableAllModulesCommand;
-	command->token = new Event(nullptr, Event::Type::AUTO_RESET);
-	command->moduleName = nameOfExeOrDll;
+	HMODULE module = ::GetModuleHandleW(nameOfExeOrDll);
+	if (!module)
 	{
-		CriticalSection::ScopedLock lock(&g_userCommandQueueCS);
-		g_userCommandQueue.push_front(command);
+		LC_ERROR_USER("Cannot disable module %S because it is not loaded by this process.", nameOfExeOrDll);
+		return nullptr;
 	}
 
-	// signal to the thread that a new item is in the queue
-	m_itemInQueueEvent->Signal();
+	const std::vector<commands::ModuleData>& allModuleData = GatherImportModuleData(module);
+	const size_t moduleCount = allModuleData.size();
+
+	ProxyCommand<commands::DisableModules>* proxy = new ProxyCommand<commands::DisableModules>(true, sizeof(commands::ModuleData) * moduleCount);
+	proxy->m_command.processId = process::GetId();
+	proxy->m_command.moduleCount = static_cast<unsigned int>(moduleCount);
+	proxy->m_command.token = new Event(nullptr, Event::Type::AUTO_RESET);
+
+	for (size_t i = 0u; i < moduleCount; ++i)
+	{
+		const commands::ModuleData& moduleData = allModuleData[i];
+		proxy->m_payload.Write(moduleData);
+	}
+
+	PushUserCommand(proxy);
 
-	return command->token;
+	return proxy->m_command.token;
 }
 
 
@@ -519,218 +301,176 @@ void ClientUserCommandThread::WaitForToken(void* token)
 
 void ClientUserCommandThread::TriggerRecompile(void)
 {
-	userCommands::TriggerRecompileCommand* command = new userCommands::TriggerRecompileCommand;
-	{
-		CriticalSection::ScopedLock lock(&g_userCommandQueueCS);
-		g_userCommandQueue.push_front(command);
-	}
+	ProxyCommand<commands::TriggerRecompile>* proxy = new ProxyCommand<commands::TriggerRecompile>(false, 0u);
 
-	// signal to the thread that a new item is in the queue
-	m_itemInQueueEvent->Signal();
+	PushUserCommand(proxy);
 }
 
 
 void ClientUserCommandThread::BuildPatch(const wchar_t* moduleNames[], const wchar_t* objPaths[], const wchar_t* amalgamatedObjPaths[], unsigned int count)
 {
-	userCommands::BuildPatchCommand* command = new userCommands::BuildPatchCommand;
-	command->count = count;
-	command->moduleNames.reserve(count);
-	command->objPaths.reserve(count);
-	command->amalgamatedObjPaths.reserve(count);
+	const size_t perFileSize = sizeof(wchar_t) * MAX_PATH * 3u;
+
+	ProxyCommand<commands::BuildPatch>* proxy = new ProxyCommand<commands::BuildPatch>(false, perFileSize*count);
+	proxy->m_command.fileCount = count;
 
 	for (unsigned int i = 0u; i < count; ++i)
 	{
-		command->moduleNames.push_back(moduleNames[i]);
-		command->objPaths.push_back(objPaths[i]);
+		commands::BuildPatch::PatchData patchData = {};
+		wcscpy_s(patchData.moduleName, moduleNames[i]);
+		wcscpy_s(patchData.objPath, objPaths[i]);
 
 		// the amalgamated object paths are optional
 		if (amalgamatedObjPaths && amalgamatedObjPaths[i])
 		{
-			command->amalgamatedObjPaths.push_back(amalgamatedObjPaths[i]);
+			wcscpy_s(patchData.amalgamatedObjPath, amalgamatedObjPaths[i]);
 		}
-		else
-		{
-			command->amalgamatedObjPaths.push_back(std::wstring());
-		}
-	}
-	{
-		CriticalSection::ScopedLock lock(&g_userCommandQueueCS);
-		g_userCommandQueue.push_front(command);
+
+		proxy->m_payload.Write(patchData);
 	}
 
-	// signal to the thread that a new item is in the queue
-	m_itemInQueueEvent->Signal();
+	PushUserCommand(proxy);
 }
 
 
-void ClientUserCommandThread::InstallExceptionHandler(void)
+void ClientUserCommandThread::ApplySettingBool(const char* const settingName, int value)
 {
-	// BEGIN EPIC MOD - Using internal CrashReporter
-	// exceptionHandler::Register(this);
-	// END EPIC MOD
+	ProxyCommand<commands::ApplySettingBool>* proxy = new ProxyCommand<commands::ApplySettingBool>(false, 0u);
+	strcpy_s(proxy->m_command.settingName, settingName);
+	proxy->m_command.settingValue = value;
+
+	PushUserCommand(proxy);
 }
 
 
-ClientUserCommandThread::ExceptionResult ClientUserCommandThread::HandleException(EXCEPTION_RECORD* exception, CONTEXT* context, unsigned int threadId)
+void ClientUserCommandThread::ApplySettingInt(const char* const settingName, int value)
 {
-	commands::HandleException serverCommand;
-	serverCommand.processId = process::GetId();
-	serverCommand.threadId = threadId;
-	serverCommand.exception = *exception;
-	serverCommand.context = *context;
-	serverCommand.clientContextPtr = context;
-
-	m_exceptionPipe->SendCommandAndWaitForAck(serverCommand);
+	ProxyCommand<commands::ApplySettingInt>* proxy = new ProxyCommand<commands::ApplySettingInt>(false, 0u);
+	strcpy_s(proxy->m_command.settingName, settingName);
+	proxy->m_command.settingValue = value;
 
-	ExceptionResult result = {};
-
-	CommandMap commandMap;
-	commandMap.RegisterAction<actions::HandleExceptionFinished>();
-	commandMap.HandleCommands(m_exceptionPipe, &result);
-
-	return result;
+	PushUserCommand(proxy);
 }
 
 
-void ClientUserCommandThread::End(void)
+void ClientUserCommandThread::ApplySettingString(const char* const settingName, const wchar_t* const value)
 {
-	// signal to the thread that a new item is in the queue to make it break out of its main loop
-	m_itemInQueueEvent->Reset();
-	m_itemInQueueEvent->Signal();
-}
+	ProxyCommand<commands::ApplySettingString>* proxy = new ProxyCommand<commands::ApplySettingString>(false, 0u);
+	strcpy_s(proxy->m_command.settingName, settingName);
+	wcscpy_s(proxy->m_command.settingValue, value);
 
+	PushUserCommand(proxy);
+}
 
 // BEGIN EPIC MOD - Adding ShowConsole command
 void ClientUserCommandThread::ShowConsole()
 {
-	userCommands::ShowConsoleCommand* command = new userCommands::ShowConsoleCommand;
-	{
-		CriticalSection::ScopedLock lock(&g_userCommandQueueCS);
-		g_userCommandQueue.push_front(command);
-	}
-
-	// signal to the thread that a new item is in the queue
-	m_itemInQueueEvent->Signal();
+	ProxyCommand<commands::ShowConsole>* proxy = new ProxyCommand<commands::ShowConsole>(false, 0u);
+	PushUserCommand(proxy);
 }
 // END EPIC MOD
 
 // BEGIN EPIC MOD - Adding SetVisible command
 void ClientUserCommandThread::SetVisible(bool visible)
 {
-	userCommands::SetVisibleCommand* command = new userCommands::SetVisibleCommand;
-	command->visible = visible;
-	{
-		CriticalSection::ScopedLock lock(&g_userCommandQueueCS);
-		g_userCommandQueue.push_front(command);
-	}
+	ProxyCommand<commands::SetVisible>* proxy = new ProxyCommand<commands::SetVisible>(false, 0u);
+	proxy->m_command.visible = visible;
 
-	// signal to the thread that a new item is in the queue
-	m_itemInQueueEvent->Signal();
+	PushUserCommand(proxy);
 }
 // END EPIC MOD
 
 // BEGIN EPIC MOD - Adding SetActive command
 void ClientUserCommandThread::SetActive(bool active)
 {
-	userCommands::SetActiveCommand* command = new userCommands::SetActiveCommand;
-	command->active = active;
-	{
-		CriticalSection::ScopedLock lock(&g_userCommandQueueCS);
-		g_userCommandQueue.push_front(command);
-	}
+	ProxyCommand<commands::SetActive>* proxy = new ProxyCommand<commands::SetActive>(false, 0u);
+	proxy->m_command.active = active;
 
-	// signal to the thread that a new item is in the queue
-	m_itemInQueueEvent->Signal();
+	PushUserCommand(proxy);
 }
 // END EPIC MOD
 
 // BEGIN EPIC MOD - Adding SetBuildArguments command
 void ClientUserCommandThread::SetBuildArguments(const wchar_t* arguments)
 {
-	userCommands::SetBuildArgumentsCommand* command = new userCommands::SetBuildArgumentsCommand;
-	command->arguments = arguments;
-	{
-		CriticalSection::ScopedLock lock(&g_userCommandQueueCS);
-		g_userCommandQueue.push_front(command);
-	}
+	ProxyCommand<commands::SetBuildArguments>* proxy = new ProxyCommand<commands::SetBuildArguments>(false, 0u);
+	proxy->m_command.processId = process::GetId();
+	wcscpy_s(proxy->m_command.arguments, arguments);
 
-	// signal to the thread that a new item is in the queue
-	m_itemInQueueEvent->Signal();
+	PushUserCommand(proxy);
 }
 // END EPIC MOD
 
 // BEGIN EPIC MOD - Adding support for lazy-loading modules
 void ClientUserCommandThread::EnableLazyLoadedModule(const wchar_t* fileName, Windows::HMODULE moduleBase)
 {
-	userCommands::EnableLazyLoadedModuleCommand* command = new userCommands::EnableLazyLoadedModuleCommand;
-	command->fileName = fileName;
-	command->moduleBase = moduleBase;
-	{
-		CriticalSection::ScopedLock lock(&g_userCommandQueueCS);
-		g_userCommandQueue.push_front(command);
-	}
+	ProxyCommand<commands::EnableLazyLoadedModule>* proxy = new ProxyCommand<commands::EnableLazyLoadedModule>(false, 0u);
+	proxy->m_command.processId = process::GetId();
+	wcscpy_s(proxy->m_command.fileName, fileName);
+	proxy->m_command.moduleBase = moduleBase;
 
-	// signal to the thread that a new item is in the queue
-	m_itemInQueueEvent->Signal();
+	PushUserCommand(proxy);
 }
 // END EPIC MOD
 
 
-void ClientUserCommandThread::ApplySettingBool(const char* const settingName, int value)
+void ClientUserCommandThread::InstallExceptionHandler(void)
 {
-	userCommands::ApplySettingBoolCommand* command = new userCommands::ApplySettingBoolCommand;
-	command->settingName = settingName;
-	command->value = value;
-	{
-		CriticalSection::ScopedLock lock(&g_userCommandQueueCS);
-		g_userCommandQueue.push_front(command);
-	}
-
-	// signal to the thread that a new item is in the queue
-	m_itemInQueueEvent->Signal();
+	// BEGIN EPIC MOD - Using internal CrashReporter
+	// exceptionHandler::Register(this);
+	// END EPIC MOD
 }
 
 
-void ClientUserCommandThread::ApplySettingInt(const char* const settingName, int value)
+ClientUserCommandThread::ExceptionResult ClientUserCommandThread::HandleException(EXCEPTION_RECORD* exception, CONTEXT* context, unsigned int threadId)
 {
-	userCommands::ApplySettingIntCommand* command = new userCommands::ApplySettingIntCommand;
-	command->settingName = settingName;
-	command->value = value;
-	{
-		CriticalSection::ScopedLock lock(&g_userCommandQueueCS);
-		g_userCommandQueue.push_front(command);
-	}
+	commands::HandleException serverCommand;
+	serverCommand.processId = process::GetId();
+	serverCommand.threadId = threadId;
+	serverCommand.exception = *exception;
+	serverCommand.context = *context;
+	serverCommand.clientContextPtr = context;
 
-	// signal to the thread that a new item is in the queue
-	m_itemInQueueEvent->Signal();
+	m_exceptionPipe->SendCommandAndWaitForAck(serverCommand, nullptr, 0u);
+
+	ExceptionResult result = {};
+
+	CommandMap commandMap;
+	commandMap.RegisterAction<actions::HandleExceptionFinished>();
+	commandMap.HandleCommands(m_exceptionPipe, &result);
+
+	return result;
 }
 
 
-void ClientUserCommandThread::ApplySettingString(const char* const settingName, const wchar_t* const value)
+void ClientUserCommandThread::End(void)
+{
+	// signal to the thread that a new item is in the queue to make it break out of its main loop
+	PushUserCommand(nullptr);
+}
+
+
+void ClientUserCommandThread::PushUserCommand(BaseCommand* command)
 {
-	userCommands::ApplySettingStringCommand* command = new userCommands::ApplySettingStringCommand;
-	command->settingName = settingName;
-	command->value = value;
 	{
-		CriticalSection::ScopedLock lock(&g_userCommandQueueCS);
-		g_userCommandQueue.push_front(command);
+		CriticalSection::ScopedLock lock(&m_userCommandQueueCS);
+		m_userCommandQueue.push_front(command);
 	}
 
 	// signal to the thread that a new item is in the queue
-	m_itemInQueueEvent->Signal();
+	m_userCommandQueueSema.Signal();
 }
 
 
-unsigned int __stdcall ClientUserCommandThread::ThreadProxy(void* context)
+ClientUserCommandThread::BaseCommand* ClientUserCommandThread::PopUserCommand(void)
 {
-	thread::SetName("Live coding user commands");
+	m_userCommandQueueSema.Wait();
 
-	ThreadContext* realContext = static_cast<ThreadContext*>(context);
-	const unsigned int exitCode = realContext->thisInstance->ThreadFunction(realContext->waitForStartEvent, realContext->pipeAccessCS);
+	CriticalSection::ScopedLock lock(&m_userCommandQueueCS);
+	BaseCommand* command = m_userCommandQueue.back();
+	m_userCommandQueue.pop_back();
 
-	delete realContext;
-
-	return exitCode;
+	return command;
 }
 
 
@@ -782,11 +522,8 @@ unsigned int ClientUserCommandThread::ThreadFunction(Event* waitForStartEvent, C
 	waitForStartEvent->Wait();
 
 	CommandMap moduleCommandMap;
-	moduleCommandMap.RegisterAction<actions::GetModule>();
-	moduleCommandMap.RegisterAction<actions::EnableModuleFinished>();
-	moduleCommandMap.RegisterAction<actions::EnableAllModulesFinished>();
-	moduleCommandMap.RegisterAction<actions::DisableModuleFinished>();
-	moduleCommandMap.RegisterAction<actions::DisableAllModulesFinished>();
+	moduleCommandMap.RegisterAction<actions::EnableModulesFinished>();
+	moduleCommandMap.RegisterAction<actions::DisableModulesFinished>();
 
 	// those commands are needed when loading compiled patches into spawned executables
 	moduleCommandMap.RegisterAction<actions::LoadPatch>();
@@ -798,132 +535,39 @@ unsigned int ClientUserCommandThread::ThreadFunction(Event* waitForStartEvent, C
 
 	for (;;)
 	{
-		// wait for event that signals that something is in the queue
-		m_itemInQueueEvent->Wait();
-
-		if (!m_pipe->IsValid())
+		// wait until a command becomes available in the queue
+		BaseCommand* command = PopUserCommand();
+		if (command == nullptr)
 		{
 			// BEGIN EPIC MOD - Using internal CrashReporter
-			// // pipe was closed or is broken, bail out
+			// // no new item available, bail out
 			// exceptionHandler::Unregister();
 			// END EPIC MOD - Using internal CrashReporter
-			return 1u;
+			return 2u;
 		}
 
-		// lock critical section for accessing the pipe.
-		// we need to make sure that other threads talking through the pipe don't use it at the same time.
-		CriticalSection::ScopedLock pipeLock(pipeAccessCS);
-
-		// lock critical section for accessing the queue.
-		// user code might be calling other exported functions in the mean time.
-		// BEGIN EPIC MOD - Rearrange lock scope to prevent hangs
-		LeaveableScopedLock queueLock(&g_userCommandQueueCS);
-		// END EPIC MOD
-
-		if (g_userCommandQueue.size() == 0u)
+		if (!m_pipe->IsValid())
 		{
 			// BEGIN EPIC MOD - Using internal CrashReporter
-			// // no new item available, bail out
+			// // pipe was closed or is broken, bail out
 			// exceptionHandler::Unregister();
 			// END EPIC MOD - Using internal CrashReporter
-			return 2u;
-		}
-
-		// separate commands into three groups: ones that need to be scoped for enabling modules, the ones that
-		// need to be scoped for disabling modules, and the others that don't need to be scoped at all
-		std::vector<userCommands::BaseCommand*> enabledScopedCommands;
-		enabledScopedCommands.reserve(g_userCommandQueue.size());
-
-		std::vector<userCommands::BaseCommand*> disableScopedCommands;
-		disableScopedCommands.reserve(g_userCommandQueue.size());
-
-		std::vector<userCommands::BaseCommand*> commands;
-		commands.reserve(g_userCommandQueue.size());
-
-		while (g_userCommandQueue.size() > 0u)
-		{
-			userCommands::BaseCommand* command = g_userCommandQueue.back();
-			g_userCommandQueue.pop_back();
-
-			if (command->GetScope() == userCommands::Scope::NONE)
-			{
-				commands.push_back(command);
-			}
-			else if (command->GetScope() == userCommands::Scope::ENABLE_MODULES)
-			{
-				enabledScopedCommands.push_back(command);
-			}
-			else if (command->GetScope() == userCommands::Scope::DISABLE_MODULES)
-			{
-				disableScopedCommands.push_back(command);
-			}
+			return 1u;
 		}
 
-		// BEGIN EPIC MOD - Temporarily release lock to prevent hangs
-		queueLock.Leave();
-		// BEGIN EPIC MOD - Pre
-
-		// send out scoped commands first
-		{
-			const size_t count = enabledScopedCommands.size();
-			if (count != 0u)
-			{
-				m_pipe->SendCommandAndWaitForAck(commands::EnableModuleBatchBegin {});
-
-				for (size_t i=0u; i < count; ++i)
-				{
-					userCommands::BaseCommand* command = enabledScopedCommands[i];
-					command->Execute(m_pipe);
-
-					moduleCommandMap.HandleCommands(m_pipe, nullptr);
-
-					delete command;
-				}
-
-				m_pipe->SendCommandAndWaitForAck(commands::EnableModuleBatchEnd {});
-			}
-		}
+		// lock critical section for accessing the pipe.
+		// we need to make sure that other threads talking through the pipe don't use it at the same time.
 		{
-			const size_t count = disableScopedCommands.size();
-			if (count != 0u)
-			{
-				m_pipe->SendCommandAndWaitForAck(commands::DisableModuleBatchBegin{});
-
-				for (size_t i = 0u; i < count; ++i)
-				{
-					userCommands::BaseCommand* command = disableScopedCommands[i];
-					command->Execute(m_pipe);
-
-					moduleCommandMap.HandleCommands(m_pipe, nullptr);
+			CriticalSection::ScopedLock pipeLock(pipeAccessCS);
 
-					delete command;
-				}
-
-				m_pipe->SendCommandAndWaitForAck(commands::DisableModuleBatchEnd{});
-			}
-		}
-
-		// send out non-scoped commands second
-		{
-			const size_t count = commands.size();
-			for (size_t i = 0u; i < count; ++i)
+			command->Execute(m_pipe);
+			if (command->ExpectsResponse())
 			{
-				userCommands::BaseCommand* command = commands[i];
-				command->Execute(m_pipe);
-
-				delete command;
+				moduleCommandMap.HandleCommands(m_pipe, nullptr);
 			}
 		}
 
-		// BEGIN EPIC MOD - Temporarily release lock to prevent hangs
-		queueLock.Enter();
-		if(g_userCommandQueue.size() > 0)
-		{
-			continue;
-		}
-		// BEGIN EPIC MOD
-
-		m_itemInQueueEvent->Reset();
+		delete command;
 	}
 
 	return 0u;
diff --git a/Engine/Source/Developer/Windows/LiveCoding/Private/External/LC_ClientUserCommandThread.h b/Engine/Source/Developer/Windows/LiveCoding/Private/External/LC_ClientUserCommandThread.h
--- a/Engine/Source/Developer/Windows/LiveCoding/Private/External/LC_ClientUserCommandThread.h
+++ b/Engine/Source/Developer/Windows/LiveCoding/Private/External/LC_ClientUserCommandThread.h
@@ -4,17 +4,34 @@
 
 #include "CoreTypes.h"
 #include "LC_Thread.h"
+#include "LC_CriticalSection.h"
+#include "LC_Semaphore.h"
 #include <string>
+#include <deque>
 
+class DuplexPipe;
 class DuplexPipeClient;
 class Event;
-class CriticalSection;
 
 
 // handles incoming commands from the host (the executable that loaded the Live++ DLL)
 class ClientUserCommandThread
 {
 public:
+	class BaseCommand
+	{
+	public:
+		explicit BaseCommand(bool expectResponse);
+		virtual ~BaseCommand(void);
+
+		virtual void Execute(DuplexPipe* pipe) = 0;
+
+		bool ExpectsResponse(void) const;
+
+	private:
+		bool m_expectResponse;
+	};
+
 	struct ExceptionResult
 	{
 		const void* returnAddress;
@@ -43,10 +60,6 @@ public:
 	void TriggerRecompile(void);
 	void BuildPatch(const wchar_t* moduleNames[], const wchar_t* objPaths[], const wchar_t* amalgamatedObjPaths[], unsigned int count);
 
-	void InstallExceptionHandler(void);
-	ExceptionResult HandleException(EXCEPTION_RECORD* exception, CONTEXT* context, unsigned int threadId);
-	void End(void);
-
 	// BEGIN EPIC MOD - Adding ShowConsole command
 	void ShowConsole();
 	// END EPIC MOD
@@ -71,20 +84,27 @@ public:
 	void ApplySettingInt(const char* const settingName, int value);
 	void ApplySettingString(const char* const settingName, const wchar_t* const value);
 
+	void InstallExceptionHandler(void);
+	ExceptionResult HandleException(EXCEPTION_RECORD* exception, CONTEXT* context, unsigned int threadId);
+	void End(void);
+
 private:
-	struct ThreadContext
-	{
-		ClientUserCommandThread* thisInstance;
-		Event* waitForStartEvent;
-		CriticalSection* pipeAccessCS;
-	};
+	// pushes a user command into the command queue
+	void PushUserCommand(BaseCommand* command);
+
+	// pops a user command from the command queue.
+	// blocks until a command becomes available.
+	BaseCommand* PopUserCommand(void);
 
-	static unsigned int __stdcall ThreadProxy(void* context);
 	unsigned int ThreadFunction(Event* waitForStartEvent, CriticalSection* pipeAccessCS);
 
 	thread::Handle m_thread;
 	std::wstring m_processGroupName;
 	DuplexPipeClient* m_pipe;
 	DuplexPipeClient* m_exceptionPipe;
-	Event* m_itemInQueueEvent;
+
+	// queue for working on commands received by user code
+	std::deque<BaseCommand*> m_userCommandQueue;
+	CriticalSection m_userCommandQueueCS;
+	Semaphore m_userCommandQueueSema;
 };
diff --git a/Engine/Source/Developer/Windows/LiveCoding/Private/External/LC_CommandMap.cpp b/Engine/Source/Developer/Windows/LiveCoding/Private/External/LC_CommandMap.cpp
--- a/Engine/Source/Developer/Windows/LiveCoding/Private/External/LC_CommandMap.cpp
+++ b/Engine/Source/Developer/Windows/LiveCoding/Private/External/LC_CommandMap.cpp
@@ -7,11 +7,11 @@
 namespace
 {
 	template <class T>
-	static bool DefaultAction(const DuplexPipe* pipe, void*)
+	static bool DefaultAction(const DuplexPipe* pipe, void*, void* payload, size_t payloadSize)
 	{
 		// receive command and continue execution
 		T command = {};
-		const bool success = pipe->ReceiveCommand(&command);
+		const bool success = pipe->ReceiveCommand(&command, payload, payloadSize);
 		pipe->SendAck();
 
 		if (!success)
@@ -38,26 +38,14 @@ CommandMap::CommandMap(void)
 	RegisterDefaultAction<commands::RegisterProcess>(m_actions);
 	RegisterDefaultAction<commands::RegisterProcessFinished>(m_actions);
 
-	RegisterDefaultAction<commands::EnableModuleBatchBegin>(m_actions);
-	RegisterDefaultAction<commands::EnableModuleBatchEnd>(m_actions);
-	RegisterDefaultAction<commands::DisableModuleBatchBegin>(m_actions);
-	RegisterDefaultAction<commands::DisableModuleBatchEnd>(m_actions);
-
-	RegisterDefaultAction<commands::EnableModule>(m_actions);
-	RegisterDefaultAction<commands::EnableModuleFinished>(m_actions);
-	RegisterDefaultAction<commands::DisableModule>(m_actions);
-	RegisterDefaultAction<commands::DisableModuleFinished>(m_actions);
-
-	RegisterDefaultAction<commands::EnableAllModules>(m_actions);
-	RegisterDefaultAction<commands::EnableAllModulesFinished>(m_actions);
-	RegisterDefaultAction<commands::DisableAllModules>(m_actions);
-	RegisterDefaultAction<commands::DisableAllModulesFinished>(m_actions);
+	RegisterDefaultAction<commands::EnableModules>(m_actions);
+	RegisterDefaultAction<commands::EnableModulesFinished>(m_actions);
+	RegisterDefaultAction<commands::DisableModules>(m_actions);
+	RegisterDefaultAction<commands::DisableModulesFinished>(m_actions);
 
 	RegisterDefaultAction<commands::EnterSyncPoint>(m_actions);
 	RegisterDefaultAction<commands::LeaveSyncPoint>(m_actions);
 	RegisterDefaultAction<commands::CallHooks>(m_actions);
-	RegisterDefaultAction<commands::GetModule>(m_actions);
-	RegisterDefaultAction<commands::GetModuleInfo>(m_actions);
 	RegisterDefaultAction<commands::LoadPatch>(m_actions);
 	RegisterDefaultAction<commands::LoadPatchInfo>(m_actions);
 	RegisterDefaultAction<commands::UnloadPatch>(m_actions);
@@ -102,20 +90,31 @@ bool CommandMap::HandleCommands(const DuplexPipe* pipe, void* context)
 {
 	for (;;)
 	{
-		// fetch incoming command id
-		uint32_t commandId = 0u;
+		// fetch incoming command header
+		commands::Header header = {};
 		{
-			const bool success = pipe->ReceiveCommandId(&commandId);
+			const bool success = pipe->ReceiveHeader(&header);
 			if (!success)
 			{
 				return false;
 			}
 		}
 
+		const Action action = m_actions[header.commandId];
+
+		// make space for optional payload
+		void* payload = nullptr;
+		if (header.payloadSize != 0u)
+		{
+			payload = malloc(header.payloadSize);
+		}
+
 		// call handler for this command
-		const Action action = m_actions[commandId];
+		const bool continueExecution = action(pipe, context, payload, header.payloadSize);
+
+		// free payload
+		free(payload);
 
-		const bool continueExecution = action(pipe, context);
 		if (!continueExecution)
 		{
 			return true;
diff --git a/Engine/Source/Developer/Windows/LiveCoding/Private/External/LC_CommandMap.h b/Engine/Source/Developer/Windows/LiveCoding/Private/External/LC_CommandMap.h
--- a/Engine/Source/Developer/Windows/LiveCoding/Private/External/LC_CommandMap.h
+++ b/Engine/Source/Developer/Windows/LiveCoding/Private/External/LC_CommandMap.h
@@ -11,23 +11,23 @@ class DuplexPipe;
 class CommandMap
 {
 	template <class T>
-	static bool ReceiveAndCallAction(const DuplexPipe* pipe, void* context)
+	static bool ReceiveAndCallAction(const DuplexPipe* pipe, void* context, void* payload, size_t payloadSize)
 	{
 		// receive command and call user action
 		typename T::CommandType command = {};
-		const bool success = pipe->ReceiveCommand(&command);
+		const bool success = pipe->ReceiveCommand(&command, payload, payloadSize);
 		if (!success)
 		{
 			pipe->SendAck();
 			return false;
 		}
 
-		return T::Execute(&command, pipe, context);
+		return T::Execute(&command, pipe, context, payload, payloadSize);
 	}
 
 public:
 	// an action returns whether execution should continue
-	typedef bool (*Action)(const DuplexPipe*, void*);
+	typedef bool (*Action)(const DuplexPipe*, void*, void*, size_t);
 
 	CommandMap(void);
 	~CommandMap(void);
diff --git a/Engine/Source/Developer/Windows/LiveCoding/Private/External/LC_Commands.h b/Engine/Source/Developer/Windows/LiveCoding/Private/External/LC_Commands.h
--- a/Engine/Source/Developer/Windows/LiveCoding/Private/External/LC_Commands.h
+++ b/Engine/Source/Developer/Windows/LiveCoding/Private/External/LC_Commands.h
@@ -12,6 +12,18 @@ namespace hook
 
 namespace commands
 {
+	struct Header
+	{
+		uint32_t commandId;
+		uint32_t payloadSize;
+	};
+
+	struct ModuleData
+	{
+		void* base;
+		wchar_t path[WINDOWS_MAX_PATH];
+	};
+
 	// acknowledge that a command has been received
 	struct Acknowledge
 	{
@@ -26,6 +38,7 @@ namespace commands
 		void* processBase;
 		unsigned int processId;
 		unsigned int threadId;		// thread ID of Live++ thread running in host
+		const void* jumpToSelf;		// address of jump-to-self instruction in host
 	};
 
 	// tell the DLL that registration has finished
@@ -36,96 +49,42 @@ namespace commands
 		bool success;
 	};
 
-	// tell Live++ that a new batch of modules is going to be loaded
-	struct EnableModuleBatchBegin
+	// tell Live++ to enable modules for live coding
+	struct EnableModules
 	{
 		static const uint32_t ID = RegisterProcessFinished::ID + 1u;
-	};
-	struct EnableModuleBatchEnd
-	{
-		static const uint32_t ID = EnableModuleBatchBegin::ID + 1u;
-	};
-
-
-	// tell Live++ that a new batch of modules is going to be unloaded
-	struct DisableModuleBatchBegin
-	{
-		static const uint32_t ID = EnableModuleBatchEnd::ID + 1u;
-	};
-	struct DisableModuleBatchEnd
-	{
-		static const uint32_t ID = DisableModuleBatchBegin::ID + 1u;
-	};
-
-
-	// tell Live++ to enable a module for live coding
-	struct EnableModule
-	{
-		static const uint32_t ID = DisableModuleBatchEnd::ID + 1u;
 
 		unsigned int processId;
-		wchar_t path[WINDOWS_MAX_PATH];
-		void* token;
-	};
-
-	// tell the DLL that enabling a module has finished
-	struct EnableModuleFinished
-	{
-		static const uint32_t ID = EnableModule::ID + 1u;
-
+		unsigned int moduleCount;
 		void* token;
-	};
-
-	// tell Live++ to disable a module for live coding
-	struct DisableModule
-	{
-		static const uint32_t ID = EnableModuleFinished::ID + 1u;
 
-		unsigned int processId;
-		wchar_t path[WINDOWS_MAX_PATH];
-		void* token;
+		// this command always contains an array of 'ModuleData x moduleCount' as payload
 	};
 
-	// tell the DLL that disabling a module has finished
-	struct DisableModuleFinished
+	// tell the DLL that enabling modules has finished
+	struct EnableModulesFinished
 	{
-		static const uint32_t ID = DisableModule::ID + 1u;
+		static const uint32_t ID = EnableModules::ID + 1u;
 
 		void* token;
 	};
 
-	// tell Live++ to enable a module and all its imports for live coding
-	struct EnableAllModules
+	// tell Live++ to disable modules for live coding
+	struct DisableModules
 	{
-		static const uint32_t ID = DisableModuleFinished::ID + 1u;
+		static const uint32_t ID = EnableModulesFinished::ID + 1u;
 
 		unsigned int processId;
-		wchar_t path[WINDOWS_MAX_PATH];
+		unsigned int moduleCount;
 		void* token;
-	};
-
-	// tell the DLL that enabling all modules has finished
-	struct EnableAllModulesFinished
-	{
-		static const uint32_t ID = EnableAllModules::ID + 1u;
 
-		void* token;
+		// this command always contains an array of 'ModuleData x moduleCount' as payload
 	};
 
-	// tell Live++ to disable a module and all its imports for live coding
-	struct DisableAllModules
-	{
-		static const uint32_t ID = EnableAllModulesFinished::ID + 1u;
-
-		unsigned int processId;
-		wchar_t path[WINDOWS_MAX_PATH];
-		void* token;
-	};
-
-	// tell the DLL that disabling all modules has finished
-	struct DisableAllModulesFinished
+	// tell the DLL that disabling a module has finished
+	struct DisableModulesFinished
 	{
-		static const uint32_t ID = DisableAllModules::ID + 1u;
+		static const uint32_t ID = DisableModules::ID + 1u;
 
 		void* token;
 	};
@@ -133,7 +92,7 @@ namespace commands
 	// tell the DLL to enter the synchronization point
 	struct EnterSyncPoint
 	{
-		static const uint32_t ID = DisableAllModulesFinished::ID + 1u;
+		static const uint32_t ID = DisableModulesFinished::ID + 1u;
 	};
 
 	// tell the DLL to leave the synchronization point
@@ -151,32 +110,10 @@ namespace commands
 		const hook::Function* last;
 	};
 
-	// tell the DLL to give us information about a module
-	struct GetModule
-	{
-		static const uint32_t ID = CallHooks::ID + 1u;
-
-		bool loadImports;
-		void* taskContext;
-		wchar_t path[WINDOWS_MAX_PATH];
-	};
-
-	// returns module info to Live++
-	struct GetModuleInfo
-	{
-		static const uint32_t ID = GetModule::ID + 1u;
-
-		Windows::HMODULE moduleBase;
-		unsigned int processId;
-		bool loadImports;
-		void* taskContext;
-		wchar_t path[WINDOWS_MAX_PATH];
-	};
-
 	// tell the DLL to load a DLL
 	struct LoadPatch
 	{
-		static const uint32_t ID = GetModuleInfo::ID + 1u;
+		static const uint32_t ID = CallHooks::ID + 1u;
 
 		wchar_t path[WINDOWS_MAX_PATH];
 	};
@@ -210,11 +147,6 @@ namespace commands
 	struct LogOutput
 	{
 		static const uint32_t ID = CallEntryPoint::ID + 1u;
-
-		static const size_t BUFFER_SIZE = 1024u;
-
-		size_t size;
-		wchar_t buffer[BUFFER_SIZE];
 	};
 
 	// tell Live++ server we're ready for compilation
@@ -246,23 +178,21 @@ namespace commands
 	{
 		static const uint32_t ID = TriggerRecompile::ID + 1u;
 
-		// number of BuildPatchPackets to follow this command
-		unsigned int count;
-	};
-
-	struct BuildPatchPacket
-	{
-		static const uint32_t ID = BuildPatch::ID + 1u;
+		unsigned int fileCount;
 
-		wchar_t moduleName[MAX_PATH];
-		wchar_t objPath[MAX_PATH];
-		wchar_t amalgamatedObjPath[MAX_PATH];
+		// this command always contains an array of 'PatchData x fileCount' as payload
+		struct PatchData
+		{
+			wchar_t moduleName[MAX_PATH];
+			wchar_t objPath[MAX_PATH];
+			wchar_t amalgamatedObjPath[MAX_PATH];
+		};
 	};
 
 	// tell Live++ to handle an exception
 	struct HandleException
 	{
-		static const uint32_t ID = BuildPatchPacket::ID + 1u;
+		static const uint32_t ID = BuildPatch::ID + 1u;
 
 		unsigned int processId;
 		unsigned int threadId;
diff --git a/Engine/Source/Developer/Windows/LiveCoding/Private/External/LC_DuplexPipe.cpp b/Engine/Source/Developer/Windows/LiveCoding/Private/External/LC_DuplexPipe.cpp
--- a/Engine/Source/Developer/Windows/LiveCoding/Private/External/LC_DuplexPipe.cpp
+++ b/Engine/Source/Developer/Windows/LiveCoding/Private/External/LC_DuplexPipe.cpp
@@ -29,9 +29,9 @@ bool DuplexPipe::IsValid(void) const volatile
 }
 
 
-bool DuplexPipe::ReceiveCommandId(uint32_t* id) const
+bool DuplexPipe::ReceiveHeader(commands::Header* header) const
 {
-	const bool success = Read(id, sizeof(uint32_t));
+	const bool success = Read(header, sizeof(commands::Header));
 	return success;
 }
 
diff --git a/Engine/Source/Developer/Windows/LiveCoding/Private/External/LC_DuplexPipe.h b/Engine/Source/Developer/Windows/LiveCoding/Private/External/LC_DuplexPipe.h
--- a/Engine/Source/Developer/Windows/LiveCoding/Private/External/LC_DuplexPipe.h
+++ b/Engine/Source/Developer/Windows/LiveCoding/Private/External/LC_DuplexPipe.h
@@ -4,7 +4,7 @@
 
 #include "CoreTypes.h"
 #include "LC_Commands.h"
-
+#include "LC_Platform.h"
 
 class DuplexPipe
 {
@@ -17,25 +17,39 @@ public:
 
 	// send command, synchronous
 	template <typename T>
-	void SendCommandAndWaitForAck(const T& command) const
+	void SendCommandAndWaitForAck(const T& command, const void* payload, size_t payloadSize) const
 	{
-		// send the command ID and the command, and wait for ACK to come in
-		const uint32_t commandId = T::ID;
-		Send(&commandId, sizeof(commandId));
-		Send(&command, sizeof(command));
+		// send the header and command, and wait for the ACK to be received
+		const commands::Header header = { T::ID, static_cast<uint32_t>(payloadSize) };
+		Send(&header, sizeof(commands::Header));
+		Send(&command, sizeof(T));
+
+		if (payloadSize != 0u)
+		{
+			LC_ASSERT(payload != nullptr, "Cannot send command payload.");
+			Send(payload, payloadSize);
+		}
 
 		commands::Acknowledge ack = {};
 		Read(&ack, sizeof(commands::Acknowledge));
 	}
 
-	// receive command ID
-	bool ReceiveCommandId(uint32_t* id) const;
+	// receive command header
+	bool ReceiveHeader(commands::Header* header) const;
 
+	// receive any command
 	template <typename T>
-	bool ReceiveCommand(T* command) const
+	bool ReceiveCommand(T* command, void* payload, size_t payloadSize) const
 	{
-		// receive command
-		const bool success = Read(command, sizeof(T));
+		// receive command and optional payload
+		bool success = Read(command, sizeof(T));
+
+		if (payloadSize != 0u)
+		{
+			LC_ASSERT(payload != nullptr, "Cannot receive command payload.");
+			success &= Read(payload, payloadSize);
+		}
+
 		return success;
 	}
 
diff --git a/Engine/Source/Developer/Windows/LiveCoding/Private/External/LC_EntryPoint.cpp b/Engine/Source/Developer/Windows/LiveCoding/Private/External/LC_EntryPoint.cpp
--- a/Engine/Source/Developer/Windows/LiveCoding/Private/External/LC_EntryPoint.cpp
+++ b/Engine/Source/Developer/Windows/LiveCoding/Private/External/LC_EntryPoint.cpp
@@ -115,6 +115,13 @@ LPP_DLL_API void* __cdecl LppDisableAllModules(const wchar_t* const nameOfExeOrD
 
 LPP_DLL_API void __cdecl LppWaitForToken(void* token)
 {
+	if (!token)
+	{
+		// nullptr tokens are returned by Live++ when trying to enable modules which are not loaded into the host process.
+		// therefore, we need to handle this case gracefully.
+		return;
+	}
+
 	g_startupThread->WaitForToken(token);
 }
 
diff --git a/Engine/Source/Developer/Windows/LiveCoding/Private/External/LC_Executable.cpp b/Engine/Source/Developer/Windows/LiveCoding/Private/External/LC_Executable.cpp
--- a/Engine/Source/Developer/Windows/LiveCoding/Private/External/LC_Executable.cpp
+++ b/Engine/Source/Developer/Windows/LiveCoding/Private/External/LC_Executable.cpp
@@ -89,7 +89,8 @@ executable::Header executable::GetHeader(const Image* image)
 		return executable::Header {};
 	}
 
-	return executable::Header { ntHeader->FileHeader, ntHeader->OptionalHeader.SizeOfImage };
+	const uint64_t sizeOnDisk = file::GetSizeOnDisk(image);
+	return executable::Header { sizeOnDisk, ntHeader->FileHeader, ntHeader->OptionalHeader };
 }
 
 
@@ -181,7 +182,7 @@ void executable::RebaseImage(Image* image, PreferredBase preferredBase)
 		return;
 	}
 
-	ImageSectionDB* database = GatherSections(image);
+	ImageSectionDB* database = GatherImageSectionDB(image);
 
 	// the image has been linked against a certain base address, namely ntHeader->OptionalHeader.ImageBase.
 	// work out by how much all relocations need to be shifted if basing the image against the new
@@ -257,7 +258,7 @@ void executable::RebaseImage(Image* image, PreferredBase preferredBase)
 }
 
 
-executable::ImageSectionDB* executable::GatherSections(const Image* image)
+executable::ImageSectionDB* executable::GatherImageSectionDB(const Image* image)
 {
 	const IMAGE_NT_HEADERS* ntHeader = detail::GetNtHeader(image);
 	if (!ntHeader)
@@ -293,6 +294,62 @@ void executable::DestroyImageSectionDB(ImageSectionDB* database)
 }
 
 
+executable::ImportModuleDB* executable::GatherImportModuleDB(const Image* image, const ImageSectionDB* imageSections)
+{
+	const IMAGE_NT_HEADERS* ntHeader = detail::GetNtHeader(image);
+	if (!ntHeader)
+	{
+		return nullptr;
+	}
+
+	const IMAGE_SECTION_HEADER* sectionHeader = detail::GetSectionHeader(ntHeader);
+	if (!sectionHeader)
+	{
+		return nullptr;
+	}
+
+	ImportModuleDB* database = new ImportModuleDB;
+
+	// the import directory stores an array of IMAGE_IMPORT_DESCRIPTOR entries
+	const size_t count = ntHeader->OptionalHeader.DataDirectory[IMAGE_DIRECTORY_ENTRY_IMPORT].Size / sizeof(IMAGE_IMPORT_DESCRIPTOR);
+	if (count == 0u)
+	{
+		// no import modules
+		return database;
+	}
+
+	const DWORD baseImportModule = RvaToFileOffset(imageSections, ntHeader->OptionalHeader.DataDirectory[IMAGE_DIRECTORY_ENTRY_IMPORT].VirtualAddress);
+	if (baseImportModule == 0u)
+	{
+		// no import modules
+		return database;
+	}
+
+	database->modules.reserve(count);
+
+	const IMAGE_IMPORT_DESCRIPTOR* importModules = pointer::Offset<const IMAGE_IMPORT_DESCRIPTOR*>(image->base, baseImportModule);
+	while (importModules->Name != 0)
+	{
+		const DWORD fileOffset = RvaToFileOffset(imageSections, importModules->Name);
+		const char* importModuleName = pointer::Offset<const char*>(image->base, fileOffset);
+
+		ImportModule module;
+		strcpy_s(module.path, importModuleName);
+		database->modules.emplace_back(module);
+
+		++importModules;
+	}
+
+	return database;
+}
+
+
+void executable::DestroyImportModuleDB(ImportModuleDB* database)
+{
+	delete database;
+}
+
+
 void executable::CallDllEntryPoint(void* moduleBase, uint32_t entryPointRva)
 {
 	typedef decltype(_DllMainCRTStartup) DllEntryPoint;
diff --git a/Engine/Source/Developer/Windows/LiveCoding/Private/External/LC_Executable.h b/Engine/Source/Developer/Windows/LiveCoding/Private/External/LC_Executable.h
--- a/Engine/Source/Developer/Windows/LiveCoding/Private/External/LC_Executable.h
+++ b/Engine/Source/Developer/Windows/LiveCoding/Private/External/LC_Executable.h
@@ -20,8 +20,9 @@ namespace executable
 
 	struct Header
 	{
+		uint64_t sizeOnDisk;
 		IMAGE_FILE_HEADER imageHeader;
-		uint64_t size;
+		IMAGE_OPTIONAL_HEADER optionalHeader;
 	};
 
 	typedef file::MemoryFile Image;
@@ -39,6 +40,16 @@ namespace executable
 		std::vector<ImageSection> sections;
 	};
 
+	struct ImportModule
+	{
+		char path[MAX_PATH];
+	};
+
+	struct ImportModuleDB
+	{
+		std::vector<ImportModule> modules;
+	};
+
 
 	Image* OpenImage(const wchar_t* filename, file::OpenMode::Enum openMode);
 	void CloseImage(Image*& image);
@@ -46,9 +57,13 @@ namespace executable
 	void RebaseImage(Image* image, PreferredBase preferredBase);
 
 
-	ImageSectionDB* GatherSections(const Image* image);
+	ImageSectionDB* GatherImageSectionDB(const Image* image);
 	void DestroyImageSectionDB(ImageSectionDB* database);
 
+	ImportModuleDB* GatherImportModuleDB(const Image* image, const ImageSectionDB* imageSections);
+	void DestroyImportModuleDB(ImportModuleDB* database);
+
+
 
 
 	// returns the RVA of the entry point of the image when loaded into memory
@@ -100,7 +115,10 @@ namespace std
 	{
 		inline std::size_t operator()(const executable::Header& header) const
 		{
-			return static_cast<uint32_t>(XXH32(&header.imageHeader, sizeof(executable::Header::imageHeader), 0u));
+			const uint32_t hash1 = XXH32(&header.imageHeader, sizeof(executable::Header::imageHeader), 0u);
+			const uint32_t hash2 = XXH32(&header.optionalHeader, sizeof(executable::Header::optionalHeader), hash1);
+
+			return hash2;
 		}
 	};
 
@@ -109,10 +127,23 @@ namespace std
 	{
 		inline bool operator()(const executable::Header& lhs, const executable::Header& rhs) const
 		{
-			if (lhs.size != rhs.size)
+			// comparing just the size on disk is fastest
+			if (lhs.sizeOnDisk != rhs.sizeOnDisk)
+			{
+				return false;
+			}
+			// comparing the image header is the second-best option
+			else if (::memcmp(&lhs.imageHeader, &rhs.imageHeader, sizeof(executable::Header::imageHeader)) != 0)
+			{
+				return false;
+			}
+			// comparing the optional header is slowest
+			else if (::memcmp(&lhs.optionalHeader, &rhs.optionalHeader, sizeof(executable::Header::optionalHeader)) != 0)
+			{
 				return false;
+			}
 
-			return (::memcmp(&lhs.imageHeader, &rhs.imageHeader, sizeof(executable::Header::imageHeader)) == 0);
+			return true;
 		}
 	};
 }
diff --git a/Engine/Source/Developer/Windows/LiveCoding/Private/External/LC_MemoryFile.cpp b/Engine/Source/Developer/Windows/LiveCoding/Private/External/LC_MemoryFile.cpp
--- a/Engine/Source/Developer/Windows/LiveCoding/Private/External/LC_MemoryFile.cpp
+++ b/Engine/Source/Developer/Windows/LiveCoding/Private/External/LC_MemoryFile.cpp
@@ -123,3 +123,16 @@ void file::Close(MemoryFile*& memoryFile)
 	delete memoryFile;
 	memoryFile = nullptr;
 }
+
+
+uint64_t file::GetSizeOnDisk(const MemoryFile* file)
+{
+	::BY_HANDLE_FILE_INFORMATION info = {};
+	::GetFileInformationByHandle(file->file, &info);
+
+	::ULARGE_INTEGER integer = {};
+	integer.LowPart = info.nFileSizeLow;
+	integer.HighPart = info.nFileSizeHigh;
+
+	return static_cast<uint64_t>(integer.QuadPart);
+}
diff --git a/Engine/Source/Developer/Windows/LiveCoding/Private/External/LC_MemoryFile.h b/Engine/Source/Developer/Windows/LiveCoding/Private/External/LC_MemoryFile.h
--- a/Engine/Source/Developer/Windows/LiveCoding/Private/External/LC_MemoryFile.h
+++ b/Engine/Source/Developer/Windows/LiveCoding/Private/External/LC_MemoryFile.h
@@ -29,4 +29,6 @@ namespace file
 	MemoryFile* Open(const wchar_t* path, OpenMode::Enum openMode);
 	const void* GetData(const MemoryFile* file);
 	void Close(MemoryFile*& file);
+
+	uint64_t GetSizeOnDisk(const MemoryFile* file);
 }
diff --git a/Engine/Source/Developer/Windows/LiveCoding/Private/External/LC_MemoryStream.cpp b/Engine/Source/Developer/Windows/LiveCoding/Private/External/LC_MemoryStream.cpp
new file mode 100644
--- /dev/null
+++ b/Engine/Source/Developer/Windows/LiveCoding/Private/External/LC_MemoryStream.cpp
@@ -0,0 +1,44 @@
+// Copyright 2011-2019 Molecular Matters GmbH, all rights reserved.
+
+#include "LC_MemoryStream.h"
+#include "LC_PointerUtil.h"
+
+
+memoryStream::Reader::Reader(const void* data, size_t capacity)
+	: m_data(data)
+	, m_capacity(capacity)
+	, m_offset(0u)
+{
+}
+
+
+void memoryStream::Reader::Read(void* data, size_t size)
+{
+	LC_ASSERT(m_offset + size <= m_capacity, "Not enough data left to read.");
+
+	memcpy(data, pointer::Offset<const void*>(m_data, m_offset), size);
+	m_offset += size;
+}
+
+
+memoryStream::Writer::Writer(size_t capacity)
+	: m_data(new char[capacity])
+	, m_capacity(capacity)
+	, m_offset(0u)
+{
+}
+
+
+memoryStream::Writer::~Writer(void)
+{
+	delete[] m_data;
+}
+
+
+void memoryStream::Writer::Write(const void* data, size_t size)
+{
+	LC_ASSERT(m_offset + size <= m_capacity, "Not enough space to write data.");
+
+	memcpy(m_data + m_offset, data, size);
+	m_offset += size;
+}
diff --git a/Engine/Source/Developer/Windows/LiveCoding/Private/External/LC_MemoryStream.h b/Engine/Source/Developer/Windows/LiveCoding/Private/External/LC_MemoryStream.h
new file mode 100644
--- /dev/null
+++ b/Engine/Source/Developer/Windows/LiveCoding/Private/External/LC_MemoryStream.h
@@ -0,0 +1,71 @@
+// Copyright 2011-2019 Molecular Matters GmbH, all rights reserved.
+
+#pragma once
+
+#include "CoreTypes.h"
+#include "LC_Platform.h"
+
+namespace memoryStream
+{
+	class Reader
+	{
+	public:
+		Reader(const void* data, size_t capacity);
+
+		void Read(void* data, size_t size);
+
+		template <typename T>
+		inline T Read(void)
+		{
+			T data;
+			Read(&data, sizeof(T));
+			return data;
+		}
+
+	private:
+		const void* m_data;
+		size_t m_capacity;
+		size_t m_offset;
+
+		LC_DISABLE_COPY(Reader);
+		LC_DISABLE_MOVE(Reader);
+		LC_DISABLE_ASSIGNMENT(Reader);
+		LC_DISABLE_MOVE_ASSIGNMENT(Reader);
+	};
+
+
+	class Writer
+	{
+	public:
+		explicit Writer(size_t capacity);
+		~Writer(void);
+
+		void Write(const void* data, size_t size);
+
+		template <typename T>
+		inline void Write(const T& data)
+		{
+			Write(&data, sizeof(T));
+		}
+
+		inline const void* GetData(void) const
+		{
+			return m_data;
+		}
+
+		inline size_t GetSize(void) const
+		{
+			return m_offset;
+		}
+
+	private:
+		char* m_data;
+		size_t m_capacity;
+		size_t m_offset;
+
+		LC_DISABLE_COPY(Writer);
+		LC_DISABLE_MOVE(Writer);
+		LC_DISABLE_ASSIGNMENT(Writer);
+		LC_DISABLE_MOVE_ASSIGNMENT(Writer);
+	};
+}
diff --git a/Engine/Source/Developer/Windows/LiveCoding/Private/External/LC_Thread.h b/Engine/Source/Developer/Windows/LiveCoding/Private/External/LC_Thread.h
--- a/Engine/Source/Developer/Windows/LiveCoding/Private/External/LC_Thread.h
+++ b/Engine/Source/Developer/Windows/LiveCoding/Private/External/LC_Thread.h
@@ -17,7 +17,18 @@ namespace thread
 	// returns the thread ID of the given thread
 	unsigned int GetId(Handle handle);
 
+	// raw thread creation function
 	Handle Create(unsigned int stackSize, Function function, void* context);
+
+	// creates a thread from a given function and provided arguments
+	template <typename F, typename... Args>
+	inline Handle Create(const char* threadName, unsigned int stackSize, F ptrToFunction, Args&&... args);
+
+	// creates a thread from a given member function, instance and provided arguments
+	template <class C, typename F, typename... Args>
+	inline Handle Create(const char* threadName, unsigned int stackSize, F ptrToMemberFunction, C* instance, Args&&... args);
+
+
 	void Join(Handle handle);
 	void Terminate(Handle handle);
 
@@ -63,3 +74,85 @@ namespace thread
 	// sets the name of the calling thread
 	void SetName(const char* name);
 }
+
+
+template <typename F, typename... Args>
+inline thread::Handle thread::Create(const char* threadName, unsigned int stackSize, F ptrToFunction, Args&&... args)
+{
+	// a helper lambda that calls the given functions with the provided arguments.
+	// because this lambda needs to capture, it cannot be converted to a function pointer implicitly,
+	// and therefore cannot be used as a thread function directly.
+	// however, we can solve this by another indirection.
+	auto captureLambda = [threadName, ptrToFunction, args...]() -> unsigned int
+	{
+		thread::SetName(threadName);
+
+		return (*ptrToFunction)(args...);
+	};
+
+	// helper to make the following easier to read
+	typedef decltype(captureLambda) CaptureLambdaType;
+
+	// the lambda with captures will be called in the thread about to be created, hence it
+	// must be allocated on the heap.
+	CaptureLambdaType* lambdaOnHeap = new CaptureLambdaType(captureLambda);
+
+	// here's the trick: we generate another capture-less lambda that has the same signature as a thread function,
+	// and internally cast the given object to its original type, calling the lambda with captures from within
+	// this capture-less lambda.
+	auto capturelessLambda = [](void* captureLambda) -> unsigned int
+	{
+		CaptureLambdaType* lambdaOnHeap = static_cast<CaptureLambdaType*>(captureLambda);
+		const unsigned int result = (*lambdaOnHeap)();
+		delete lambdaOnHeap;
+
+		return result;
+	};
+
+	unsigned int(__stdcall *threadFunction)(void*) = capturelessLambda;
+
+	thread::Handle handle = thread::Create(stackSize, threadFunction, lambdaOnHeap);
+
+	return handle;
+}
+
+
+template <class C, typename F, typename... Args>
+inline thread::Handle thread::Create(const char* threadName, unsigned int stackSize, F ptrToMemberFunction, C* instance, Args&&... args)
+{
+	// a helper lambda that calls the given functions with the provided arguments.
+	// because this lambda needs to capture, it cannot be converted to a function pointer implicitly,
+	// and therefore cannot be used as a thread function directly.
+	// however, we can solve this by another indirection.
+	auto captureLambda = [threadName, ptrToMemberFunction, instance, args...]() -> unsigned int
+	{
+		thread::SetName(threadName);
+
+		return (instance->*ptrToMemberFunction)(args...);
+	};
+
+	// helper to make the following easier to read
+	typedef decltype(captureLambda) CaptureLambdaType;
+
+	// the lambda with captures will be called in the thread about to be created, hence it
+	// must be allocated on the heap.
+	CaptureLambdaType* lambdaOnHeap = new CaptureLambdaType(captureLambda);
+
+	// here's the trick: we generate another capture-less lambda that has the same signature as a thread function,
+	// and internally cast the given object to its original type, calling the lambda with captures from within
+	// this capture-less lambda.
+	auto capturelessLambda = [](void* captureLambda) -> unsigned int
+	{
+		CaptureLambdaType* lambdaOnHeap = static_cast<CaptureLambdaType*>(captureLambda);
+		const unsigned int result = (*lambdaOnHeap)();
+		delete lambdaOnHeap;
+
+		return result;
+	};
+
+	unsigned int (__stdcall *threadFunction)(void*) = capturelessLambda;
+
+	thread::Handle handle = thread::Create(stackSize, threadFunction, lambdaOnHeap);
+
+	return handle;
+}
diff --git a/Engine/Source/Developer/Windows/LiveCoding/Private/External/xxhash.c b/Engine/Source/Developer/Windows/LiveCoding/Private/External/xxhash.c
--- a/Engine/Source/Developer/Windows/LiveCoding/Private/External/xxhash.c
+++ b/Engine/Source/Developer/Windows/LiveCoding/Private/External/xxhash.c
@@ -43,6 +43,7 @@
 #pragma warning(disable : 4820)
 #define XXH_FORCE_MEMORY_ACCESS 2
 #define XXH_FORCE_ALIGN_CHECK 0
+#define XXH_CPU_LITTLE_ENDIAN 1
 
 
 /* *************************************
diff --git a/Engine/Source/Developer/Windows/LiveCoding/Private/External/xxhash.h b/Engine/Source/Developer/Windows/LiveCoding/Private/External/xxhash.h
--- a/Engine/Source/Developer/Windows/LiveCoding/Private/External/xxhash.h
+++ b/Engine/Source/Developer/Windows/LiveCoding/Private/External/xxhash.h
@@ -64,6 +64,9 @@ XXH64       13.8 GB/s            1.9 GB/s
 XXH32        6.8 GB/s            6.0 GB/s
 */
 
+// MM: tuning
+// EPIC REMOVED: #define XXH_PRIVATE_API 1
+
 #ifndef XXHASH_H_5627135585666179
 #define XXHASH_H_5627135585666179 1
 
diff --git a/Engine/Source/Developer/Windows/LiveCoding/Private/LiveCodingModule.cpp b/Engine/Source/Developer/Windows/LiveCoding/Private/LiveCodingModule.cpp
--- a/Engine/Source/Developer/Windows/LiveCoding/Private/LiveCodingModule.cpp
+++ b/Engine/Source/Developer/Windows/LiveCoding/Private/LiveCodingModule.cpp
@@ -263,9 +263,6 @@ void FLiveCodingModule::UpdateModules()
 	TArray<FModuleStatus> ModuleStatuses;
 	FModuleManager::Get().QueryModules(ModuleStatuses);
 
-	extern void BeginCommandBatch();
-	BeginCommandBatch();
-
 	for (const FModuleStatus& ModuleStatus : ModuleStatuses)
 	{
 		if (ModuleStatus.bIsLoaded)
@@ -274,9 +271,6 @@ void FLiveCodingModule::UpdateModules()
 			ConfigureModule(FName(*ModuleStatus.Name), FullFilePath);
 		}
 	}
-
-	extern void EndCommandBatch();
-	EndCommandBatch();
 #endif
 }
 
diff --git a/Engine/Source/Developer/Windows/LiveCodingServer/Private/External/LC_CodeCave.cpp b/Engine/Source/Developer/Windows/LiveCodingServer/Private/External/LC_CodeCave.cpp
--- a/Engine/Source/Developer/Windows/LiveCodingServer/Private/External/LC_CodeCave.cpp
+++ b/Engine/Source/Developer/Windows/LiveCodingServer/Private/External/LC_CodeCave.cpp
@@ -1,16 +1,14 @@
 // Copyright 2011-2019 Molecular Matters GmbH, all rights reserved.
 
 #include "LC_CodeCave.h"
-#include "LC_VirtualMemory.h"
-#include "LC_Patch.h"
 #include "LC_Thread.h"
 
 
-CodeCave::CodeCave(process::Handle processHandle, unsigned int processId, unsigned int commandThreadId)
+CodeCave::CodeCave(process::Handle processHandle, unsigned int processId, unsigned int commandThreadId, const void* jumpToSelf)
 	: m_processHandle(processHandle)
 	, m_processId(processId)
 	, m_commandThreadId(commandThreadId)
-	, m_cave(nullptr)
+	, m_jumpToSelf(jumpToSelf)
 	, m_perThreadData()
 {
 	m_perThreadData.reserve(128u);
@@ -21,11 +19,6 @@ void CodeCave::Install(void)
 {
 	process::Suspend(m_processHandle);
 
-	// prepare jump-to-self code cave
-	const uint32_t pageSize = virtualMemory::GetPageSize();
-	m_cave = virtualMemory::Allocate(m_processHandle, pageSize, virtualMemory::PageType::EXECUTE_READ_WRITE);
-	patch::InstallJumpToSelf(m_processHandle, m_cave);
-
 	// enumerate all threads of the process now that it's suspended
 	const std::vector<unsigned int>& threadIds = process::EnumerateThreads(m_processId);
 	
@@ -52,7 +45,7 @@ void CodeCave::Install(void)
 		m_perThreadData[i].priority = thread::GetPriority(threadHandle);
 		m_perThreadData[i].originalIp = thread::ReadInstructionPointer(context);
 		thread::SetPriority(threadHandle, THREAD_PRIORITY_IDLE);
-		thread::WriteInstructionPointer(context, m_cave);
+		thread::WriteInstructionPointer(context, m_jumpToSelf);
 		thread::SetContext(threadHandle, context);
 		thread::Close(threadHandle);
 	}
@@ -84,7 +77,7 @@ void CodeCave::Uninstall(void)
 		// only set the original instruction pointer if the thread is really being held in the cave.
 		// in certain situations (e.g. after an exception), the debugger/OS already restored the context
 		// of all threads, and it would be fatal to interfere with this.
-		if (currentIp == m_cave)
+		if (currentIp == m_jumpToSelf)
 		{
 			thread::SetPriority(threadHandle, m_perThreadData[i].priority);
 			thread::WriteInstructionPointer(context, m_perThreadData[i].originalIp);
@@ -93,8 +86,5 @@ void CodeCave::Uninstall(void)
 		thread::Close(threadHandle);
 	}
 
-	// get rid of the code cave
-	virtualMemory::Free(m_processHandle, m_cave);
-
 	process::Resume(m_processHandle);
 }
diff --git a/Engine/Source/Developer/Windows/LiveCodingServer/Private/External/LC_CodeCave.h b/Engine/Source/Developer/Windows/LiveCodingServer/Private/External/LC_CodeCave.h
--- a/Engine/Source/Developer/Windows/LiveCodingServer/Private/External/LC_CodeCave.h
+++ b/Engine/Source/Developer/Windows/LiveCodingServer/Private/External/LC_CodeCave.h
@@ -10,7 +10,8 @@
 class CodeCave
 {
 public:
-	CodeCave(process::Handle processHandle, unsigned int processId, unsigned int commandThreadId);
+	// the jump-to-self code needs to be available in the address space of the given process
+	CodeCave(process::Handle processHandle, unsigned int processId, unsigned int commandThreadId, const void* jumpToSelf);
 
 	void Install(void);
 	void Uninstall(void);
@@ -19,8 +20,7 @@ private:
 	process::Handle m_processHandle;
 	unsigned int m_processId;
 	unsigned int m_commandThreadId;
-
-	void* m_cave;
+	const void* m_jumpToSelf;
 
 	struct PerThreadData
 	{
diff --git a/Engine/Source/Developer/Windows/LiveCodingServer/Private/External/LC_Coff.cpp b/Engine/Source/Developer/Windows/LiveCodingServer/Private/External/LC_Coff.cpp
--- a/Engine/Source/Developer/Windows/LiveCodingServer/Private/External/LC_Coff.cpp
+++ b/Engine/Source/Developer/Windows/LiveCodingServer/Private/External/LC_Coff.cpp
@@ -7,7 +7,7 @@
 #include "LC_FileUtil.h"
 #include "LC_SymbolPatterns.h"
 #include "LC_Symbols.h"
-#include "LC_MemoryBlock.h"
+#include "LC_MemoryStream.h"
 #include "LC_UniqueId.h"
 #include "LC_Allocators.h"
 #include <algorithm>
@@ -631,7 +631,7 @@ namespace coff
 
 		// keeping the section data, relocations, and line numbers at the end of the file makes it easier to keep track
 		// of file offsets that are stored in the section headers.
-		MemoryBlock outputData(static_cast<size_t>(rawCoff->size));
+		memoryStream::Writer outputData(static_cast<size_t>(rawCoff->size));
 		CoffHeader coffHeader(rawCoff->header);
 
 		// careful: the number of sections is of type WORD, but DWORD in files compiled with /bigobj!
@@ -641,7 +641,7 @@ namespace coff
 		coffHeader.NumberOfSymbols = static_cast<DWORD>(rawCoff->symbols.size());
 		coffHeader.PointerToSymbolTable = static_cast<DWORD>(sizeof(CoffHeader) + rawCoff->sections.size() * sizeof(IMAGE_SECTION_HEADER));
 
-		outputData.Insert(coffHeader);
+		outputData.Write(coffHeader);
 		const size_t baseOffset =
 			sizeof(CoffHeader) +
 			rawCoff->sections.size() * sizeof(IMAGE_SECTION_HEADER) +
@@ -697,22 +697,22 @@ namespace coff
 			}
 			sectionHeader.NumberOfLinenumbers = static_cast<WORD>(section.lineNumbers.size());
 
-			outputData.Insert(sectionHeader);
+			outputData.Write(sectionHeader);
 		}
 
 		for (size_t i = 0u; i < rawCoff->symbols.size(); ++i)
 		{
-			outputData.Insert(rawCoff->symbols[i]);
+			outputData.Write(rawCoff->symbols[i]);
 		}
 
-		outputData.Insert(rawCoff->rawStringTable.data, rawCoff->rawStringTable.size);
+		outputData.Write(rawCoff->rawStringTable.data, rawCoff->rawStringTable.size);
 
 		for (size_t i = 0u; i < rawCoff->sections.size(); ++i)
 		{
 			const RawSection& section = rawCoff->sections[i];
 			if (section.header.PointerToRawData != 0u)
 			{
-				outputData.Insert(section.data, section.header.SizeOfRawData);
+				outputData.Write(section.data, section.header.SizeOfRawData);
 			}
 
 			const size_t relocationCount = section.relocations.size();
@@ -722,7 +722,7 @@ namespace coff
 				// the actual number of relocations.
 				IMAGE_RELOCATION dummyRelocation = {};
 				dummyRelocation.RelocCount = static_cast<DWORD>(relocationCount);
-				outputData.Insert(dummyRelocation);
+				outputData.Write(dummyRelocation);
 			}
 
 			if (removalStrategy == SymbolRemovalStrategy::Enum::MSVC_COMPATIBLE)
@@ -730,7 +730,7 @@ namespace coff
 				for (size_t j = 0u; j < relocationCount; ++j)
 				{
 					const IMAGE_RELOCATION& relocation = section.relocations[j];
-					outputData.Insert(relocation);
+					outputData.Write(relocation);
 				}
 			}
 			else if (removalStrategy == SymbolRemovalStrategy::Enum::LLD_COMPATIBLE)
@@ -756,13 +756,13 @@ namespace coff
 						}
 					}
 
-					outputData.Insert(relocation);
+					outputData.Write(relocation);
 				}
 			}
 
 			for (size_t j = 0u; j < section.lineNumbers.size(); ++j)
 			{
-				outputData.Insert(section.lineNumbers[j]);
+				outputData.Write(section.lineNumbers[j]);
 			}
 		}
 
diff --git a/Engine/Source/Developer/Windows/LiveCodingServer/Private/External/LC_LiveModule.cpp b/Engine/Source/Developer/Windows/LiveCodingServer/Private/External/LC_LiveModule.cpp
--- a/Engine/Source/Developer/Windows/LiveCodingServer/Private/External/LC_LiveModule.cpp
+++ b/Engine/Source/Developer/Windows/LiveCodingServer/Private/External/LC_LiveModule.cpp
@@ -338,28 +338,17 @@ namespace
 
 			logging::LogNoFormat<logging::Channel::USER>(compilerOutput);
 
-			if (updateType != LiveModule::UpdateType::NO_CLIENT_COMMUNICATION)
+			const size_t outputLength = processContext->stdoutData.length();
+			if (outputLength != 0u)
 			{
-				for (size_t p = 0u; p < processCount; ++p)
+				if (updateType != LiveModule::UpdateType::NO_CLIENT_COMMUNICATION)
 				{
-					const DuplexPipe* pipe = processData[p].liveProcess->GetPipe();
-
-					size_t sentAlready = 0u;
-					for (;;)
+					for (size_t p = 0u; p < processCount; ++p)
 					{
-						const size_t remainingOutput = processContext->stdoutData.length() - sentAlready;
-						const size_t toSend = remainingOutput > (commands::LogOutput::BUFFER_SIZE - 1u) ? (commands::LogOutput::BUFFER_SIZE - 1u) : remainingOutput;
-
-						commands::LogOutput cmd { toSend };
-						memcpy(cmd.buffer, compilerOutput + sentAlready, toSend * sizeof(wchar_t));
-						cmd.buffer[toSend] = L'\0';
-						pipe->SendCommandAndWaitForAck(cmd);
+						const DuplexPipe* pipe = processData[p].liveProcess->GetPipe();
 
-						sentAlready += toSend;
-						if (sentAlready >= processContext->stdoutData.length())
-						{
-							break;
-						}
+						commands::LogOutput cmd {};
+						pipe->SendCommandAndWaitForAck(cmd, compilerOutput, (outputLength + 1u) * sizeof(wchar_t));
 					}
 				}
 			}
@@ -551,8 +540,22 @@ namespace
 			// this is a relocation to the symbol in question
 			if (!relocations::WouldPatchRelocation(relocation, coffDb, srcSymbolName, findData))
 			{
-				// this relocation to the symbol would not be patched by us, hence we
-				// absolutely need this symbol
+				// this relocation to the symbol would not be patched by us, hence we probably need this symbol.
+				// however, there are special cases where we want to strip symbols even though we might not patch
+				// all relocations to it.
+
+				// special case #1: a relocation from an exception-related unwind symbol (?dtor$) to a
+				// dynamic initializer (??__E), e.g. a relocation from ?dtor$0@?0???__ESomeGlobalVariable@@YAXXZ@4HA
+				// ("int `void __cdecl `dynamic initializer for 'SomeGlobalVariable''(void)'::`1'::dtor$0")
+				// in this case, the relocation is not patched, but the dynamic initializer refers to an already
+				// existing symbol. that dynamic initializer will be removed by us later on anyway, hence
+				// those relocations do not really need patching.
+				if (symbols::IsExceptionUnwindSymbolForDynamicInitializer(srcSymbolName))
+				{
+					LC_LOG_DEV("Ignoring unpatched relocation from symbol %s", srcSymbolName.c_str());
+					continue;
+				}
+
 				return nullptr;
 			}
 		}
@@ -822,7 +825,7 @@ namespace
 				const DuplexPipe* pipe = processData.pipe;
 
 				LC_LOG_USER("Calling compile start hooks (PID: %d)", pid);
-				pipe->SendCommandAndWaitForAck(commands::CallHooks { hook::MakeFunction(moduleBase, hookData.firstRva), hook::MakeFunction(moduleBase, hookData.lastRva) });
+				pipe->SendCommandAndWaitForAck(commands::CallHooks { hook::MakeFunction(moduleBase, hookData.firstRva), hook::MakeFunction(moduleBase, hookData.lastRva) }, nullptr, 0u);
 			}
 		}
 	}
@@ -849,7 +852,7 @@ namespace
 				const DuplexPipe* pipe = processData.pipe;
 
 				LC_LOG_USER("Calling compile success hooks (PID: %d)", pid);				
-				pipe->SendCommandAndWaitForAck(commands::CallHooks { hook::MakeFunction(moduleBase, hookData.firstRva), hook::MakeFunction(moduleBase, hookData.lastRva) });
+				pipe->SendCommandAndWaitForAck(commands::CallHooks { hook::MakeFunction(moduleBase, hookData.firstRva), hook::MakeFunction(moduleBase, hookData.lastRva) }, nullptr, 0u);
 			}
 		}
 	}
@@ -876,7 +879,7 @@ namespace
 				const DuplexPipe* pipe = processData.pipe;
 
 				LC_LOG_USER("Calling compile error hooks (PID: %d)", pid);
-				pipe->SendCommandAndWaitForAck(commands::CallHooks { hook::MakeFunction(moduleBase, hookData.firstRva), hook::MakeFunction(moduleBase, hookData.lastRva) });
+				pipe->SendCommandAndWaitForAck(commands::CallHooks { hook::MakeFunction(moduleBase, hookData.firstRva), hook::MakeFunction(moduleBase, hookData.lastRva) }, nullptr, 0u);
 			}
 		}
 	}
@@ -1248,7 +1251,7 @@ void LiveModule::Unload(void)
 			}
 
 			const DuplexPipe* clientPipe = process.pipe;
-			clientPipe->SendCommandAndWaitForAck(commands::UnloadPatch { static_cast<HMODULE>(process.moduleBase) });
+			clientPipe->SendCommandAndWaitForAck(commands::UnloadPatch { static_cast<HMODULE>(process.moduleBase) }, nullptr, 0u);
 		}
 	}
 }
@@ -2165,7 +2168,7 @@ LiveModule::ErrorType::Enum LiveModule::Update(FileAttributeCache* fileCache, Di
 		if (count > 0u)
 		{
 			executable::Image* image = executable::OpenImage(m_moduleName.c_str(), file::OpenMode::READ_ONLY);
-			executable::ImageSectionDB* imageSections = executable::GatherSections(image);
+			executable::ImageSectionDB* imageSections = executable::GatherImageSectionDB(image);
 
 			// load and cache all .obj not in the cache yet concurrently
 			{
@@ -3031,28 +3034,17 @@ LiveModule::ErrorType::Enum LiveModule::Update(FileAttributeCache* fileCache, Di
 	{
 		logging::LogNoFormat<logging::Channel::USER>(linkerOutput);
 
-		if (updateType != LiveModule::UpdateType::NO_CLIENT_COMMUNICATION)
+		const size_t outputLength = linkerProcessContext->stdoutData.length();
+		if (outputLength != 0u)
 		{
-			for (size_t p = 0u; p < processCount; ++p)
+			if (updateType != LiveModule::UpdateType::NO_CLIENT_COMMUNICATION)
 			{
-				const DuplexPipe* pipe = processData[p].liveProcess->GetPipe();
-
-				size_t sentAlready = 0u;
-				for (;;)
+				for (size_t p = 0u; p < processCount; ++p)
 				{
-					const size_t remainingOutput = linkerProcessContext->stdoutData.length() - sentAlready;
-					const size_t toSend = remainingOutput > (commands::LogOutput::BUFFER_SIZE - 1u) ? (commands::LogOutput::BUFFER_SIZE - 1u) : remainingOutput;
-
-					commands::LogOutput cmd { toSend };
-					memcpy(cmd.buffer, linkerOutput + sentAlready, toSend * sizeof(wchar_t));
-					cmd.buffer[toSend] = L'\0';
-					pipe->SendCommandAndWaitForAck(cmd);
+					const DuplexPipe* pipe = processData[p].liveProcess->GetPipe();
 
-					sentAlready += toSend;
-					if (sentAlready >= linkerProcessContext->stdoutData.length())
-					{
-						break;
-					}
+					commands::LogOutput cmd {};
+					pipe->SendCommandAndWaitForAck(cmd, linkerOutput, (outputLength + 1u) * sizeof(wchar_t));
 				}
 			}
 		}
@@ -3103,7 +3095,7 @@ LiveModule::ErrorType::Enum LiveModule::Update(FileAttributeCache* fileCache, Di
 		return ErrorType::LOAD_PATCH_ERROR;
 	}
 
-	executable::ImageSectionDB* imageSections = executable::GatherSections(image);
+	executable::ImageSectionDB* imageSections = executable::GatherImageSectionDB(image);
 
 	// before loading the DLL, disable its entry point so we can load it without initializing anything.
 	// we first want to reconstruct symbol information and patch dynamic initializers, only then do
@@ -3179,11 +3171,11 @@ LiveModule::ErrorType::Enum LiveModule::Update(FileAttributeCache* fileCache, Di
 			process::Resume(data.liveProcess->GetProcessHandle());
 #endif
 
-			data.liveProcess->GetPipe()->SendCommandAndWaitForAck(cmd);
+			data.liveProcess->GetPipe()->SendCommandAndWaitForAck(cmd, nullptr, 0u);
 
 			// receive command with patch info
 			CommandMap commandMap;
-			commandMap.RegisterAction<LoadPatchInfoAction>();
+			commandMap.RegisterAction<actions::LoadPatchInfo>();
 			commandMap.HandleCommands(data.liveProcess->GetPipe(), &loadedPatches);
 		}
 	}
@@ -3224,7 +3216,7 @@ LiveModule::ErrorType::Enum LiveModule::Update(FileAttributeCache* fileCache, Di
 		for (size_t i = 0u; i < processCount; ++i)
 		{
 			const DuplexPipe* clientPipe = processData[i].liveProcess->GetPipe();
-			clientPipe->SendCommandAndWaitForAck(commands::UnloadPatch { static_cast<HMODULE>(loadedPatches[i]) });
+			clientPipe->SendCommandAndWaitForAck(commands::UnloadPatch { static_cast<HMODULE>(loadedPatches[i]) }, nullptr, 0u);
 		}
 
 		// clear the set for the next update
@@ -3244,7 +3236,7 @@ LiveModule::ErrorType::Enum LiveModule::Update(FileAttributeCache* fileCache, Di
 		for (size_t p = 0u; p < processCount; ++p)
 		{
 			const PerProcessData& data = processData[p];
-			data.liveProcess->GetPipe()->SendCommandAndWaitForAck(commands::EnterSyncPoint{});
+			data.liveProcess->GetPipe()->SendCommandAndWaitForAck(commands::EnterSyncPoint {}, nullptr, 0u);
 		}
 	}
 
@@ -3409,7 +3401,7 @@ LiveModule::ErrorType::Enum LiveModule::Update(FileAttributeCache* fileCache, Di
 			LC_LOG_INDENT_DEV;
 
 			executable::Image* originalImage = executable::OpenImage(m_moduleName.c_str(), file::OpenMode::READ_ONLY);
-			executable::ImageSectionDB* originalImageSections = executable::GatherSections(originalImage);
+			executable::ImageSectionDB* originalImageSections = executable::GatherImageSectionDB(originalImage);
 
 			types::StringSet noSymbolsToIgnore;
 
@@ -3455,7 +3447,7 @@ LiveModule::ErrorType::Enum LiveModule::Update(FileAttributeCache* fileCache, Di
 
 	// reconstruct symbols for all compilands that are part of the new patch executable
 	executable::Image* patchImage = executable::OpenImage(exePath.c_str(), file::OpenMode::READ_ONLY);
-	executable::ImageSectionDB* patchImageSections = executable::GatherSections(patchImage);
+	executable::ImageSectionDB* patchImageSections = executable::GatherImageSectionDB(patchImage);
 
 	// gather the dynamic initializers and remaining symbols by walking the module
 	const symbols::DynamicInitializerDB initializerDb = symbols::GatherDynamicInitializers(patchSymbolProvider, patchImage, patchImageSections, patch_imageSectionDB, patch_contributionDB, patch_compilandDB, m_coffCache, patch_symbolDB);
@@ -3557,7 +3549,7 @@ LiveModule::ErrorType::Enum LiveModule::Update(FileAttributeCache* fileCache, Di
 					const DuplexPipe* pipe = hookProcessData.pipe;
 
 					LC_LOG_USER("Calling pre-patch hooks (PID: %d)", pid);
-					pipe->SendCommandAndWaitForAck(commands::CallHooks { hook::MakeFunction(moduleBase, hookData.firstRva), hook::MakeFunction(moduleBase, hookData.lastRva) });
+					pipe->SendCommandAndWaitForAck(commands::CallHooks { hook::MakeFunction(moduleBase, hookData.firstRva), hook::MakeFunction(moduleBase, hookData.lastRva) }, nullptr, 0u);
 				}
 
 				compiledModulePatch->RegisterPrePatchHooks(hookData.data->index, hookData.firstRva, hookData.lastRva);
@@ -3771,7 +3763,7 @@ LiveModule::ErrorType::Enum LiveModule::Update(FileAttributeCache* fileCache, Di
 		for (size_t p = 0u; p < processCount; ++p)
 		{
 			const PerProcessData& data = processData[p];
-			data.liveProcess->GetPipe()->SendCommandAndWaitForAck(commands::CallEntryPoint { loadedPatches[p], entryPointRva });
+			data.liveProcess->GetPipe()->SendCommandAndWaitForAck(commands::CallEntryPoint { loadedPatches[p], entryPointRva }, nullptr, 0u);
 		}
 
 		// disable entry point in all processes again.
@@ -4120,7 +4112,7 @@ LiveModule::ErrorType::Enum LiveModule::Update(FileAttributeCache* fileCache, Di
 					const DuplexPipe* pipe = hookData.data->processes[p].pipe;
 
 					LC_LOG_USER("Calling post-patch hooks (PID: %d)", pid);
-					pipe->SendCommandAndWaitForAck(commands::CallHooks { hook::MakeFunction(moduleBase, hookData.firstRva), hook::MakeFunction(moduleBase, hookData.lastRva) });
+					pipe->SendCommandAndWaitForAck(commands::CallHooks { hook::MakeFunction(moduleBase, hookData.firstRva), hook::MakeFunction(moduleBase, hookData.lastRva) }, nullptr, 0u);
 				}
 
 				compiledModulePatch->RegisterPostPatchHooks(hookData.data->index, hookData.firstRva, hookData.lastRva);
@@ -4134,7 +4126,7 @@ LiveModule::ErrorType::Enum LiveModule::Update(FileAttributeCache* fileCache, Di
 		for (size_t p = 0u; p < processCount; ++p)
 		{
 			const PerProcessData& data = processData[p];
-			data.liveProcess->GetPipe()->SendCommandAndWaitForAck(commands::LeaveSyncPoint {});
+			data.liveProcess->GetPipe()->SendCommandAndWaitForAck(commands::LeaveSyncPoint {}, nullptr, 0u);
 		}
 	}
 
@@ -4168,6 +4160,12 @@ bool LiveModule::InstallCompiledPatches(LiveProcess* liveProcess, void* original
 		return true;
 	}
 
+	if (m_compiledModulePatches.size() == 0u)
+	{
+		// nothing to install
+		return true;
+	}
+
 	LC_LOG_DEV("LiveModule InstallCompiledPatches: %S", m_moduleName.c_str());
 
 	telemetry::Scope wholeScope("Installing patches");
@@ -4240,11 +4238,11 @@ bool LiveModule::InstallCompiledPatches(LiveProcess* liveProcess, void* original
 			process::Resume(processHandle);
 #endif
 
-			pipe->SendCommandAndWaitForAck(cmd);
+			pipe->SendCommandAndWaitForAck(cmd, nullptr, 0u);
 
 			// receive command with patch info
 			CommandMap commandMap;
-			commandMap.RegisterAction<LoadPatchInfoAction>();
+			commandMap.RegisterAction<actions::LoadPatchInfo>();
 			commandMap.HandleCommands(pipe, &loadedPatches);
 		}
 
@@ -4254,13 +4252,13 @@ bool LiveModule::InstallCompiledPatches(LiveProcess* liveProcess, void* original
 		{
 			LC_ERROR_USER("Patch could not be activated.");
 
-			pipe->SendCommandAndWaitForAck(commands::UnloadPatch { static_cast<HMODULE>(moduleBase) });
+			pipe->SendCommandAndWaitForAck(commands::UnloadPatch { static_cast<HMODULE>(moduleBase) }, nullptr, 0u);
 			return false;
 		}
 
 
 		// enter sync point
-		pipe->SendCommandAndWaitForAck(commands::EnterSyncPoint {});
+		pipe->SendCommandAndWaitForAck(commands::EnterSyncPoint {}, nullptr, 0u);
 
 
 		// store the new databases into the module cache
@@ -4272,7 +4270,7 @@ bool LiveModule::InstallCompiledPatches(LiveProcess* liveProcess, void* original
 		LC_LOG_DEV("Calling pre-patch hooks");
 		{
 			void* hookModule = processModuleBases[patchData.prePatchHookModuleIndex];
-			pipe->SendCommandAndWaitForAck(commands::CallHooks { hook::MakeFunction(hookModule, patchData.firstPrePatchHook), hook::MakeFunction(hookModule, patchData.lastPrePatchHook) });
+			pipe->SendCommandAndWaitForAck(commands::CallHooks { hook::MakeFunction(hookModule, patchData.firstPrePatchHook), hook::MakeFunction(hookModule, patchData.lastPrePatchHook) }, nullptr, 0u);
 		}
 
 
@@ -4316,7 +4314,7 @@ bool LiveModule::InstallCompiledPatches(LiveProcess* liveProcess, void* original
 
 			LC_LOG_DEV("Calling original entry point");
 
-			pipe->SendCommandAndWaitForAck(commands::CallEntryPoint { moduleBase, entryPointRva });
+			pipe->SendCommandAndWaitForAck(commands::CallEntryPoint { moduleBase, entryPointRva }, nullptr, 0u);
 
 			executablePatcher.DisableEntryPoint(processHandle, moduleBase, entryPointRva);
 		}
@@ -4357,15 +4355,15 @@ bool LiveModule::InstallCompiledPatches(LiveProcess* liveProcess, void* original
 		LC_LOG_DEV("Calling post-patch hooks");
 		{
 			void* hookModule = processModuleBases[patchData.postPatchHookModuleIndex];
-			pipe->SendCommandAndWaitForAck(commands::CallHooks { hook::MakeFunction(hookModule, patchData.firstPostPatchHook), hook::MakeFunction(hookModule, patchData.lastPostPatchHook) });
+			pipe->SendCommandAndWaitForAck(commands::CallHooks { hook::MakeFunction(hookModule, patchData.firstPostPatchHook), hook::MakeFunction(hookModule, patchData.lastPostPatchHook) }, nullptr, 0u);
 		}
 
 
 		// leave sync point
-		pipe->SendCommandAndWaitForAck(commands::LeaveSyncPoint {});
+		pipe->SendCommandAndWaitForAck(commands::LeaveSyncPoint {}, nullptr, 0u);
 	}
 
-	LC_SUCCESS_USER("Successfully installed patches (%.3fs)", wholeScope.ReadSeconds());
+	LC_SUCCESS_USER("Successfully installed %zu patches (%.3fs)", m_compiledModulePatches.size(), wholeScope.ReadSeconds());
 
 	return true;
 }
@@ -4401,7 +4399,7 @@ bool LiveModule::HasInstalledPatches(void) const
 }
 
 
-bool LiveModule::LoadPatchInfoAction::Execute(CommandType* command, const DuplexPipe* pipe, void* context)
+bool LiveModule::actions::LoadPatchInfo::Execute(const CommandType* command, const DuplexPipe* pipe, void* context, const void*, size_t)
 {
 	types::vector<void*>* loadedPatches = static_cast<types::vector<void*>*>(context);
 	loadedPatches->emplace_back(command->module);
diff --git a/Engine/Source/Developer/Windows/LiveCodingServer/Private/External/LC_LiveModule.h b/Engine/Source/Developer/Windows/LiveCodingServer/Private/External/LC_LiveModule.h
--- a/Engine/Source/Developer/Windows/LiveCodingServer/Private/External/LC_LiveModule.h
+++ b/Engine/Source/Developer/Windows/LiveCodingServer/Private/External/LC_LiveModule.h
@@ -99,12 +99,21 @@ private:
 	void OnCompiledFile(const symbols::ObjPath& objPath, symbols::Compiland* compiland, const CompileResult& compileResult, double compileTime, bool forceAmalgamationPartsLinkage);
 
 	// actions
-	struct LoadPatchInfoAction
+	struct actions
 	{
-		typedef commands::LoadPatchInfo CommandType;
-		static bool Execute(CommandType* command, const DuplexPipe* pipe, void* context);
+		#define DECLARE_ACTION(_name)																													\
+			struct _name																																\
+			{																																			\
+				typedef ::commands::_name CommandType;																									\
+				static bool Execute(const CommandType* command, const DuplexPipe* pipe, void* context, const void* payload, size_t payloadSize);		\
+			}
+
+		DECLARE_ACTION(LoadPatchInfo);
+
+		#undef DECLARE_ACTION
 	};
 
+
 	std::wstring m_moduleName;
 	executable::Header m_imageHeader;
 	RunMode::Enum m_runMode;
diff --git a/Engine/Source/Developer/Windows/LiveCodingServer/Private/External/LC_LiveProcess.cpp b/Engine/Source/Developer/Windows/LiveCodingServer/Private/External/LC_LiveProcess.cpp
--- a/Engine/Source/Developer/Windows/LiveCodingServer/Private/External/LC_LiveProcess.cpp
+++ b/Engine/Source/Developer/Windows/LiveCodingServer/Private/External/LC_LiveProcess.cpp
@@ -5,10 +5,11 @@
 #include "LC_CodeCave.h"
 
 
-LiveProcess::LiveProcess(process::Handle processHandle, unsigned int processId, unsigned int commandThreadId, const DuplexPipe* pipe)
+LiveProcess::LiveProcess(process::Handle processHandle, unsigned int processId, unsigned int commandThreadId, const void* jumpToSelf, const DuplexPipe* pipe)
 	: m_processHandle(processHandle)
 	, m_processId(processId)
 	, m_commandThreadId(commandThreadId)
+	, m_jumpToSelf(jumpToSelf)
 	, m_pipe(pipe)
 	, m_imagesTriedToLoad()
 	, m_heartBeatDelta(0ull)
@@ -41,7 +42,7 @@ bool LiveProcess::MadeProgress(void) const
 
 void LiveProcess::InstallCodeCave(void)
 {
-	m_codeCave = new CodeCave(m_processHandle, m_processId, m_commandThreadId);
+	m_codeCave = new CodeCave(m_processHandle, m_processId, m_commandThreadId, m_jumpToSelf);
 	m_codeCave->Install();
 }
 
diff --git a/Engine/Source/Developer/Windows/LiveCodingServer/Private/External/LC_LiveProcess.h b/Engine/Source/Developer/Windows/LiveCodingServer/Private/External/LC_LiveProcess.h
--- a/Engine/Source/Developer/Windows/LiveCodingServer/Private/External/LC_LiveProcess.h
+++ b/Engine/Source/Developer/Windows/LiveCodingServer/Private/External/LC_LiveProcess.h
@@ -14,7 +14,7 @@ class CodeCave;
 class LiveProcess
 {
 public:
-	LiveProcess(process::Handle processHandle, unsigned int processId, unsigned int commandThreadId, const DuplexPipe* pipe);
+	LiveProcess(process::Handle processHandle, unsigned int processId, unsigned int commandThreadId, const void* jumpToSelf, const DuplexPipe* pipe);
 
 
 	void ReadHeartBeatDelta(const wchar_t* const processGroupName);
@@ -47,6 +47,11 @@ public:
 		return m_commandThreadId;
 	}
 
+	inline const void* GetJumpToSelf(void) const
+	{
+		return m_jumpToSelf;
+	}
+
 	inline const DuplexPipe* GetPipe(void) const
 	{
 		return m_pipe;
@@ -75,6 +80,7 @@ private:
 	process::Handle m_processHandle;
 	unsigned int m_processId;
 	unsigned int m_commandThreadId;
+	const void* m_jumpToSelf;
 	const DuplexPipe* m_pipe;
 
 	// BEGIN EPIC MOD - Add build arguments
diff --git a/Engine/Source/Developer/Windows/LiveCodingServer/Private/External/LC_MemoryBlock.cpp b/Engine/Source/Developer/Windows/LiveCodingServer/Private/External/LC_MemoryBlock.cpp
deleted file mode 100644
diff --git a/Engine/Source/Developer/Windows/LiveCodingServer/Private/External/LC_MemoryBlock.h b/Engine/Source/Developer/Windows/LiveCodingServer/Private/External/LC_MemoryBlock.h
deleted file mode 100644
diff --git a/Engine/Source/Developer/Windows/LiveCodingServer/Private/External/LC_SchedulerWorkerThread.cpp b/Engine/Source/Developer/Windows/LiveCodingServer/Private/External/LC_SchedulerWorkerThread.cpp
--- a/Engine/Source/Developer/Windows/LiveCodingServer/Private/External/LC_SchedulerWorkerThread.cpp
+++ b/Engine/Source/Developer/Windows/LiveCodingServer/Private/External/LC_SchedulerWorkerThread.cpp
@@ -5,24 +5,10 @@
 #include "LC_SchedulerTask.h"
 
 
-namespace
-{
-	struct ThreadContext
-	{
-		scheduler::WorkerThread* instance;
-		scheduler::TaskQueue* queue;
-	};
-}
-
-
 scheduler::WorkerThread::WorkerThread(TaskQueue* queue)
 	: m_thread()
 {
-	ThreadContext* context = new ThreadContext;
-	context->instance = this;
-	context->queue = queue;
-	
-	m_thread = thread::Create(128u * 1024u, &ThreadProxy, context);
+	m_thread = thread::Create("Live coding worker", 128u * 1024u, &scheduler::WorkerThread::ThreadFunction, this, queue);
 }
 
 
@@ -32,19 +18,6 @@ scheduler::WorkerThread::~WorkerThread(void)
 }
 
 
-unsigned int __stdcall scheduler::WorkerThread::ThreadProxy(void* context)
-{
-	thread::SetName("Live coding worker");
-
-	ThreadContext* realContext = static_cast<ThreadContext*>(context);
-	const unsigned int exitCode = realContext->instance->ThreadFunction(realContext->queue);
-
-	delete realContext;
-
-	return exitCode;
-}
-
-
 unsigned int scheduler::WorkerThread::ThreadFunction(TaskQueue* queue)
 {
 	for (;;)
diff --git a/Engine/Source/Developer/Windows/LiveCodingServer/Private/External/LC_SchedulerWorkerThread.h b/Engine/Source/Developer/Windows/LiveCodingServer/Private/External/LC_SchedulerWorkerThread.h
--- a/Engine/Source/Developer/Windows/LiveCodingServer/Private/External/LC_SchedulerWorkerThread.h
+++ b/Engine/Source/Developer/Windows/LiveCodingServer/Private/External/LC_SchedulerWorkerThread.h
@@ -18,7 +18,6 @@ namespace scheduler
 		~WorkerThread(void);
 
 	private:
-		static unsigned int __stdcall ThreadProxy(void* context);
 		unsigned int ThreadFunction(TaskQueue* queue);
 
 		thread::Handle m_thread;
diff --git a/Engine/Source/Developer/Windows/LiveCodingServer/Private/External/LC_ServerCommandThread.cpp b/Engine/Source/Developer/Windows/LiveCodingServer/Private/External/LC_ServerCommandThread.cpp
--- a/Engine/Source/Developer/Windows/LiveCodingServer/Private/External/LC_ServerCommandThread.cpp
+++ b/Engine/Source/Developer/Windows/LiveCodingServer/Private/External/LC_ServerCommandThread.cpp
@@ -23,6 +23,7 @@
 #include "LC_AppSettings.h"
 #include "LC_Allocators.h"
 #include "LC_DuplexPipeClient.h"
+#include "LC_MemoryStream.h"
 #include <mmsystem.h>
 
 // unreachable code
@@ -107,13 +108,11 @@ ServerCommandThread::ServerCommandThread(MainFrame* mainFrame, const wchar_t* co
 	, m_directoryCache(new DirectoryCache(2048u))
 	, m_connectionCS()
 	, m_commandThreads()
-	, m_moduleBatchScope("Module loading")
-	, m_loadedCompilandCountInBatchScope(0u)
 	, m_manualRecompileTriggered(false)
 	, m_liveModuleToModifiedOrNewObjFiles()
 {
-	m_serverThread = thread::Create(64u * 1024u, &ServerThreadProxy, this);
-	m_compileThread = thread::Create(64u * 1024u, &CompileThreadProxy, this);
+	m_serverThread = thread::Create("Live coding server", 64u * 1024u, &ServerCommandThread::ServerThread, this);
+	m_compileThread = thread::Create("Live coding compilation", 64u * 1024u, &ServerCommandThread::CompileThread, this);
 
 	m_liveModules.reserve(256u);
 	m_liveProcesses.reserve(8u);
@@ -147,15 +146,16 @@ std::wstring ServerCommandThread::GetProcessImagePath(void) const
 }
 
 
-void ServerCommandThread::LoadModule(const wchar_t* givenModulePath, const DuplexPipe* pipe, TaskContext* tasks, unsigned int processId)
+scheduler::Task<LiveModule*>* ServerCommandThread::LoadModule(unsigned int processId, void* moduleBase, const wchar_t* givenModulePath, scheduler::TaskBase* taskRoot)
 {
-	// EPIC REMOVED: g_theApp.GetMainFrame()->ChangeStatusBarText(L"Loading modules...");
-
+	// note that the path we get from the client might not be normalized, depending on how the executable was launched.
+	// it is crucial to normalize the path again, otherwise we could load already loaded modules into the same
+	// Live++ instance, which would wreak havoc
 	const std::wstring& modulePath = file::NormalizePath(givenModulePath);
 	const executable::Header imageHeader = GetImageHeader(modulePath.c_str());
 	if (!executable::IsValidHeader(imageHeader))
 	{
-		return;
+		return nullptr;
 	}
 
 	LiveProcess* liveProcess = FindProcessById(processId);
@@ -164,157 +164,97 @@ void ServerCommandThread::LoadModule(const wchar_t* givenModulePath, const Duple
 	if (liveProcess->TriedToLoadImage(imageHeader))
 	{
 		// tried loading this module into this process already
-		return;
+		return nullptr;
 	}
 
+	// find any other process ID that tried to load this module already
 	{
-		CommandMap commandMap;
-		commandMap.RegisterAction<GetModuleInfoAction>();
-
-		// defer loading of the module to make sure that we get the correct module base address,
-		// no matter if .exe or .dll.
+		const size_t count = m_liveProcesses.size();
+		for (size_t i = 0u; i < count; ++i)
 		{
-			commands::GetModule cmd = {};
-			cmd.loadImports = false;
-			cmd.taskContext = tasks;
-			wcscpy_s(cmd.path, modulePath.c_str());
-			pipe->SendCommandAndWaitForAck(cmd);
-		}
-
-		// handle commands that return module info
-		commandMap.HandleCommands(pipe, this);
-	}
-
-	liveProcess->AddLoadedImage(imageHeader);
-}
-
+			LiveProcess* otherLiveProcess = m_liveProcesses[i];
+			if (otherLiveProcess->TriedToLoadImage(imageHeader))
+			{
+				// some *other* process loaded this module already
+				LC_LOG_USER("Registering module %S (PID: %d)", modulePath.c_str(), processId);
 
-void ServerCommandThread::LoadAllModules(const wchar_t* givenModulePath, const DuplexPipe* pipe, TaskContext* tasks, unsigned int processId)
-{
-	// EPIC REMOVED: g_theApp.GetMainFrame()->ChangeStatusBarText(L"Loading modules...");
+				LiveModule* liveModule = m_imageHeaderToLiveModule[imageHeader];
+				if (liveModule)
+				{
+					liveModule->RegisterProcess(liveProcess, moduleBase, modulePath);
+					liveModule->DisableControlFlowGuard(liveProcess, moduleBase);
 
-	const std::wstring& modulePath = file::NormalizePath(givenModulePath);
-	const executable::Header imageHeader = GetImageHeader(modulePath.c_str());
-	if (!executable::IsValidHeader(imageHeader))
-	{
-		return;
-	}
+					const bool installedPatchesSuccessfully = liveModule->InstallCompiledPatches(liveProcess, moduleBase);
+					if (!installedPatchesSuccessfully)
+					{
+						LC_ERROR_USER("Compiled patches could not be installed (PID: %d)", processId);
+						liveModule->UnregisterProcess(liveProcess);
+					}
 
-	LiveProcess* liveProcess = FindProcessById(processId);
-	LC_ASSERT(liveProcess, "Invalid process ID.");
+					liveProcess->AddLoadedImage(imageHeader);
+				}
 
-	if (liveProcess->TriedToLoadImage(imageHeader))
-	{
-		// tried loading this module into this process already
-		return;
+				return nullptr;
+			}
+		}
 	}
 
-	symbols::Provider* provider = symbols::OpenEXE(modulePath.c_str(), symbols::OpenOptions::ACCUMULATE_SIZE);
-	if (!provider)
+	symbols::Provider* moduleProvider = symbols::OpenEXE(modulePath.c_str(), symbols::OpenOptions::ACCUMULATE_SIZE);
+	if (!moduleProvider)
 	{
-		liveProcess->AddLoadedImage(imageHeader);
-		return;
+		return nullptr;
 	}
 
-	// grab DIA compilands first. this is very fast, and needed in order to gather modules next
-	symbols::DiaCompilandDB* diaCompilandDb = symbols::GatherDiaCompilands(provider);
-	symbols::ModuleDB* moduleDB = symbols::GatherModules(diaCompilandDb);
+	liveProcess->AddLoadedImage(imageHeader);
 
-	// now that we have a list of modules, load them all concurrently, starting with the main executable, followed
-	// by all DLLs.
+	// accumulate module info
 	{
-		CommandMap commandMap;
-		commandMap.RegisterAction<GetModuleInfoAction>();
-		{
-			commands::GetModule cmd = {};
-			cmd.loadImports = false;
-			cmd.taskContext = tasks;
-			wcscpy_s(cmd.path, modulePath.c_str());
-			pipe->SendCommandAndWaitForAck(cmd);
-		}
-
-		commandMap.HandleCommands(pipe, this);
-
-		const size_t count = moduleDB->modules.size();
-		for (size_t i = 0u; i < count; ++i)
-		{
-			const std::wstring& path = moduleDB->modules[i];
+		const file::Attributes attributes = file::GetAttributes(modulePath.c_str());
+		const uint64_t size = file::GetSize(attributes);
 
-			// all we have is a relative path to the DLL. get the full path from the modules loaded into the main process
-			{
-				// because DLLs might also have import DLLs, load all those as well
-				commands::GetModule cmd = {};
-				cmd.loadImports = true;
-				cmd.taskContext = tasks;
-				wcscpy_s(cmd.path, path.c_str());
-				pipe->SendCommandAndWaitForAck(cmd);
-			}
+		g_loadedModuleSize.Accumulate(size);
+		g_loadedModuleSize.Print();
+		g_loadedModuleSize.ResetCurrent();
 
-			// handle commands that return module info
-			commandMap.HandleCommands(pipe, this);
-		}
+		LC_LOG_USER("Loading module %S (%.3f MB)", modulePath.c_str(), size / 1048576.0f);
 	}
 
-	symbols::DestroyDiaCompilandDB(diaCompilandDb);
-	symbols::DestroyModuleDB(moduleDB);
-
-	symbols::Close(provider);
-
-	liveProcess->AddLoadedImage(imageHeader);
-}
+	// create a task to load the module of this batch concurrently
+	LiveModule* liveModule = new LiveModule(modulePath.c_str(), imageHeader, m_runMode);
+	m_imageHeaderToLiveModule.emplace(imageHeader, liveModule);
 
-
-void ServerCommandThread::UnloadModule(const wchar_t* givenModulePath, const DuplexPipe* pipe, unsigned int processId)
-{
-	// EPIC REMOVED: g_theApp.GetMainFrame()->ChangeStatusBarText(L"Unloading modules...");
-
-	const std::wstring& modulePath = file::NormalizePath(givenModulePath);
-	const executable::Header imageHeader = GetImageHeader(modulePath.c_str());
-	if (!executable::IsValidHeader(imageHeader))
+	auto task = scheduler::CreateTask(taskRoot, [liveModule, liveProcess, modulePath, moduleBase, moduleProvider]()
 	{
-		return;
-	}
+		telemetry::Scope scope("Loading module");
 
-	LiveProcess* liveProcess = FindProcessById(processId);
-	LC_ASSERT(liveProcess, "Invalid process ID.");
+		symbols::DiaCompilandDB* moduleDiaCompilandDb = symbols::GatherDiaCompilands(moduleProvider);
 
-	if (!liveProcess->TriedToLoadImage(imageHeader))
-	{
-		// this module was never loaded
-		return;
-	}
+		liveModule->Load(moduleProvider, moduleDiaCompilandDb);
+		liveModule->RegisterProcess(liveProcess, moduleBase, modulePath);
+		liveModule->DisableControlFlowGuard(liveProcess, moduleBase);
 
-	{
-		CommandMap commandMap;
-		commandMap.RegisterAction<GetModuleInfoAction>();
+		symbols::DestroyDiaCompilandDB(moduleDiaCompilandDb);
+		symbols::Close(moduleProvider);
 
-		// defer unloading of the module to make sure that we get the correct module base address,
-		// no matter if .exe or .dll.
-		{
-			commands::GetModule cmd = {};
-			cmd.loadImports = false;
-			cmd.taskContext = nullptr;
-			wcscpy_s(cmd.path, modulePath.c_str());
-			pipe->SendCommandAndWaitForAck(cmd);
-		}
+		return liveModule;
+	});
 
-		// handle commands that return module info
-		commandMap.HandleCommands(pipe, this);
-	}
+	scheduler::RunTask(task);
 
-	liveProcess->RemoveLoadedImage(imageHeader);
+	return task;
 }
 
 
-void ServerCommandThread::UnloadAllModules(const wchar_t* givenModulePath, const DuplexPipe* pipe, unsigned int processId)
+bool ServerCommandThread::UnloadModule(unsigned int processId, const wchar_t* givenModulePath)
 {
-	// EPIC REMOVED: g_theApp.GetMainFrame()->ChangeStatusBarText(L"Unloading modules...");
-
+	// note that the path we get from the client might not be normalized, depending on how the executable was launched.
+	// it is crucial to normalize the path again, otherwise we could load already loaded modules into the same
+	// Live++ instance, which would wreak havoc
 	const std::wstring& modulePath = file::NormalizePath(givenModulePath);
 	const executable::Header imageHeader = GetImageHeader(modulePath.c_str());
 	if (!executable::IsValidHeader(imageHeader))
 	{
-		return;
+		return false;
 	}
 
 	LiveProcess* liveProcess = FindProcessById(processId);
@@ -323,61 +263,33 @@ void ServerCommandThread::UnloadAllModules(const wchar_t* givenModulePath, const
 	if (!liveProcess->TriedToLoadImage(imageHeader))
 	{
 		// this module was never loaded
-		return;
+		return false;
 	}
 
-	symbols::Provider* provider = symbols::OpenEXE(modulePath.c_str(), symbols::OpenOptions::ACCUMULATE_SIZE);
-	if (!provider)
-	{
-		liveProcess->RemoveLoadedImage(imageHeader);
-		return;
-	}
+	LC_LOG_USER("Unloading module %S", modulePath.c_str());
 
-	// grab DIA compilands first. this is very fast, and needed in order to gather modules next
-	symbols::DiaCompilandDB* diaCompilandDb = symbols::GatherDiaCompilands(provider);
-	symbols::ModuleDB* moduleDB = symbols::GatherModules(diaCompilandDb);
+	liveProcess->RemoveLoadedImage(imageHeader);
+	m_imageHeaderToLiveModule.erase(imageHeader);
 
-	// now that we have a list of modules, load them all concurrently, starting with the main executable, followed
-	// by all DLLs.
+	for (auto it = m_liveModules.begin(); it != m_liveModules.end(); /* nothing */)
 	{
-		CommandMap commandMap;
-		commandMap.RegisterAction<GetModuleInfoAction>();
+		LiveModule* liveModule = *it;
+		if (std::equal_to<executable::Header>()(liveModule->GetImageHeader(), imageHeader))
 		{
-			commands::GetModule cmd = {};
-			cmd.loadImports = false;
-			cmd.taskContext = nullptr;
-			wcscpy_s(cmd.path, modulePath.c_str());
-			pipe->SendCommandAndWaitForAck(cmd);
-		}
+			liveModule->Unload();
+			delete liveModule;
 
-		commandMap.HandleCommands(pipe, this);
+			it = m_liveModules.erase(it);
 
-		const size_t count = moduleDB->modules.size();
-		for (size_t i = 0u; i < count; ++i)
+			return true;
+		}
+		else
 		{
-			const std::wstring& path = moduleDB->modules[i];
-
-			// all we have is a relative path to the DLL. get the full path from the modules loaded into the main process
-			{
-				// because DLLs might also have import DLLs, load all those as well
-				commands::GetModule cmd = {};
-				cmd.loadImports = true;
-				cmd.taskContext = nullptr;
-				wcscpy_s(cmd.path, path.c_str());
-				pipe->SendCommandAndWaitForAck(cmd);
-			}
-
-			// handle commands that return module info
-			commandMap.HandleCommands(pipe, this);
+			++it;
 		}
 	}
 
-	symbols::DestroyDiaCompilandDB(diaCompilandDb);
-	symbols::DestroyModuleDB(moduleDB);
-
-	symbols::Close(provider);
-
-	liveProcess->RemoveLoadedImage(imageHeader);
+	return false;
 }
 
 
@@ -511,16 +423,10 @@ void ServerCommandThread::PrewarmCompilerEnvironmentCache(void)
 	scheduler::DestroyTasks(tasks);
 	scheduler::DestroyTask(taskRoot);
 
-	LC_SUCCESS_USER("Prewarmed compiler/linker environment cache (%.3fs, %zu)", scope.ReadSeconds(), uniquePaths.size());
-}
-
-
-unsigned int __stdcall ServerCommandThread::ServerThreadProxy(void* context)
-{
-	thread::SetName("Live coding server");
-
-	ServerCommandThread* instance = static_cast<ServerCommandThread*>(context);
-	return instance->ServerThread();
+	if (uniquePaths.size() != 0u)
+	{
+		LC_SUCCESS_USER("Prewarmed compiler/linker environment cache (%.3fs, %zu executables)", scope.ReadSeconds(), uniquePaths.size());
+	}
 }
 
 
@@ -535,8 +441,6 @@ unsigned int ServerCommandThread::ServerThread(void)
 	for (;;)
 	{
 		CommandThreadContext* context = new CommandThreadContext;
-		context->instance = this;
-
 		context->pipe.Create(primitiveNames::Pipe(m_processGroupName).c_str());
 		context->exceptionPipe.Create(primitiveNames::ExceptionPipe(m_processGroupName).c_str());
 
@@ -550,8 +454,8 @@ unsigned int ServerCommandThread::ServerThread(void)
 		context->exceptionPipe.WaitForClient();
 
 		// a new client has connected, open a new thread for communication
-		context->commandThread = thread::Create(64u * 1024u, &CommandThreadProxy, context);
-		context->exceptionCommandThread = thread::Create(64u * 1024u, &ExceptionCommandThreadProxy, context);
+		context->commandThread = thread::Create("Live coding client command communication", 64u * 1024u, &ServerCommandThread::CommandThread, this, &context->pipe, context->readyEvent);
+		context->exceptionCommandThread = thread::Create("Live coding client exception command communication", 64u * 1024u, &ServerCommandThread::ExceptionCommandThread, this, &context->exceptionPipe);
 
 		// register this connection
 		{
@@ -564,14 +468,6 @@ unsigned int ServerCommandThread::ServerThread(void)
 }
 
 
-unsigned int __stdcall ServerCommandThread::CompileThreadProxy(void* context)
-{
-	thread::SetName("Live coding compilation");
-
-	ServerCommandThread* instance = static_cast<ServerCommandThread*>(context);
-	return instance->CompileThread();
-}
-
 // BEGIN EPIC MOD - Focus application windows on patch complete
 BOOL CALLBACK FocusApplicationWindows(HWND WindowHandle, LPARAM Lparam)
 {
@@ -591,7 +487,7 @@ BOOL CALLBACK FocusApplicationWindows(HWND WindowHandle, LPARAM Lparam)
 // END EPIC MOD
 
 // BEGIN EPIC MOD - Support for lazy-loading modules
-bool ServerCommandThread::FinishedLazyLoadingModulesAction::Execute(CommandType* command, const DuplexPipe* pipe, void* context)
+bool ServerCommandThread::actions::FinishedLazyLoadingModules::Execute(const CommandType* command, const DuplexPipe* pipe, void*, const void*, size_t)
 { 
 	pipe->SendAck(); 
 	return false; 
@@ -599,33 +495,11 @@ bool ServerCommandThread::FinishedLazyLoadingModulesAction::Execute(CommandType*
 
 struct ClientProxyThread
 {
-	struct ProxyGetModuleAction
+	struct ProxyEnableModulesFinishedAction
 	{
-		typedef commands::GetModule CommandType;
+		typedef commands::EnableModulesFinished CommandType;
 
-		static bool Execute(CommandType* command, const DuplexPipe* pipe, void* context)
-		{
-			pipe->SendAck();
-
-			LiveProcess* process = static_cast<LiveProcess*>(context);
-
-			commands::GetModuleInfo cmd;
-			cmd.moduleBase = process->GetLazyLoadedModuleBase(command->path);
-			cmd.processId = process->GetProcessId();
-			cmd.loadImports = command->loadImports;
-			cmd.taskContext = command->taskContext;
-			wcscpy_s(cmd.path, command->path);
-			pipe->SendCommandAndWaitForAck(cmd);
-
-			return true;
-		}
-	};
-
-	struct ProxyEnableModuleFinishedAction
-	{
-		typedef commands::EnableModuleFinished CommandType;
-
-		static bool Execute(CommandType* command, const DuplexPipe* pipe, void* context)
+		static bool Execute(CommandType* command, const DuplexPipe* pipe, void*, const void*, size_t)
 		{
 			pipe->SendAck();
 			return false;
@@ -660,22 +534,27 @@ struct ClientProxyThread
 
 	void EntryPoint()
 	{
-		m_pipe->SendCommandAndWaitForAck(commands::EnableModuleBatchBegin());
-		for(const std::wstring& enableModule : m_enableModules)
+		std::vector<commands::ModuleData> modules;
+		modules.resize(m_enableModules.size());
+
+		for (size_t Idx = 0; Idx < m_enableModules.size(); Idx++)
 		{
-			commands::EnableModule enableModuleCommand;
-			enableModuleCommand.processId = m_process->GetProcessId();
-			wcscpy_s(enableModuleCommand.path, enableModule.c_str());
-			enableModuleCommand.token = nullptr;
-			m_pipe->SendCommandAndWaitForAck(enableModuleCommand);
-
-			CommandMap commandMap;
-			commandMap.RegisterAction<ProxyGetModuleAction>();
-			commandMap.RegisterAction<ProxyEnableModuleFinishedAction>();
-			commandMap.HandleCommands(m_pipe, m_process);
+			commands::ModuleData& module = modules[Idx];
+			module.base = m_process->GetLazyLoadedModuleBase(m_enableModules[Idx].c_str());
+			wcscpy_s(module.path, m_enableModules[Idx].c_str());
 		}
-		m_pipe->SendCommandAndWaitForAck(commands::EnableModuleBatchEnd());
-		m_pipe->SendCommandAndWaitForAck(commands::FinishedLazyLoadingModules());
+
+		commands::EnableModules enableModulesCommand;
+		enableModulesCommand.processId = m_process->GetProcessId();
+		enableModulesCommand.moduleCount = m_enableModules.size();
+		enableModulesCommand.token = nullptr;
+		m_pipe->SendCommandAndWaitForAck(enableModulesCommand, modules.data(), modules.size() * sizeof(commands::ModuleData));
+
+		CommandMap commandMap;
+		commandMap.RegisterAction<ProxyEnableModulesFinishedAction>();
+		commandMap.HandleCommands(m_pipe, m_process);
+
+		m_pipe->SendCommandAndWaitForAck(commands::FinishedLazyLoadingModules(), nullptr, 0);
 	}
 };
 // END EPIC MOD
@@ -739,10 +618,8 @@ void ServerCommandThread::CompileChanges(bool didAllProcessesMakeProgress)
 				ClientProxyThread ClientThread(liveProcess, &ClientPipe, LoadModuleFileNames);
 
 				CommandMap commandMap;
-				commandMap.RegisterAction<EnableModuleBatchBeginAction>();
-				commandMap.RegisterAction<EnableModuleBatchEndAction>();
-				commandMap.RegisterAction<EnableModuleAction>();
-				commandMap.RegisterAction<FinishedLazyLoadingModulesAction>();
+				commandMap.RegisterAction<actions::EnableModules>();
+				commandMap.RegisterAction<actions::FinishedLazyLoadingModules>();
 				commandMap.HandleCommands(&ServerPipe, this);
 
 				for (const std::wstring& loadModuleFileName : LoadModuleFileNames)
@@ -1154,50 +1031,35 @@ unsigned int ServerCommandThread::CompileThread(void)
 }
 
 
-unsigned int __stdcall ServerCommandThread::CommandThreadProxy(void* context)
-{
-	thread::SetName("Live coding client command communication");
-
-	CommandThreadContext* realContext = static_cast<CommandThreadContext*>(context);
-	return realContext->instance->CommandThread(&realContext->pipe, realContext->readyEvent);
-}
-
-
 unsigned int ServerCommandThread::CommandThread(DuplexPipeServer* pipe, Event* readyEvent)
 {
 	// handle incoming commands
 	CommandMap commandMap;
-	commandMap.RegisterAction<TriggerRecompileAction>();
-	commandMap.RegisterAction<BuildPatchAction>();
-	commandMap.RegisterAction<ReadyForCompilationAction>();
-	commandMap.RegisterAction<DisconnectClientAction>();
+	commandMap.RegisterAction<actions::TriggerRecompile>();
+	commandMap.RegisterAction<actions::BuildPatch>();
+	commandMap.RegisterAction<actions::ReadyForCompilation>();
+	commandMap.RegisterAction<actions::DisconnectClient>();
+	commandMap.RegisterAction<actions::RegisterProcess>();
+	commandMap.RegisterAction<actions::EnableModules>();
+	commandMap.RegisterAction<actions::DisableModules>();
+	commandMap.RegisterAction<actions::ApplySettingBool>();
+	commandMap.RegisterAction<actions::ApplySettingInt>();
+	commandMap.RegisterAction<actions::ApplySettingString>();
 	// BEGIN EPIC MOD - Adding ShowConsole command
-	commandMap.RegisterAction<ShowConsoleAction>();
+	commandMap.RegisterAction<actions::ShowConsole>();
 	// END EPIC MOD
 	// BEGIN EPIC MOD - Adding SetVisible command
-	commandMap.RegisterAction<SetVisibleAction>();
+	commandMap.RegisterAction<actions::SetVisible>();
 	// END EPIC MOD
 	// BEGIN EPIC MOD - Adding SetActive command
-	commandMap.RegisterAction<SetActiveAction>();
+	commandMap.RegisterAction<actions::SetActive>();
 	// END EPIC MOD
 	// BEGIN EPIC MOD - Adding SetBuildArguments command
-	commandMap.RegisterAction<SetBuildArgumentsAction>();
+	commandMap.RegisterAction<actions::SetBuildArguments>();
 	// END EPIC MOD
 	// BEGIN EPIC MOD - Support for lazy-loading modules
-	commandMap.RegisterAction<EnableLazyLoadedModuleAction>();
+	commandMap.RegisterAction<actions::EnableLazyLoadedModule>();
 	// END EPIC MOD
-	commandMap.RegisterAction<RegisterProcessAction>();
-	commandMap.RegisterAction<EnableModuleBatchBeginAction>();
-	commandMap.RegisterAction<EnableModuleBatchEndAction>();
-	commandMap.RegisterAction<DisableModuleBatchBeginAction>();
-	commandMap.RegisterAction<DisableModuleBatchEndAction>();
-	commandMap.RegisterAction<EnableModuleAction>();
-	commandMap.RegisterAction<EnableAllModulesAction>();
-	commandMap.RegisterAction<DisableModuleAction>();
-	commandMap.RegisterAction<DisableAllModulesAction>();
-	commandMap.RegisterAction<ApplySettingBoolAction>();
-	commandMap.RegisterAction<ApplySettingIntAction>();
-	commandMap.RegisterAction<ApplySettingStringAction>();
 
 	for (;;)
 	{
@@ -1220,7 +1082,7 @@ unsigned int ServerCommandThread::CommandThread(DuplexPipeServer* pipe, Event* r
 		m_handleCommandsEvent.Wait();
 
 		// tell client that compilation has finished
-		pipe->SendCommandAndWaitForAck(commands::CompilationFinished {});
+		pipe->SendCommandAndWaitForAck(commands::CompilationFinished {}, nullptr, 0u);
 	}
 
 	RemoveCommandThread(pipe);
@@ -1228,20 +1090,11 @@ unsigned int ServerCommandThread::CommandThread(DuplexPipeServer* pipe, Event* r
 }
 
 
-unsigned int __stdcall ServerCommandThread::ExceptionCommandThreadProxy(void* context)
-{
-	thread::SetName("Live coding client exception command communication");
-
-	CommandThreadContext* realContext = static_cast<CommandThreadContext*>(context);
-	return realContext->instance->ExceptionCommandThread(&realContext->exceptionPipe);
-}
-
-
 unsigned int ServerCommandThread::ExceptionCommandThread(DuplexPipeServer* exceptionPipe)
 {
 	// handle incoming exception commands
 	CommandMap commandMap;
-	commandMap.RegisterAction<HandleExceptionAction>();
+	commandMap.RegisterAction<actions::HandleException>();
 
 	for (;;)
 	{
@@ -1294,7 +1147,7 @@ LiveProcess* ServerCommandThread::FindProcessById(unsigned int processId)
 }
 
 
-bool ServerCommandThread::TriggerRecompileAction::Execute(CommandType*, const DuplexPipe* pipe, void* context)
+bool ServerCommandThread::actions::TriggerRecompile::Execute(const CommandType*, const DuplexPipe* pipe, void* context, const void*, size_t)
 {
 	ServerCommandThread* commandThread = static_cast<ServerCommandThread*>(context);
 
@@ -1309,30 +1162,22 @@ bool ServerCommandThread::TriggerRecompileAction::Execute(CommandType*, const Du
 }
 
 
-bool ServerCommandThread::BuildPatchAction::Execute(CommandType* command, const DuplexPipe* pipe, void* context)
+bool ServerCommandThread::actions::BuildPatch::Execute(const CommandType* command, const DuplexPipe* pipe, void* context, const void* payload, size_t payloadSize)
 {
+	pipe->SendAck();
+
 	ServerCommandThread* commandThread = static_cast<ServerCommandThread*>(context);
 
 	// protect against accepting this command while compilation is already in progress
 	CriticalSection::ScopedLock lock(&commandThread->m_actionCS);
 
-	pipe->SendAck();
-
-	// receive module names and .obj paths
-	for (unsigned int i = 0u; i < command->count; ++i)
+	memoryStream::Reader payloadStream(payload, payloadSize);
+	for (unsigned int i = 0u; i < command->fileCount; ++i)
 	{
-		uint32_t id = 0u;
-		pipe->ReceiveCommandId(&id);
-		
-		commands::BuildPatchPacket packetCommand = {};
-		pipe->ReceiveCommand(&packetCommand);
-
-		pipe->SendAck();
-
-		LiveModule::ModifiedObjFile modifiedObjFile;
-		modifiedObjFile.objPath = packetCommand.objPath;
-		modifiedObjFile.amalgamatedObjPath = packetCommand.amalgamatedObjPath;
-		commandThread->m_liveModuleToModifiedOrNewObjFiles[packetCommand.moduleName].push_back(modifiedObjFile);
+		const commands::BuildPatch::PatchData patchData = payloadStream.Read<commands::BuildPatch::PatchData>();
+		const LiveModule::ModifiedObjFile modifiedObjFile = { patchData.objPath, patchData.amalgamatedObjPath };
+
+		commandThread->m_liveModuleToModifiedOrNewObjFiles[patchData.moduleName].push_back(modifiedObjFile);
 	}
 
 	commandThread->m_manualRecompileTriggered = true;
@@ -1341,7 +1186,7 @@ bool ServerCommandThread::BuildPatchAction::Execute(CommandType* command, const
 }
 
 
-bool ServerCommandThread::HandleExceptionAction::Execute(CommandType* command, const DuplexPipe* pipe, void* context)
+bool ServerCommandThread::actions::HandleException::Execute(const CommandType* command, const DuplexPipe* pipe, void* context, const void*, size_t)
 {
 	pipe->SendAck();
 
@@ -1357,7 +1202,7 @@ bool ServerCommandThread::HandleExceptionAction::Execute(CommandType* command, c
 	if (!liveProcess)
 	{
 		// signal client we did not handle the exception
-		pipe->SendCommandAndWaitForAck(commands::HandleExceptionFinished { nullptr, nullptr, nullptr, false });
+		pipe->SendCommandAndWaitForAck(commands::HandleExceptionFinished { nullptr, nullptr, nullptr, false }, nullptr, 0u);
 		return true;
 	}
 
@@ -1388,18 +1233,18 @@ bool ServerCommandThread::HandleExceptionAction::Execute(CommandType* command, c
 	{
 		// tell the client that it needs to unwind its stack and continue at the return address
 		const ExceptionHandlerDialog::ParentFrameData& frameData = dialog.GetParentFrameData();
-		pipe->SendCommandAndWaitForAck(commands::HandleExceptionFinished { frameData.returnAddress, frameData.framePointer, frameData.stackPointer, true });
+		pipe->SendCommandAndWaitForAck(commands::HandleExceptionFinished { frameData.returnAddress, frameData.framePointer, frameData.stackPointer, true }, nullptr, 0u);
 		return true;
 	}
 	else if (result == IDC_EXCEPTION_HANDLER_IGNORE)
 	{
 		// tell the client that we ignored the exception
-		pipe->SendCommandAndWaitForAck(commands::HandleExceptionFinished { nullptr, nullptr, nullptr, false });
+		pipe->SendCommandAndWaitForAck(commands::HandleExceptionFinished { nullptr, nullptr, nullptr, false }, nullptr, 0u);
 		return true;
 	}
 
 	// signal client that we handled the exception and there's nothing left to do
-	pipe->SendCommandAndWaitForAck(commands::HandleExceptionFinished { nullptr, nullptr, nullptr, true });
+	pipe->SendCommandAndWaitForAck(commands::HandleExceptionFinished { nullptr, nullptr, nullptr, true }, nullptr, 0u);
 #endif
 	// END EPIC MOD
 
@@ -1407,7 +1252,7 @@ bool ServerCommandThread::HandleExceptionAction::Execute(CommandType* command, c
 }
 
 
-bool ServerCommandThread::ReadyForCompilationAction::Execute(CommandType*, const DuplexPipe* pipe, void*)
+bool ServerCommandThread::actions::ReadyForCompilation::Execute(const CommandType*, const DuplexPipe* pipe, void*, const void*, size_t)
 {
 	pipe->SendAck();
 
@@ -1416,7 +1261,7 @@ bool ServerCommandThread::ReadyForCompilationAction::Execute(CommandType*, const
 }
 
 
-bool ServerCommandThread::DisconnectClientAction::Execute(CommandType*, const DuplexPipe* pipe, void* context)
+bool ServerCommandThread::actions::DisconnectClient::Execute(const CommandType*, const DuplexPipe* pipe, void* context, const void*, size_t)
 {
 	ServerCommandThread* instance = static_cast<ServerCommandThread*>(context);
 
@@ -1440,7 +1285,7 @@ bool ServerCommandThread::DisconnectClientAction::Execute(CommandType*, const Du
 }
 
 // BEGIN EPIC MOD - Adding ShowConsole command
-bool ServerCommandThread::ShowConsoleAction::Execute(CommandType* command, const DuplexPipe* pipe, void* context)
+bool ServerCommandThread::actions::ShowConsole::Execute(const CommandType*, const DuplexPipe* pipe, void* context, const void*, size_t)
 {
 	pipe->SendAck();
 
@@ -1451,7 +1296,7 @@ bool ServerCommandThread::ShowConsoleAction::Execute(CommandType* command, const
 // END EPIC MOD
 
 // BEGIN EPIC MOD - Adding SetVisible command
-bool ServerCommandThread::SetVisibleAction::Execute(CommandType* command, const DuplexPipe* pipe, void* context)
+bool ServerCommandThread::actions::SetVisible::Execute(const CommandType* command, const DuplexPipe* pipe, void* context, const void*, size_t)
 {
 	pipe->SendAck();
 
@@ -1462,7 +1307,7 @@ bool ServerCommandThread::SetVisibleAction::Execute(CommandType* command, const
 // END EPIC MOD
 
 // BEGIN EPIC MOD - Adding SetActive command
-bool ServerCommandThread::SetActiveAction::Execute(CommandType* command, const DuplexPipe* pipe, void* context)
+bool ServerCommandThread::actions::SetActive::Execute(const CommandType* command, const DuplexPipe* pipe, void* context, const void*, size_t)
 {
 	ServerCommandThread* commandThread = static_cast<ServerCommandThread*>(context);
 
@@ -1478,7 +1323,7 @@ bool ServerCommandThread::SetActiveAction::Execute(CommandType* command, const D
 // END EPIC MOD
 
 // BEGIN EPIC MOD - Adding SetBuildArguments command
-bool ServerCommandThread::SetBuildArgumentsAction::Execute(CommandType* command, const DuplexPipe* pipe, void* context)
+bool ServerCommandThread::actions::SetBuildArguments::Execute(const CommandType* command, const DuplexPipe* pipe, void* context, const void*, size_t)
 {
 	ServerCommandThread* commandThread = static_cast<ServerCommandThread*>(context);
 
@@ -1500,7 +1345,7 @@ bool ServerCommandThread::SetBuildArgumentsAction::Execute(CommandType* command,
 // END EPIC MOD
 
 // BEGIN EPIC MOD - Support for lazy-loading modules
-bool ServerCommandThread::EnableLazyLoadedModuleAction::Execute(CommandType* command, const DuplexPipe* pipe, void* context)
+bool ServerCommandThread::actions::EnableLazyLoadedModule::Execute(const CommandType* command, const DuplexPipe* pipe, void* context, const void*, size_t)
 {
 	ServerCommandThread* commandThread = static_cast<ServerCommandThread*>(context);
 
@@ -1523,7 +1368,7 @@ bool ServerCommandThread::EnableLazyLoadedModuleAction::Execute(CommandType* com
 }
 // END EPIC MOD
 
-bool ServerCommandThread::RegisterProcessAction::Execute(CommandType* command, const DuplexPipe* pipe, void* context)
+bool ServerCommandThread::actions::RegisterProcess::Execute(const CommandType* command, const DuplexPipe* pipe, void* context, const void*, size_t)
 {
 	pipe->SendAck();
 
@@ -1565,7 +1410,7 @@ bool ServerCommandThread::RegisterProcessAction::Execute(CommandType* command, c
 
 		if (registeredSuccessfully)
 		{
-			LiveProcess* liveProcess = new LiveProcess(processHandle, command->processId, command->threadId, pipe);
+			LiveProcess* liveProcess = new LiveProcess(processHandle, command->processId, command->threadId, command->jumpToSelf, pipe);
 			commandThread->m_liveProcesses.push_back(liveProcess);
 			// BEGIN EPIC MOD - No built-in UI
 			// commandThread->m_mainFrame->UpdateWindowTitle();
@@ -1575,139 +1420,75 @@ bool ServerCommandThread::RegisterProcessAction::Execute(CommandType* command, c
 		}
 
 		// tell client we are finished
-		pipe->SendCommandAndWaitForAck(commands::RegisterProcessFinished { registeredSuccessfully });
+		pipe->SendCommandAndWaitForAck(commands::RegisterProcessFinished { registeredSuccessfully }, nullptr, 0u);
 	}
 
 	return true;
 }
 
 
-bool ServerCommandThread::EnableModuleBatchBeginAction::Execute(CommandType*, const DuplexPipe* pipe, void* context)
+bool ServerCommandThread::actions::EnableModules::Execute(const CommandType* command, const DuplexPipe* pipe, void* context, const void* payload, size_t payloadSize)
 {
+	pipe->SendAck();
+
+	// EPIC REMOVED: g_theApp.GetMainFrame()->ChangeStatusBarText(L"Loading modules...");
+
 	ServerCommandThread* commandThread = static_cast<ServerCommandThread*>(context);
 
 	// protect against several client DLLs calling into this action at the same time.
-	// we hold this critical section until we get the BatchEnd signal.
 	// this ensures that all modules are loaded serialized per process.
-	commandThread->m_actionCS.Enter();
+	CriticalSection::ScopedLock lock(&commandThread->m_actionCS);
 
-	commandThread->m_moduleBatchScope.Restart();
-	commandThread->m_loadedCompilandCountInBatchScope = 0u;
+	telemetry::Scope moduleLoadingScope("Module loading");
 
 	// set up virtual drives before loading anything, otherwise files won't be detected and therefore discarded
-	const std::wstring virtualDriveLetter = appSettings::g_virtualDriveLetter->GetValue();
-	const std::wstring virtualDrivePath = appSettings::g_virtualDrivePath->GetValue();
-	if ((virtualDriveLetter.size() != 0) && (virtualDrivePath.size() != 0))
-	{
-		virtualDrive::Add(virtualDriveLetter.c_str(), virtualDrivePath.c_str());
-	}
-
-	pipe->SendAck();
-
-	return true;
-}
-
-
-bool ServerCommandThread::EnableModuleBatchEndAction::Execute(CommandType*, const DuplexPipe* pipe, void* context)
-{
-	ServerCommandThread* commandThread = static_cast<ServerCommandThread*>(context);
+	AddVirtualDrive();
 
-	commandThread->m_moduleBatchScope.End();
-	LC_SUCCESS_USER("Successfully loaded modules (%.3fs, %zu translation units)", commandThread->m_moduleBatchScope.ReadSeconds(), commandThread->m_loadedCompilandCountInBatchScope);
+	scheduler::TaskBase* rootTask = scheduler::CreateEmptyTask();
+	types::vector<scheduler::Task<LiveModule*>*> loadModuleTasks;
 
-	// EPIC REMOVED: commandThread->PrewarmCompilerEnvironmentCache();
+	const unsigned int moduleCount = command->moduleCount;
+	loadModuleTasks.reserve(moduleCount);
 
-	// EPIC REMOVED: g_theApp.GetMainFrame()->ResetStatusBarText();
-
-	// tell user we are ready
-	// BEGIN EPIC MOD - Support for lazy-loading modules
-	if (thread::GetId() != thread::GetId(commandThread->m_compileThread))
-	// END EPIC MOD
+	memoryStream::Reader payloadStream(payload, payloadSize);
+	for (unsigned int i = 0u; i < moduleCount; ++i)
 	{
-		const int shortcut = appSettings::g_compileShortcut->GetValue();
-		const std::wstring& shortcutText = shortcut::ConvertShortcutToText(shortcut);
-		LC_SUCCESS_USER("Live coding ready - Save changes and press %S to re-compile code", shortcutText.c_str());
-	}
+		const commands::ModuleData moduleData = payloadStream.Read<commands::ModuleData>();
+		scheduler::Task<LiveModule*>* task = commandThread->LoadModule(command->processId, moduleData.base, moduleData.path, rootTask);
 
-	// remove virtual drives once we're finished
-	const std::wstring virtualDriveLetter = appSettings::g_virtualDriveLetter->GetValue();
-	const std::wstring virtualDrivePath = appSettings::g_virtualDrivePath->GetValue();
-	if ((virtualDriveLetter.size() != 0) && (virtualDrivePath.size() != 0))
-	{
-		virtualDrive::Remove(virtualDriveLetter.c_str(), virtualDrivePath.c_str());
+		// the module could have failed to load
+		if (task)
+		{
+			loadModuleTasks.push_back(task);
+		}
 	}
 
-	pipe->SendAck();
-
-	// protect against several client DLLs calling into this action at the same time
-	commandThread->m_actionCS.Leave();
-
-	return true;
-}
-
-
-bool ServerCommandThread::DisableModuleBatchBeginAction::Execute(CommandType*, const DuplexPipe* pipe, void* context)
-{
-	ServerCommandThread* commandThread = static_cast<ServerCommandThread*>(context);
-
-	// protect against several client DLLs calling into this action at the same time.
-	// we hold this critical section until we get the BatchEnd signal.
-	// this ensures that all modules are unloaded serialized per process.
-	commandThread->m_actionCS.Enter();
-
-	pipe->SendAck();
-
-	return true;
-}
-
-
-bool ServerCommandThread::DisableModuleBatchEndAction::Execute(CommandType*, const DuplexPipe* pipe, void* context)
-{
-	// EPIC REMOVED: g_theApp.GetMainFrame()->ResetStatusBarText();
-
-	ServerCommandThread* commandThread = static_cast<ServerCommandThread*>(context);
-
-	pipe->SendAck();
-
-	// protect against several client DLLs calling into this action at the same time
-	commandThread->m_actionCS.Leave();
-
-	return true;
-}
-
-
-bool ServerCommandThread::EnableModuleAction::Execute(CommandType* command, const DuplexPipe* pipe, void* context)
-{
-	pipe->SendAck();
-
-	ServerCommandThread* commandThread = static_cast<ServerCommandThread*>(context);
+	// wait for all tasks to finish
+	scheduler::RunTask(rootTask);
+	scheduler::WaitForTask(rootTask);
 
-	TaskContext taskContext = {};
-	taskContext.taskRoot = scheduler::CreateEmptyTask();
-	commandThread->LoadModule(command->path, pipe, &taskContext, command->processId);
+	const size_t loadModuleTaskCount = loadModuleTasks.size();
+	commandThread->m_liveModules.reserve(loadModuleTaskCount);
 
-	// wait for all tasks to finish
-	scheduler::RunTask(taskContext.taskRoot);
-	scheduler::WaitForTask(taskContext.taskRoot);
+	size_t loadedTranslationUnits = 0u;
 
-	// add all live modules loaded by the tasks
-	for (size_t i = 0u; i < taskContext.tasks.size(); ++i)
+	// update all live modules loaded by the tasks
+	for (size_t i = 0u; i < loadModuleTaskCount; ++i)
 	{
-		LiveModule* liveModule = taskContext.tasks[i]->GetResult();
+		scheduler::Task<LiveModule*>* task = loadModuleTasks[i];
+		LiveModule* liveModule = task->GetResult();
+
 		commandThread->m_liveModules.push_back(liveModule);
 
 		// update directory cache for this live module
 		liveModule->UpdateDirectoryCache(commandThread->m_directoryCache);
 
-		// update the number of loaded translation units during this batch
-		commandThread->m_loadedCompilandCountInBatchScope += liveModule->GetCompilandDatabase()->compilands.size();
+		// update the number of loaded translation units
+		loadedTranslationUnits += liveModule->GetCompilandDatabase()->compilands.size();
 	}
 
-	scheduler::DestroyTasks(taskContext.tasks);
-	scheduler::DestroyTask(taskContext.taskRoot);
-
-	// tell client we are finished
-	pipe->SendCommandAndWaitForAck(commands::EnableModuleFinished { command->token });
+	scheduler::DestroyTasks(loadModuleTasks);
+	scheduler::DestroyTask(rootTask);
 
 	// dump memory statistics
 	{
@@ -1719,246 +1500,79 @@ bool ServerCommandThread::EnableModuleAction::Execute(CommandType* command, cons
 		g_dependencyAllocator.PrintStats();
 	}
 
-	return true;
-}
-
-
-bool ServerCommandThread::EnableAllModulesAction::Execute(CommandType* command, const DuplexPipe* pipe, void* context)
-{
-	pipe->SendAck();
-
-	ServerCommandThread* commandThread = static_cast<ServerCommandThread*>(context);
-
-	TaskContext taskContext = {};
-	taskContext.taskRoot = scheduler::CreateEmptyTask();
-	commandThread->LoadAllModules(command->path, pipe, &taskContext, command->processId);
-
-	// wait for all tasks to finish
-	scheduler::RunTask(taskContext.taskRoot);
-	scheduler::WaitForTask(taskContext.taskRoot);
-
-	// add all live modules loaded by the tasks
-	for (size_t i = 0u; i < taskContext.tasks.size(); ++i)
+	if (loadModuleTaskCount > 0u)
 	{
-		LiveModule* liveModule = taskContext.tasks[i]->GetResult();
-		commandThread->m_liveModules.push_back(liveModule);
-
-		// update directory cache for this live module
-		liveModule->UpdateDirectoryCache(commandThread->m_directoryCache);
-
-		// update the number of loaded translation units during this batch
-		commandThread->m_loadedCompilandCountInBatchScope += liveModule->GetCompilandDatabase()->compilands.size();
+		LC_SUCCESS_USER("Loaded %zu module(s) (%.3fs, %zu translation units)", loadModuleTaskCount, moduleLoadingScope.ReadSeconds(), loadedTranslationUnits);
 	}
 
-	scheduler::DestroyTasks(taskContext.tasks);
-	scheduler::DestroyTask(taskContext.taskRoot);
-
-	// tell server we are finished
-	pipe->SendCommandAndWaitForAck(commands::EnableAllModulesFinished { command->token });
+	// EPIC REMOVED commandThread->PrewarmCompilerEnvironmentCache();
 
-	// dump memory statistics
+	// tell user we are ready, but only once to not clutter the log
 	{
-		LC_LOG_INDENT_TELEMETRY;
-		g_symbolAllocator.PrintStats();
-		g_immutableStringAllocator.PrintStats();
-		g_contributionAllocator.PrintStats();
-		g_compilandAllocator.PrintStats();
-		g_dependencyAllocator.PrintStats();
+		static bool showedOnce = false;
+		if (!showedOnce)
+		{
+			showedOnce = true;
+			const int shortcut = appSettings::g_compileShortcut->GetValue();
+			const std::wstring& shortcutText = shortcut::ConvertShortcutToText(shortcut);
+			LC_SUCCESS_USER("Live coding ready - Save changes and press %S to re-compile code", shortcutText.c_str());
+		}
 	}
 
-	return true;
-}
-
-
-bool ServerCommandThread::DisableModuleAction::Execute(CommandType* command, const DuplexPipe* pipe, void* context)
-{
-	pipe->SendAck();
-
-	ServerCommandThread* commandThread = static_cast<ServerCommandThread*>(context);
-	commandThread->UnloadModule(command->path, pipe, command->processId);
+	// remove virtual drives once we're finished
+	RemoveVirtualDrive();
 
 	// tell server we are finished
-	pipe->SendCommandAndWaitForAck(commands::DisableModuleFinished { command->token });
-
-	return true;
-}
-
-
-bool ServerCommandThread::DisableAllModulesAction::Execute(CommandType* command, const DuplexPipe* pipe, void* context)
-{
-	pipe->SendAck();
+	pipe->SendCommandAndWaitForAck(commands::EnableModulesFinished { command->token }, nullptr, 0u);
 
-	ServerCommandThread* commandThread = static_cast<ServerCommandThread*>(context);
-	commandThread->UnloadAllModules(command->path, pipe, command->processId);
-
-	// tell server we are finished
-	pipe->SendCommandAndWaitForAck(commands::DisableAllModulesFinished { command->token });
+	// EPIC REMOVED: g_theApp.GetMainFrame()->ResetStatusBarText();
 
 	return true;
 }
 
 
-bool ServerCommandThread::GetModuleInfoAction::Execute(CommandType* command, const DuplexPipe* pipe, void* context)
+bool ServerCommandThread::actions::DisableModules::Execute(const CommandType* command, const DuplexPipe* pipe, void* context, const void* payload, size_t payloadSize)
 {
 	pipe->SendAck();
 
-	if (!command->moduleBase)
-	{
-		return false;
-	}
-
-	// note that the path we get back from the DLL might not be normalized, depending on how the executable was launched.
-	// it is crucial to normalize the path again, otherwise we could load already loaded modules into the same
-	// Live++ instance, which would wreak havoc
-	const std::wstring modulePath = file::NormalizePath(file::RelativeToAbsolutePath(command->path).c_str());
+	// EPIC REMOVED: g_theApp.GetMainFrame()->ChangeStatusBarText(L"Unloading modules...");
 
 	ServerCommandThread* commandThread = static_cast<ServerCommandThread*>(context);
-	TaskContext* taskContext = static_cast<TaskContext*>(command->taskContext);
 
-	// a task context is provided for loading modules
-	const bool shouldLoad = (taskContext != nullptr);
+	// protect against several client DLLs calling into this action at the same time.
+	// this ensures that all modules are loaded serialized per process.
+	CriticalSection::ScopedLock lock(&commandThread->m_actionCS);
 
-	if (command->loadImports)
+	telemetry::Scope moduleUnloadingScope("Module unloading");
+
+	unsigned int unloadedModules = 0u;
+	const unsigned int moduleCount = command->moduleCount;
+	memoryStream::Reader payloadStream(payload, payloadSize);
+	for (unsigned int i = 0u; i < moduleCount; ++i)
 	{
-		if (shouldLoad)
-		{
-			// load this module and all its import DLLs as well
-			commandThread->LoadAllModules(modulePath.c_str(), pipe, taskContext, command->processId);
-		}
-		else
+		const commands::ModuleData moduleData = payloadStream.Read<commands::ModuleData>();
+		const bool success = commandThread->UnloadModule(command->processId, moduleData.path);
+		if (success)
 		{
-			// unload this module and all its import DLLs as well
-			commandThread->UnloadAllModules(modulePath.c_str(), pipe, command->processId);
+			++unloadedModules;
 		}
-
-		return false;
 	}
 
-	LiveProcess* liveProcess = commandThread->FindProcessById(command->processId);
-	LC_ASSERT(liveProcess, "Invalid process ID.");
-
-	const executable::Header imageHeader = GetImageHeader(modulePath.c_str());
-	if (shouldLoad)
+	if (unloadedModules > 0u)
 	{
-		if (liveProcess->TriedToLoadImage(imageHeader))
-		{
-			// tried loading this module into this process already
-			return false;
-		}
-
-		// find any other process ID that tried to load this module already (if any)
-		{
-			const size_t count = commandThread->m_liveProcesses.size();
-			for (size_t i = 0u; i < count; ++i)
-			{
-				LiveProcess* otherLiveProcess = commandThread->m_liveProcesses[i];
-				if (otherLiveProcess->TriedToLoadImage(imageHeader))
-				{
-					// some *other* process loaded this module already
-					LC_LOG_USER("Registering module %S (PID: %d)", modulePath.c_str(), command->processId);
-
-					LiveModule* liveModule = commandThread->m_imageHeaderToLiveModule[imageHeader];
-					if (liveModule)
-					{
-						const unsigned int processId = command->processId;
-						void* moduleBase = command->moduleBase;
-
-						liveModule->RegisterProcess(liveProcess, moduleBase, modulePath);
-						liveModule->DisableControlFlowGuard(liveProcess, moduleBase);
-					
-						const bool installedPatchesSuccessfully = liveModule->InstallCompiledPatches(liveProcess, moduleBase);
-						if (!installedPatchesSuccessfully)
-						{
-							LC_ERROR_USER("Compiled patches could not be installed (PID: %d)", processId);
-							liveModule->UnregisterProcess(liveProcess);
-						}
-
-						liveProcess->AddLoadedImage(imageHeader);
-					}
-
-					return false;
-				}
-			}
-		}
-
-		symbols::Provider* moduleProvider = symbols::OpenEXE(modulePath.c_str(), symbols::OpenOptions::ACCUMULATE_SIZE);
-		if (!moduleProvider)
-		{
-			liveProcess->AddLoadedImage(imageHeader);
-			return false;
-		}
-
-		// this live module hasn't been loaded yet by any process
-		void* moduleBase = command->moduleBase;
-
-		// accumulate module info
-		const file::Attributes attributes = file::GetAttributes(modulePath.c_str());
-		const uint64_t size = file::GetSize(attributes);
-		g_loadedModuleSize.Accumulate(size);
-
-		{
-			// create a task to load the module of this batch concurrently
-			LC_LOG_USER("Loading module %S (%.3f MB)", modulePath.c_str(), size / 1048576.0f);
-
-			LiveModule* liveModule = new LiveModule(modulePath.c_str(), imageHeader, commandThread->m_runMode);
-			commandThread->m_imageHeaderToLiveModule.emplace(imageHeader, liveModule);
-
-			auto task = scheduler::CreateTask(taskContext->taskRoot, [liveModule, liveProcess, modulePath, moduleBase, moduleProvider]()
-			{
-				telemetry::Scope scope("Loading module");
-
-				symbols::DiaCompilandDB* moduleDiaCompilandDb = symbols::GatherDiaCompilands(moduleProvider);
-
-				liveModule->Load(moduleProvider, moduleDiaCompilandDb);
-				liveModule->RegisterProcess(liveProcess, moduleBase, modulePath);
-				liveModule->DisableControlFlowGuard(liveProcess, moduleBase);
-
-				symbols::DestroyDiaCompilandDB(moduleDiaCompilandDb);
-				symbols::Close(moduleProvider);
-
-				return liveModule;
-			});
-			scheduler::RunTask(task);
-
-			taskContext->tasks.emplace_back(task);
-		}
-
-		g_loadedModuleSize.Print();
-		g_loadedModuleSize.ResetCurrent();
-
-		liveProcess->AddLoadedImage(imageHeader);
-
-		return false;
+		LC_SUCCESS_USER("Unloaded %u module(s) (%.3fs)", unloadedModules, moduleUnloadingScope.ReadSeconds());
 	}
-	else
-	{
-		LC_LOG_USER("Unloading module %S", modulePath.c_str());
-
-		liveProcess->RemoveLoadedImage(imageHeader);
-		commandThread->m_imageHeaderToLiveModule.erase(imageHeader);
 
-		for (auto it = commandThread->m_liveModules.begin(); it != commandThread->m_liveModules.end(); /* nothing */)
-		{
-			LiveModule* liveModule = *it;
-			if (std::equal_to<executable::Header>()(liveModule->GetImageHeader(), imageHeader))
-			{
-				liveModule->Unload();
-				delete liveModule;
+	// tell server we are finished
+	pipe->SendCommandAndWaitForAck(commands::DisableModulesFinished { command->token }, nullptr, 0u);
 
-				it = commandThread->m_liveModules.erase(it);
-			}
-			else
-			{
-				++it;
-			}
-		}
+	// EPIC REMOVED: g_theApp.GetMainFrame()->ResetStatusBarText();
 
-		return false;
-	}
+	return true;
 }
 
 
-bool ServerCommandThread::ApplySettingBoolAction::Execute(CommandType* command, const DuplexPipe* pipe, void*)
+bool ServerCommandThread::actions::ApplySettingBool::Execute(const CommandType* command, const DuplexPipe* pipe, void*, const void*, size_t)
 {
 	pipe->SendAck();
 
@@ -1968,7 +1582,7 @@ bool ServerCommandThread::ApplySettingBoolAction::Execute(CommandType* command,
 }
 
 
-bool ServerCommandThread::ApplySettingIntAction::Execute(CommandType* command, const DuplexPipe* pipe, void*)
+bool ServerCommandThread::actions::ApplySettingInt::Execute(const CommandType* command, const DuplexPipe* pipe, void*, const void*, size_t)
 {
 	pipe->SendAck();
 
@@ -1978,7 +1592,7 @@ bool ServerCommandThread::ApplySettingIntAction::Execute(CommandType* command, c
 }
 
 
-bool ServerCommandThread::ApplySettingStringAction::Execute(CommandType* command, const DuplexPipe* pipe, void*)
+bool ServerCommandThread::actions::ApplySettingString::Execute(const CommandType* command, const DuplexPipe* pipe, void*, const void*, size_t)
 {
 	pipe->SendAck();
 
diff --git a/Engine/Source/Developer/Windows/LiveCodingServer/Private/External/LC_ServerCommandThread.h b/Engine/Source/Developer/Windows/LiveCodingServer/Private/External/LC_ServerCommandThread.h
--- a/Engine/Source/Developer/Windows/LiveCodingServer/Private/External/LC_ServerCommandThread.h
+++ b/Engine/Source/Developer/Windows/LiveCodingServer/Private/External/LC_ServerCommandThread.h
@@ -29,31 +29,17 @@ public:
 
 	std::wstring GetProcessImagePath(void) const;
 
-	struct TaskContext
-	{
-		scheduler::TaskBase* taskRoot;
-		types::vector<scheduler::Task<LiveModule*>*> tasks;
-	};
-
 private:
-	void LoadModule(const wchar_t* modulePath, const DuplexPipe* pipe, TaskContext* tasks, unsigned int processId);
-	void LoadAllModules(const wchar_t* modulePath, const DuplexPipe* pipe, TaskContext* tasks, unsigned int processId);
-
-	void UnloadModule(const wchar_t* modulePath, const DuplexPipe* pipe, unsigned int processId);
-	void UnloadAllModules(const wchar_t* modulePath, const DuplexPipe* pipe, unsigned int processId);
+	scheduler::Task<LiveModule*>* LoadModule(unsigned int processId, void* moduleBase, const wchar_t* modulePath, scheduler::TaskBase* taskRoot);
+	bool UnloadModule(unsigned int processId, const wchar_t* modulePath);
 
 	void PrewarmCompilerEnvironmentCache(void);
 
-	static unsigned int __stdcall ServerThreadProxy(void* context);
 	unsigned int ServerThread(void);
-
-	static unsigned int __stdcall CompileThreadProxy(void* context);
 	unsigned int CompileThread(void);
 
 	struct CommandThreadContext
 	{
-		ServerCommandThread* instance;
-
 		DuplexPipeServer pipe;
 		Event* readyEvent;
 		thread::Handle commandThread;
@@ -62,10 +48,7 @@ private:
 		thread::Handle exceptionCommandThread;
 	};
 
-	static unsigned int __stdcall CommandThreadProxy(void* context);
 	unsigned int CommandThread(DuplexPipeServer* pipe, Event* readyEvent);
-
-	static unsigned int __stdcall ExceptionCommandThreadProxy(void* context);
 	unsigned int ExceptionCommandThread(DuplexPipeServer* exceptionPipe);
 
 	void RemoveCommandThread(const DuplexPipe* pipe);
@@ -74,159 +57,53 @@ private:
 
 	void CompileChanges(bool didAllProcessesMakeProgress);
 
-	// actions
-	struct TriggerRecompileAction
-	{
-		typedef commands::TriggerRecompile CommandType;
-		static bool Execute(CommandType* command, const DuplexPipe* pipe, void* context);
-	};
-
-	struct BuildPatchAction
-	{
-		typedef commands::BuildPatch CommandType;
-		static bool Execute(CommandType* command, const DuplexPipe* pipe, void* context);
-	};
-
-	struct HandleExceptionAction
-	{
-		typedef commands::HandleException CommandType;
-		static bool Execute(CommandType* command, const DuplexPipe* pipe, void* context);
-	};
-
-	struct ReadyForCompilationAction
-	{
-		typedef commands::ReadyForCompilation CommandType;
-		static bool Execute(CommandType* command, const DuplexPipe* pipe, void* context);
-	};
-
-	struct DisconnectClientAction
-	{
-		typedef commands::DisconnectClient CommandType;
-		static bool Execute(CommandType* command, const DuplexPipe* pipe, void* context);
-	};
-
-	// BEGIN EPIC MOD - Adding ShowConsole command
-	struct ShowConsoleAction
-	{
-		typedef commands::ShowConsole CommandType;
-		static bool Execute(CommandType* command, const DuplexPipe* pipe, void* context);
-	};
-	// END EPIC MOD
 
-	// BEGIN EPIC MOD - Adding ShowConsole command
-	struct SetVisibleAction
-	{
-		typedef commands::SetVisible CommandType;
-		static bool Execute(CommandType* command, const DuplexPipe* pipe, void* context);
-	};
-	// END EPIC MOD
-
-	// BEGIN EPIC MOD - Adding SetActive command
-	struct SetActiveAction
-	{
-		typedef commands::SetActive CommandType;
-		static bool Execute(CommandType* command, const DuplexPipe* pipe, void* context);
-	};
-	// END EPIC MOD
-
-	// BEGIN EPIC MOD - Adding SetBuildArguments command
-	struct SetBuildArgumentsAction
-	{
-		typedef commands::SetBuildArguments CommandType;
-		static bool Execute(CommandType* command, const DuplexPipe* pipe, void* context);
-	};
-	// END EPIC MOD
-
-	// BEGIN EPIC MOD - Adding support for lazy-loading modules
-	struct EnableLazyLoadedModuleAction
+	// actions
+	struct actions
 	{
-		typedef commands::EnableLazyLoadedModule CommandType;
-		static bool Execute(CommandType* command, const DuplexPipe* pipe, void* context);
-	};
+		#define DECLARE_ACTION(_name)																													\
+			struct _name																																\
+			{																																			\
+				typedef ::commands::_name CommandType;																									\
+				static bool Execute(const CommandType* command, const DuplexPipe* pipe, void* context, const void* payload, size_t payloadSize);		\
+			}
 
-	struct FinishedLazyLoadingModulesAction
-	{
-		typedef commands::FinishedLazyLoadingModules CommandType;
-		static bool Execute(CommandType* command, const DuplexPipe* pipe, void* context);
-	};
-	// END EPIC MOD
+		DECLARE_ACTION(TriggerRecompile);
+		DECLARE_ACTION(BuildPatch);
+		DECLARE_ACTION(HandleException);
+		DECLARE_ACTION(ReadyForCompilation);
+		DECLARE_ACTION(DisconnectClient);
+		DECLARE_ACTION(RegisterProcess);
 
-	struct RegisterProcessAction
-	{
-		typedef commands::RegisterProcess CommandType;
-		static bool Execute(CommandType* command, const DuplexPipe* pipe, void* context);
-	};
+		DECLARE_ACTION(EnableModules);
+		DECLARE_ACTION(DisableModules);
 
-	struct EnableModuleBatchBeginAction
-	{
-		typedef commands::EnableModuleBatchBegin CommandType;
-		static bool Execute(CommandType* command, const DuplexPipe* pipe, void* context);
-	};
+		DECLARE_ACTION(ApplySettingBool);
+		DECLARE_ACTION(ApplySettingInt);
+		DECLARE_ACTION(ApplySettingString);
 
-	struct EnableModuleBatchEndAction
-	{
-		typedef commands::EnableModuleBatchEnd CommandType;
-		static bool Execute(CommandType* command, const DuplexPipe* pipe, void* context);
-	};
+		// BEGIN EPIC MOD - Adding ShowConsole command
+		DECLARE_ACTION(ShowConsole);
+		// END EPIC MOD
 
-	struct DisableModuleBatchBeginAction
-	{
-		typedef commands::DisableModuleBatchBegin CommandType;
-		static bool Execute(CommandType* command, const DuplexPipe* pipe, void* context);
-	};
+		// BEGIN EPIC MOD - Adding ShowConsole command
+		DECLARE_ACTION(SetVisible);
+		// END EPIC MOD
 
-	struct DisableModuleBatchEndAction
-	{
-		typedef commands::DisableModuleBatchEnd CommandType;
-		static bool Execute(CommandType* command, const DuplexPipe* pipe, void* context);
-	};
+		// BEGIN EPIC MOD - Adding SetActive command
+		DECLARE_ACTION(SetActive);
+		// END EPIC MOD
 
-	struct EnableModuleAction
-	{
-		typedef commands::EnableModule CommandType;
-		static bool Execute(CommandType* command, const DuplexPipe* pipe, void* context);
-	};
+		// BEGIN EPIC MOD - Adding SetBuildArguments command
+		DECLARE_ACTION(SetBuildArguments);
+		// END EPIC MOD
 
-	struct EnableAllModulesAction
-	{
-		typedef commands::EnableAllModules CommandType;
-		static bool Execute(CommandType* command, const DuplexPipe* pipe, void* context);
-	};
-
-	struct DisableModuleAction
-	{
-		typedef commands::DisableModule CommandType;
-		static bool Execute(CommandType* command, const DuplexPipe* pipe, void* context);
-	};
-
-	struct DisableAllModulesAction
-	{
-		typedef commands::DisableAllModules CommandType;
-		static bool Execute(CommandType* command, const DuplexPipe* pipe, void* context);
-	};
-
-	struct GetModuleInfoAction
-	{
-		typedef commands::GetModuleInfo CommandType;
-		static bool Execute(CommandType* command, const DuplexPipe* pipe, void* context);
-	};
+		// BEGIN EPIC MOD - Adding support for lazy-loading modules
+		DECLARE_ACTION(EnableLazyLoadedModule);
+		DECLARE_ACTION(FinishedLazyLoadingModules);
+		// END EPIC MOD
 
-	struct ApplySettingBoolAction
-	{
-		typedef commands::ApplySettingBool CommandType;
-		static bool Execute(CommandType* command, const DuplexPipe* pipe, void* context);
-	};
-
-	struct ApplySettingIntAction
-	{
-		typedef commands::ApplySettingInt CommandType;
-		static bool Execute(CommandType* command, const DuplexPipe* pipe, void* context);
-	};
-
-	struct ApplySettingStringAction
-	{
-		typedef commands::ApplySettingString CommandType;
-		static bool Execute(CommandType* command, const DuplexPipe* pipe, void* context);
+		#undef DECLARE_ACTION
 	};
 
 
@@ -253,9 +130,6 @@ private:
 	CriticalSection m_connectionCS;
 	types::vector<CommandThreadContext*> m_commandThreads;
 
-	telemetry::Scope m_moduleBatchScope;
-	size_t m_loadedCompilandCountInBatchScope;
-
 	// BEGIN EPIC MOD - Adding SetActive command
 	bool m_active = true;
 	// END EPIC MOD
diff --git a/Engine/Source/Developer/Windows/LiveCodingServer/Private/External/LC_SymbolPatterns.cpp b/Engine/Source/Developer/Windows/LiveCodingServer/Private/External/LC_SymbolPatterns.cpp
--- a/Engine/Source/Developer/Windows/LiveCodingServer/Private/External/LC_SymbolPatterns.cpp
+++ b/Engine/Source/Developer/Windows/LiveCodingServer/Private/External/LC_SymbolPatterns.cpp
@@ -151,6 +151,15 @@ namespace symbolPatterns
 #endif
 	};
 
+	extern const char* const EXCEPTION_UNWIND_PATTERNS[1] =
+	{
+#if LC_64_BIT
+		"?dtor$",
+#else
+		"__unwindfunclet$",
+#endif
+	};
+
 	extern const char* const EXCEPTION_CLAUSE_PATTERNS[1] =
 	{
 		"__catch$"
diff --git a/Engine/Source/Developer/Windows/LiveCodingServer/Private/External/LC_SymbolPatterns.h b/Engine/Source/Developer/Windows/LiveCodingServer/Private/External/LC_SymbolPatterns.h
--- a/Engine/Source/Developer/Windows/LiveCodingServer/Private/External/LC_SymbolPatterns.h
+++ b/Engine/Source/Developer/Windows/LiveCodingServer/Private/External/LC_SymbolPatterns.h
@@ -24,6 +24,7 @@ namespace symbolPatterns
 #else
 	extern const char* const EXCEPTION_RELATED_PATTERNS[10];
 #endif
+	extern const char* const EXCEPTION_UNWIND_PATTERNS[1];
 
 	extern const char* const EXCEPTION_CLAUSE_PATTERNS[1];
 
diff --git a/Engine/Source/Developer/Windows/LiveCodingServer/Private/External/LC_Symbols.cpp b/Engine/Source/Developer/Windows/LiveCodingServer/Private/External/LC_Symbols.cpp
--- a/Engine/Source/Developer/Windows/LiveCodingServer/Private/External/LC_Symbols.cpp
+++ b/Engine/Source/Developer/Windows/LiveCodingServer/Private/External/LC_Symbols.cpp
@@ -2207,6 +2207,13 @@ namespace symbols
 	}
 
 
+	bool IsExceptionUnwindSymbolForDynamicInitializer(const ImmutableString& symbolName)
+	{
+		return (StartsWithPatterns(symbolName.c_str(), symbolPatterns::EXCEPTION_UNWIND_PATTERNS) &&
+			ContainsPatterns(symbolName.c_str(), symbolPatterns::DYNAMIC_INITIALIZER_PATTERNS));
+	}
+
+
 	bool IsRuntimeCheckRelatedSymbol(const ImmutableString& symbolName)
 	{
 		return ContainsPatterns(symbolName.c_str(), symbolPatterns::RTC_PATTERNS);
diff --git a/Engine/Source/Developer/Windows/LiveCodingServer/Private/External/LC_Symbols.h b/Engine/Source/Developer/Windows/LiveCodingServer/Private/External/LC_Symbols.h
--- a/Engine/Source/Developer/Windows/LiveCodingServer/Private/External/LC_Symbols.h
+++ b/Engine/Source/Developer/Windows/LiveCodingServer/Private/External/LC_Symbols.h
@@ -346,6 +346,7 @@ namespace symbols
 
 	bool IsExceptionRelatedSymbol(const ImmutableString& symbolName);
 	bool IsExceptionClauseSymbol(const ImmutableString& symbolName);
+	bool IsExceptionUnwindSymbolForDynamicInitializer(const ImmutableString& symbolName);
 
 	bool IsRuntimeCheckRelatedSymbol(const ImmutableString& symbolName);
 	bool IsSdlCheckRelatedSymbol(const ImmutableString& symbolName);
