Merging latest version of Live++ (check timestamps on object files, support for adding multiple modules at once)

#rnx

diff --git a/Engine/Source/Developer/Windows/LiveCoding/Private/External/LC_ClientCommandActions.cpp b/Engine/Source/Developer/Windows/LiveCoding/Private/External/LC_ClientCommandActions.cpp
--- a/Engine/Source/Developer/Windows/LiveCoding/Private/External/LC_ClientCommandActions.cpp
+++ b/Engine/Source/Developer/Windows/LiveCoding/Private/External/LC_ClientCommandActions.cpp
@@ -12,11 +12,11 @@
 
 bool actions::RegisterProcessFinished::Execute(const CommandType* command, const DuplexPipe* pipe, void* context, const void*, size_t)
 {
-	pipe->SendAck();
-
 	bool* successfullyRegisteredProcess = static_cast<bool*>(context);
 	*successfullyRegisteredProcess = command->success;
 
+	pipe->SendAck();
+
 	// don't continue execution
 	return false;
 }
@@ -60,16 +60,33 @@ bool actions::LeaveSyncPoint::Execute(const CommandType*, const DuplexPipe* pipe
 }
 
 
-bool actions::CallHooks::Execute(const CommandType* command, const DuplexPipe* pipe, void*, const void*, size_t)
+bool actions::CallHooks::Execute(const CommandType* command, const DuplexPipe* pipe, void*, const void* payload, size_t)
 {
-	for (const hook::Function* hook = command->first; hook < command->last; ++hook)
+	switch (command->type)
 	{
-		// note that sections are often padded with zeroes, so skip everything that's zero
-		hook::Function function = *hook;
-		if (function)
-		{
-			function();
-		}
+		case hook::Type::PREPATCH:
+			hook::CallHooksInRange<hook::PrepatchFunction>(command->rangeBegin, command->rangeEnd);
+			break;
+
+		case hook::Type::POSTPATCH:
+			hook::CallHooksInRange<hook::PostpatchFunction>(command->rangeBegin, command->rangeEnd);
+			break;
+
+		case hook::Type::COMPILE_START:
+			hook::CallHooksInRange<hook::CompileStartFunction>(command->rangeBegin, command->rangeEnd);
+			break;
+
+		case hook::Type::COMPILE_SUCCESS:
+			hook::CallHooksInRange<hook::CompileSuccessFunction>(command->rangeBegin, command->rangeEnd);
+			break;
+
+		case hook::Type::COMPILE_ERROR:
+			hook::CallHooksInRange<hook::CompileErrorFunction>(command->rangeBegin, command->rangeEnd);
+			break;
+
+		case hook::Type::COMPILE_ERROR_MESSAGE:
+			hook::CallHooksInRange<hook::CompileErrorMessageFunction>(command->rangeBegin, command->rangeEnd, static_cast<const wchar_t*>(payload));
+			break;
 	}
 
 	pipe->SendAck();
diff --git a/Engine/Source/Developer/Windows/LiveCoding/Private/External/LC_ClientStartupThread.cpp b/Engine/Source/Developer/Windows/LiveCoding/Private/External/LC_ClientStartupThread.cpp
--- a/Engine/Source/Developer/Windows/LiveCoding/Private/External/LC_ClientStartupThread.cpp
+++ b/Engine/Source/Developer/Windows/LiveCoding/Private/External/LC_ClientStartupThread.cpp
@@ -129,7 +129,7 @@ void ClientStartupThread::Join(void)
 }
 
 
-void* ClientStartupThread::EnableModule(const wchar_t* const nameOfExeOrDll)
+void* ClientStartupThread::EnableModule(const wchar_t* nameOfExeOrDll)
 {
 	// wait for the startup thread to finish initialization
 	Join();
@@ -143,7 +143,21 @@ void* ClientStartupThread::EnableModule(const wchar_t* const nameOfExeOrDll)
 }
 
 
-void* ClientStartupThread::EnableAllModules(const wchar_t* const nameOfExeOrDll)
+void* ClientStartupThread::EnableModules(const wchar_t* namesOfExeOrDll[], unsigned int count)
+{
+	// wait for the startup thread to finish initialization
+	Join();
+
+	if (m_userCommandThread)
+	{
+		return m_userCommandThread->EnableModules(namesOfExeOrDll, count);
+	}
+
+	return nullptr;
+}
+
+
+void* ClientStartupThread::EnableAllModules(const wchar_t* nameOfExeOrDll)
 {
 	// wait for the startup thread to finish initialization
 	Join();
@@ -157,7 +171,7 @@ void* ClientStartupThread::EnableAllModules(const wchar_t* const nameOfExeOrDll)
 }
 
 
-void* ClientStartupThread::DisableModule(const wchar_t* const nameOfExeOrDll)
+void* ClientStartupThread::DisableModule(const wchar_t* nameOfExeOrDll)
 {
 	// wait for the startup thread to finish initialization
 	Join();
@@ -171,7 +185,21 @@ void* ClientStartupThread::DisableModule(const wchar_t* const nameOfExeOrDll)
 }
 
 
-void* ClientStartupThread::DisableAllModules(const wchar_t* const nameOfExeOrDll)
+void* ClientStartupThread::DisableModules(const wchar_t* namesOfExeOrDll[], unsigned int count)
+{
+	// wait for the startup thread to finish initialization
+	Join();
+
+	if (m_userCommandThread)
+	{
+		return m_userCommandThread->DisableModules(namesOfExeOrDll, count);
+	}
+
+	return nullptr;
+}
+
+
+void* ClientStartupThread::DisableAllModules(const wchar_t* nameOfExeOrDll)
 {
 	// wait for the startup thread to finish initialization
 	Join();
@@ -301,7 +329,8 @@ void ClientStartupThread::EnableLazyLoadedModule(const wchar_t* fileName, Window
 }
 // END EPIC MOD
 
-void ClientStartupThread::ApplySettingBool(const char* const settingName, int value)
+
+void ClientStartupThread::ApplySettingBool(const char* settingName, int value)
 {
 	// wait for the startup thread to finish initialization
 	Join();
@@ -313,7 +342,7 @@ void ClientStartupThread::ApplySettingBool(const char* const settingName, int va
 }
 
 
-void ClientStartupThread::ApplySettingInt(const char* const settingName, int value)
+void ClientStartupThread::ApplySettingInt(const char* settingName, int value)
 {
 	// wait for the startup thread to finish initialization
 	Join();
@@ -325,7 +354,7 @@ void ClientStartupThread::ApplySettingInt(const char* const settingName, int val
 }
 
 
-void ClientStartupThread::ApplySettingString(const char* const settingName, const wchar_t* const value)
+void ClientStartupThread::ApplySettingString(const char* settingName, const wchar_t* value)
 {
 	// wait for the startup thread to finish initialization
 	Join();
diff --git a/Engine/Source/Developer/Windows/LiveCoding/Private/External/LC_ClientStartupThread.h b/Engine/Source/Developer/Windows/LiveCoding/Private/External/LC_ClientStartupThread.h
--- a/Engine/Source/Developer/Windows/LiveCoding/Private/External/LC_ClientStartupThread.h
+++ b/Engine/Source/Developer/Windows/LiveCoding/Private/External/LC_ClientStartupThread.h
@@ -27,11 +27,13 @@ public:
 	// Joins the thread, waiting for initialization to finish
 	void Join(void);
 
-	void* EnableModule(const wchar_t* const nameOfExeOrDll);
-	void* EnableAllModules(const wchar_t* const nameOfExeOrDll);
+	void* EnableModule(const wchar_t* nameOfExeOrDll);
+	void* EnableModules(const wchar_t* namesOfExeOrDll[], unsigned int count);
+	void* EnableAllModules(const wchar_t* nameOfExeOrDll);
 
-	void* DisableModule(const wchar_t* const nameOfExeOrDll);
-	void* DisableAllModules(const wchar_t* const nameOfExeOrDll);
+	void* DisableModule(const wchar_t* nameOfExeOrDll);
+	void* DisableModules(const wchar_t* namesOfExeOrDll[], unsigned int count);
+	void* DisableAllModules(const wchar_t* nameOfExeOrDll);
 
 	void WaitForToken(void* token);
 	void TriggerRecompile(void);
@@ -39,9 +41,9 @@ public:
 
 	void InstallExceptionHandler(void);
 
-	void ApplySettingBool(const char* const settingName, int value);
-	void ApplySettingInt(const char* const settingName, int value);
-	void ApplySettingString(const char* const settingName, const wchar_t* const value);
+	void ApplySettingBool(const char* settingName, int value);
+	void ApplySettingInt(const char* settingName, int value);
+	void ApplySettingString(const char* settingName, const wchar_t* value);
 
 	// BEGIN EPIC MOD - Adding ShowConsole command
 	void ShowConsole();
diff --git a/Engine/Source/Developer/Windows/LiveCoding/Private/External/LC_ClientUserCommandThread.cpp b/Engine/Source/Developer/Windows/LiveCoding/Private/External/LC_ClientUserCommandThread.cpp
--- a/Engine/Source/Developer/Windows/LiveCoding/Private/External/LC_ClientUserCommandThread.cpp
+++ b/Engine/Source/Developer/Windows/LiveCoding/Private/External/LC_ClientUserCommandThread.cpp
@@ -176,32 +176,59 @@ void ClientUserCommandThread::Join(void)
 }
 
 
-void* ClientUserCommandThread::EnableModule(const wchar_t* const nameOfExeOrDll)
+void* ClientUserCommandThread::EnableModule(const wchar_t* nameOfExeOrDll)
 {
-	HMODULE module = ::GetModuleHandleW(nameOfExeOrDll);
-	if (!module)
+	return EnableModules(&nameOfExeOrDll, 1u);
+}
+
+
+void* ClientUserCommandThread::EnableModules(const wchar_t* namesOfExeOrDll[], unsigned int count)
+{
+	std::vector<Windows::HMODULE> loadedModules;
+	loadedModules.reserve(count);
+
+	for (unsigned int i = 0u; i < count; ++i)
 	{
-		LC_ERROR_USER("Cannot enable module %S because it is not loaded by this process.", nameOfExeOrDll);
+		Windows::HMODULE module = ::GetModuleHandleW(namesOfExeOrDll[i]);
+		if (module)
+		{
+			loadedModules.push_back(module);
+		}
+		else
+		{
+			LC_ERROR_USER("Cannot enable module %S because it is not loaded by this process.", namesOfExeOrDll[i]);
+		}
+	}
+
+	const size_t loadedModuleCount = loadedModules.size();
+	if (loadedModuleCount == 0u)
+	{
+		// nothing to load
 		return nullptr;
 	}
 
-	ProxyCommand<commands::EnableModules>* proxy = new ProxyCommand<commands::EnableModules>(true, sizeof(commands::ModuleData)*1u);
+	ProxyCommand<commands::EnableModules>* proxy = new ProxyCommand<commands::EnableModules>(true, sizeof(commands::ModuleData) * loadedModuleCount);
 	proxy->m_command.processId = process::GetId();
-	proxy->m_command.moduleCount = 1u;
+	proxy->m_command.moduleCount = static_cast<unsigned int>(loadedModuleCount);
 	proxy->m_command.token = new Event(nullptr, Event::Type::AUTO_RESET);
 
-	commands::ModuleData moduleData = {};
-	moduleData.base = module;
-	::GetModuleFileNameW(module, moduleData.path, MAX_PATH);
+	for (size_t i = 0u; i < loadedModuleCount; ++i)
+	{
+		Windows::HMODULE module = loadedModules[i];
+
+		commands::ModuleData moduleData = {};
+		moduleData.base = module;
+		::GetModuleFileNameW(module, moduleData.path, MAX_PATH);
 
-	proxy->m_payload.Write(moduleData);
+		proxy->m_payload.Write(moduleData);
+	}
 
 	PushUserCommand(proxy);
 
 	return proxy->m_command.token;
 }
 
-void* ClientUserCommandThread::EnableAllModules(const wchar_t* const nameOfExeOrDll)
+void* ClientUserCommandThread::EnableAllModules(const wchar_t* nameOfExeOrDll)
 {
 	HMODULE module = ::GetModuleHandleW(nameOfExeOrDll);
 	if (!module)
@@ -230,25 +257,52 @@ void* ClientUserCommandThread::EnableAllModules(const wchar_t* const nameOfExeOr
 }
 
 
-void* ClientUserCommandThread::DisableModule(const wchar_t* const nameOfExeOrDll)
+void* ClientUserCommandThread::DisableModule(const wchar_t* nameOfExeOrDll)
 {
-	HMODULE module = ::GetModuleHandleW(nameOfExeOrDll);
-	if (!module)
+	return DisableModules(&nameOfExeOrDll, 1u);
+}
+
+
+void* ClientUserCommandThread::DisableModules(const wchar_t* namesOfExeOrDll[], unsigned int count)
+{
+	std::vector<Windows::HMODULE> loadedModules;
+	loadedModules.reserve(count);
+
+	for (unsigned int i = 0u; i < count; ++i)
 	{
-		LC_ERROR_USER("Cannot disable module %S because it is not loaded by this process.", nameOfExeOrDll);
+		Windows::HMODULE module = ::GetModuleHandleW(namesOfExeOrDll[i]);
+		if (module)
+		{
+			loadedModules.push_back(module);
+		}
+		else
+		{
+			LC_ERROR_USER("Cannot disable module %S because it is not loaded by this process.", namesOfExeOrDll[i]);
+		}
+	}
+
+	const size_t loadedModuleCount = loadedModules.size();
+	if (loadedModuleCount == 0u)
+	{
+		// nothing to unload
 		return nullptr;
 	}
 
-	ProxyCommand<commands::DisableModules>* proxy = new ProxyCommand<commands::DisableModules>(true, sizeof(commands::ModuleData) * 1u);
+	ProxyCommand<commands::DisableModules>* proxy = new ProxyCommand<commands::DisableModules>(true, sizeof(commands::ModuleData) * loadedModuleCount);
 	proxy->m_command.processId = process::GetId();
-	proxy->m_command.moduleCount = 1u;
+	proxy->m_command.moduleCount = static_cast<unsigned int>(loadedModuleCount);
 	proxy->m_command.token = new Event(nullptr, Event::Type::AUTO_RESET);
 
-	commands::ModuleData moduleData = {};
-	moduleData.base = module;
-	::GetModuleFileNameW(module, moduleData.path, MAX_PATH);
+	for (size_t i = 0u; i < loadedModuleCount; ++i)
+	{
+		HMODULE module = loadedModules[i];
+
+		commands::ModuleData moduleData = {};
+		moduleData.base = module;
+		::GetModuleFileNameW(module, moduleData.path, MAX_PATH);
 
-	proxy->m_payload.Write(moduleData);
+		proxy->m_payload.Write(moduleData);
+	}
 
 	PushUserCommand(proxy);
 
@@ -256,7 +310,7 @@ void* ClientUserCommandThread::DisableModule(const wchar_t* const nameOfExeOrDll
 }
 
 
-void* ClientUserCommandThread::DisableAllModules(const wchar_t* const nameOfExeOrDll)
+void* ClientUserCommandThread::DisableAllModules(const wchar_t* nameOfExeOrDll)
 {
 	HMODULE module = ::GetModuleHandleW(nameOfExeOrDll);
 	if (!module)
@@ -333,7 +387,7 @@ void ClientUserCommandThread::BuildPatch(const wchar_t* moduleNames[], const wch
 }
 
 
-void ClientUserCommandThread::ApplySettingBool(const char* const settingName, int value)
+void ClientUserCommandThread::ApplySettingBool(const char* settingName, int value)
 {
 	ProxyCommand<commands::ApplySettingBool>* proxy = new ProxyCommand<commands::ApplySettingBool>(false, 0u);
 	strcpy_s(proxy->m_command.settingName, settingName);
@@ -343,7 +397,7 @@ void ClientUserCommandThread::ApplySettingBool(const char* const settingName, in
 }
 
 
-void ClientUserCommandThread::ApplySettingInt(const char* const settingName, int value)
+void ClientUserCommandThread::ApplySettingInt(const char* settingName, int value)
 {
 	ProxyCommand<commands::ApplySettingInt>* proxy = new ProxyCommand<commands::ApplySettingInt>(false, 0u);
 	strcpy_s(proxy->m_command.settingName, settingName);
@@ -353,7 +407,7 @@ void ClientUserCommandThread::ApplySettingInt(const char* const settingName, int
 }
 
 
-void ClientUserCommandThread::ApplySettingString(const char* const settingName, const wchar_t* const value)
+void ClientUserCommandThread::ApplySettingString(const char* settingName, const wchar_t* value)
 {
 	ProxyCommand<commands::ApplySettingString>* proxy = new ProxyCommand<commands::ApplySettingString>(false, 0u);
 	strcpy_s(proxy->m_command.settingName, settingName);
diff --git a/Engine/Source/Developer/Windows/LiveCoding/Private/External/LC_ClientUserCommandThread.h b/Engine/Source/Developer/Windows/LiveCoding/Private/External/LC_ClientUserCommandThread.h
--- a/Engine/Source/Developer/Windows/LiveCoding/Private/External/LC_ClientUserCommandThread.h
+++ b/Engine/Source/Developer/Windows/LiveCoding/Private/External/LC_ClientUserCommandThread.h
@@ -50,11 +50,13 @@ public:
 	// Joins this thread.
 	void Join(void);
 
-	void* EnableModule(const wchar_t* const nameOfExeOrDll);
-	void* EnableAllModules(const wchar_t* const nameOfExeOrDll);
+	void* EnableModule(const wchar_t* nameOfExeOrDll);
+	void* EnableModules(const wchar_t* namesOfExeOrDll[], unsigned int count);
+	void* EnableAllModules(const wchar_t* nameOfExeOrDll);
 
-	void* DisableModule(const wchar_t* const nameOfExeOrDll);
-	void* DisableAllModules(const wchar_t* const nameOfExeOrDll);
+	void* DisableModule(const wchar_t* nameOfExeOrDll);
+	void* DisableModules(const wchar_t* namesOfExeOrDll[], unsigned int count);
+	void* DisableAllModules(const wchar_t* nameOfExeOrDll);
 
 	void WaitForToken(void* token);
 	void TriggerRecompile(void);
@@ -80,9 +82,9 @@ public:
 	void EnableLazyLoadedModule(const wchar_t* fileName, Windows::HMODULE moduleBase);
 	// END EPIC MOD
 
-	void ApplySettingBool(const char* const settingName, int value);
-	void ApplySettingInt(const char* const settingName, int value);
-	void ApplySettingString(const char* const settingName, const wchar_t* const value);
+	void ApplySettingBool(const char* settingName, int value);
+	void ApplySettingInt(const char* settingName, int value);
+	void ApplySettingString(const char* settingName, const wchar_t* value);
 
 	void InstallExceptionHandler(void);
 	ExceptionResult HandleException(EXCEPTION_RECORD* exception, CONTEXT* context, unsigned int threadId);
diff --git a/Engine/Source/Developer/Windows/LiveCoding/Private/External/LC_Commands.h b/Engine/Source/Developer/Windows/LiveCoding/Private/External/LC_Commands.h
--- a/Engine/Source/Developer/Windows/LiveCoding/Private/External/LC_Commands.h
+++ b/Engine/Source/Developer/Windows/LiveCoding/Private/External/LC_Commands.h
@@ -3,6 +3,7 @@
 #pragma once
 
 #include "CoreTypes.h"
+#include "../../../LiveCodingServer/Private/External/LC_Hook.h"
 #include "Windows/WindowsHWrapper.h"
 
 namespace hook
@@ -106,8 +107,9 @@ namespace commands
 	{
 		static const uint32_t ID = LeaveSyncPoint::ID + 1u;
 
-		const hook::Function* first;
-		const hook::Function* last;
+		hook::Type::Enum type;
+		const void* rangeBegin;
+		const void* rangeEnd;
 	};
 
 	// tell the DLL to load a DLL
diff --git a/Engine/Source/Developer/Windows/LiveCoding/Private/External/LC_EntryPoint.cpp b/Engine/Source/Developer/Windows/LiveCoding/Private/External/LC_EntryPoint.cpp
--- a/Engine/Source/Developer/Windows/LiveCoding/Private/External/LC_EntryPoint.cpp
+++ b/Engine/Source/Developer/Windows/LiveCoding/Private/External/LC_EntryPoint.cpp
@@ -64,7 +64,7 @@ BOOL WINAPI DllMain(_In_ HINSTANCE hinstDLL, _In_ DWORD dwReason, _In_ LPVOID /*
 // END EPIC MOD - Internalizing API
 
 
-LPP_DLL_API void __cdecl LppRegisterProcessGroup(const char* const groupName)
+LPP_DLL_API void __cdecl LppRegisterProcessGroup(const char* groupName)
 {
 	// now that we have the process group name, start Live++.
 	// ensure that initialization can happen only once, even if the user calls this more than once.
@@ -85,28 +85,42 @@ LPP_DLL_API void __cdecl LppRegisterProcessGroup(const char* const groupName)
 }
 
 
-LPP_DLL_API void* __cdecl LppEnableModule(const wchar_t* const nameOfExeOrDll)
+LPP_DLL_API void* __cdecl LppEnableModule(const wchar_t* nameOfExeOrDll)
 {
 	// hand command creation to the user command thread
 	return g_startupThread->EnableModule(nameOfExeOrDll);
 }
 
 
-LPP_DLL_API void* __cdecl LppEnableAllModules(const wchar_t* const nameOfExeOrDll)
+LPP_DLL_API void* __cdecl LppEnableModules(const wchar_t* namesOfExeOrDll[], unsigned int count)
+{
+	// hand command creation to the user command thread
+	return g_startupThread->EnableModules(namesOfExeOrDll, count);
+}
+
+
+LPP_DLL_API void* __cdecl LppEnableAllModules(const wchar_t* nameOfExeOrDll)
 {
 	// hand command creation to the user command thread
 	return g_startupThread->EnableAllModules(nameOfExeOrDll);
 }
 
 
-LPP_DLL_API void* __cdecl LppDisableModule(const wchar_t* const nameOfExeOrDll)
+LPP_DLL_API void* __cdecl LppDisableModule(const wchar_t* nameOfExeOrDll)
 {
 	// hand command creation to the user command thread
 	return g_startupThread->DisableModule(nameOfExeOrDll);
 }
 
 
-LPP_DLL_API void* __cdecl LppDisableAllModules(const wchar_t* const nameOfExeOrDll)
+LPP_DLL_API void* __cdecl LppDisableModules(const wchar_t* namesOfExeOrDll[], unsigned int count)
+{
+	// hand command creation to the user command thread
+	return g_startupThread->DisableModules(namesOfExeOrDll, count);
+}
+
+
+LPP_DLL_API void* __cdecl LppDisableAllModules(const wchar_t* nameOfExeOrDll)
 {
 	// hand command creation to the user command thread
 	return g_startupThread->DisableAllModules(nameOfExeOrDll);
@@ -183,28 +197,29 @@ LPP_DLL_API void __cdecl LppSetBuildArguments(const wchar_t* arguments)
 // END EPIC MOD
 
 // BEGIN EPIC MOD - Support for lazy-loading modules
-LPP_DLL_API void __cdecl LppEnableLazyLoadedModule(const wchar_t* const nameOfExeOrDll)
+LPP_DLL_API void __cdecl LppEnableLazyLoadedModule(const wchar_t* nameOfExeOrDll)
 {
 	HMODULE baseAddress = GetModuleHandle(nameOfExeOrDll);
 	g_startupThread->EnableLazyLoadedModule(nameOfExeOrDll, baseAddress);
 }
 // END EPIC MOD
 
-LPP_DLL_API void __cdecl LppApplySettingBool(const char* const settingName, int value)
+
+LPP_DLL_API void __cdecl LppApplySettingBool(const char* settingName, int value)
 {
 	// hand command creation to the user command thread
 	g_startupThread->ApplySettingBool(settingName, value);
 }
 
 
-LPP_DLL_API void __cdecl LppApplySettingInt(const char* const settingName, int value)
+LPP_DLL_API void __cdecl LppApplySettingInt(const char* settingName, int value)
 {
 	// hand command creation to the user command thread
 	g_startupThread->ApplySettingInt(settingName, value);
 }
 
 
-LPP_DLL_API void __cdecl LppApplySettingString(const char* const settingName, const wchar_t* const value)
+LPP_DLL_API void __cdecl LppApplySettingString(const char* settingName, const wchar_t* value)
 {
 	// hand command creation to the user command thread
 	g_startupThread->ApplySettingString(settingName, value);
diff --git a/Engine/Source/Developer/Windows/LiveCoding/Private/External/LC_EntryPoint.h b/Engine/Source/Developer/Windows/LiveCoding/Private/External/LC_EntryPoint.h
--- a/Engine/Source/Developer/Windows/LiveCoding/Private/External/LC_EntryPoint.h
+++ b/Engine/Source/Developer/Windows/LiveCoding/Private/External/LC_EntryPoint.h
@@ -2,11 +2,12 @@
 
 #pragma once
 
-void __cdecl LppRegisterProcessGroup(const char* const groupName);
-void* __cdecl LppEnableModule(const wchar_t* const nameOfExeOrDll);
-void* __cdecl LppEnableAllModules(const wchar_t* const nameOfExeOrDll);
-void* __cdecl LppDisableModule(const wchar_t* const nameOfExeOrDll);
-void* __cdecl LppDisableAllModules(const wchar_t* const nameOfExeOrDll);
+void __cdecl LppRegisterProcessGroup(const char* groupName);
+void* __cdecl LppEnableModule(const wchar_t* nameOfExeOrDll);
+void* __cdecl LppEnableModules(const wchar_t* namesOfExeOrDll[], unsigned int count);
+void* __cdecl LppEnableAllModules(const wchar_t* nameOfExeOrDll);
+void* __cdecl LppDisableModule(const wchar_t* nameOfExeOrDll);
+void* __cdecl LppDisableAllModules(const wchar_t* nameOfExeOrDll);
 void __cdecl LppWaitForToken(void* token);
 void __cdecl LppTriggerRecompile(void);
 void __cdecl LppInstallExceptionHandler(void);
@@ -14,11 +15,11 @@ void __cdecl LppUseExternalBuildSystem(void);
 void __cdecl LppShowConsole();
 void __cdecl LppSetVisible(bool visible);
 void __cdecl LppSetActive(bool active);
-void __cdecl LppSetBuildArguments(const wchar_t* const arguments);
-void __cdecl LppApplySettingBool(const char* const settingName, int value);
-void __cdecl LppApplySettingInt(const char* const settingName, int value);
-void __cdecl LppApplySettingString(const char* const settingName, const wchar_t* const value);
+void __cdecl LppSetBuildArguments(const wchar_t* arguments);
+void __cdecl LppApplySettingBool(const char* settingName, int value);
+void __cdecl LppApplySettingInt(const char* settingName, int value);
+void __cdecl LppApplySettingString(const char* settingName, const wchar_t* const value);
 
 // BEGIN EPIC MOD - Support for lazy-loading modules
-void __cdecl LppEnableLazyLoadedModule(const wchar_t* const nameOfExeOrDll);
+void __cdecl LppEnableLazyLoadedModule(const wchar_t* nameOfExeOrDll);
 // END EPIC MODS
\ No newline at end of file
diff --git a/Engine/Source/Developer/Windows/LiveCoding/Private/External/LPP_API.h b/Engine/Source/Developer/Windows/LiveCoding/Private/External/LPP_API.h
--- a/Engine/Source/Developer/Windows/LiveCoding/Private/External/LPP_API.h
+++ b/Engine/Source/Developer/Windows/LiveCoding/Private/External/LPP_API.h
@@ -15,11 +15,12 @@
 #define LPP_IDENTIFIER(_identifier)					LPP_CONCATENATE(_identifier, __LINE__)
 
 // custom section names for hooks
-#define LPP_PREPATCH_SECTION			".lc_prepatch_hooks"
-#define LPP_POSTPATCH_SECTION			".lc_postpatch_hooks"
-#define LPP_COMPILE_START_SECTION		".lc_compile_start_hooks"
-#define LPP_COMPILE_SUCCESS_SECTION		".lc_compile_success_hooks"
-#define LPP_COMPILE_ERROR_SECTION		".lc_compile_error_hooks"
+#define LPP_PREPATCH_SECTION					".lpp_prepatch_hooks"
+#define LPP_POSTPATCH_SECTION					".lpp_postpatch_hooks"
+#define LPP_COMPILE_START_SECTION				".lpp_compile_start_hooks"
+#define LPP_COMPILE_SUCCESS_SECTION				".lpp_compile_success_hooks"
+#define LPP_COMPILE_ERROR_SECTION				".lpp_compile_error_hooks"
+#define LPP_COMPILE_ERROR_MESSAGE_SECTION		".lpp_compile_error_message_hooks"
 
 // register a pre-patch hook in a custom section
 #define LPP_PREPATCH_HOOK(_function)																							\
@@ -42,6 +43,12 @@
 	__declspec(allocate(LPP_COMPILE_SUCCESS_SECTION)) extern void (*LPP_IDENTIFIER(lpp_compile_success_hook_function))(void) = &_function
 
 // register a compile error hook in a custom section
-#define LPP_COMPILE_ERROR_HOOK(_function)																								\
-	__pragma(section(LPP_COMPILE_ERROR_SECTION, read))																					\
+#define LPP_COMPILE_ERROR_HOOK(_function)																										\
+	__pragma(section(LPP_COMPILE_ERROR_SECTION, read))																							\
 	__declspec(allocate(LPP_COMPILE_ERROR_SECTION)) extern void (*LPP_IDENTIFIER(lpp_compile_error_hook_function))(void) = &_function
+
+// register a compile error message hook in a custom section
+#define LPP_COMPILE_ERROR_MESSAGE_HOOK(_function)																												\
+	__pragma(section(LPP_COMPILE_ERROR_MESSAGE_SECTION, read))																									\
+	__declspec(allocate(LPP_COMPILE_ERROR_MESSAGE_SECTION)) extern void (*LPP_IDENTIFIER(lpp_compile_error_message_hook_function))(const wchar_t*) = &_function
+
diff --git a/Engine/Source/Developer/Windows/LiveCodingServer/Private/External/LC_Hook.cpp b/Engine/Source/Developer/Windows/LiveCodingServer/Private/External/LC_Hook.cpp
--- a/Engine/Source/Developer/Windows/LiveCodingServer/Private/External/LC_Hook.cpp
+++ b/Engine/Source/Developer/Windows/LiveCodingServer/Private/External/LC_Hook.cpp
@@ -1,7 +1,6 @@
 // Copyright 2011-2019 Molecular Matters GmbH, all rights reserved.
 
 #include "LC_Hook.h"
-#include "LC_PointerUtil.h"
 #include "LC_Symbols.h"
 
 
@@ -27,9 +26,3 @@ uint32_t hook::FindLastInSection(const symbols::ImageSectionDB* imageSectionDb,
 
 	return 0u;
 }
-
-
-const hook::Function* hook::MakeFunction(const void* moduleBase, uint32_t rva)
-{
-	return pointer::Offset<const hook::Function*>(moduleBase, rva);
-}
diff --git a/Engine/Source/Developer/Windows/LiveCodingServer/Private/External/LC_Hook.h b/Engine/Source/Developer/Windows/LiveCodingServer/Private/External/LC_Hook.h
--- a/Engine/Source/Developer/Windows/LiveCodingServer/Private/External/LC_Hook.h
+++ b/Engine/Source/Developer/Windows/LiveCodingServer/Private/External/LC_Hook.h
@@ -15,10 +15,46 @@ namespace symbols
 
 namespace hook
 {
-	typedef void (*Function)(void);
+	struct Type
+	{
+		enum Enum
+		{
+			PREPATCH = 0,
+			POSTPATCH,
+			COMPILE_START,
+			COMPILE_SUCCESS,
+			COMPILE_ERROR,
+			COMPILE_ERROR_MESSAGE
+		};
+	};
+
+	typedef void (*PrepatchFunction)(void);
+	typedef void (*PostpatchFunction)(void);
+
+	typedef void (*CompileStartFunction)(void);
+	typedef void (*CompileSuccessFunction)(void);
+	typedef void (*CompileErrorFunction)(void);
+	typedef void (*CompileErrorMessageFunction)(const wchar_t*);
 
 	uint32_t FindFirstInSection(const symbols::ImageSectionDB* imageSectionDb, const ImmutableString& sectionName);
 	uint32_t FindLastInSection(const symbols::ImageSectionDB* imageSectionDb, const ImmutableString& sectionName);
 
-	const Function* MakeFunction(const void* moduleBase, uint32_t rva);
+
+	// calls arbitrary hooks in a given range
+	template <typename T, typename... Args>
+	void CallHooksInRange(const void* rangeBegin, const void* rangeEnd, Args&&... args)
+	{
+		const T* firstHook = static_cast<const T*>(rangeBegin);
+		const T* lastHook = static_cast<const T*>(rangeEnd);
+
+		for (const T* hook = firstHook; hook < lastHook; ++hook)
+		{
+			// note that sections are often padded with zeroes, so skip everything that's zero
+			T function = *hook;
+			if (function)
+			{
+				function(args...);
+			}
+		}
+	}
 }
diff --git a/Engine/Source/Developer/Windows/LiveCodingServer/Private/External/LC_LiveModule.cpp b/Engine/Source/Developer/Windows/LiveCodingServer/Private/External/LC_LiveModule.cpp
--- a/Engine/Source/Developer/Windows/LiveCodingServer/Private/External/LC_LiveModule.cpp
+++ b/Engine/Source/Developer/Windows/LiveCodingServer/Private/External/LC_LiveModule.cpp
@@ -195,7 +195,22 @@ namespace
 		return coff::SymbolRemovalStrategy::MSVC_COMPATIBLE;
 	}
 
-	static LiveModule::CompileResult Compile(const symbols::ObjPath& normalizedObjPath, symbols::Compiland* compiland, const LiveModule::PerProcessData* processData, size_t processCount, unsigned int flags, LiveModule::UpdateType::Enum updateType)
+
+	// helper function for creating a CallHooks command
+	static commands::CallHooks MakeCallHooksCommand(hook::Type::Enum type, const void* moduleBase, uint32_t firstRva, uint32_t lastRva)
+	{
+		return commands::CallHooks { type, pointer::Offset<const void*>(moduleBase, firstRva), pointer::Offset<const void*>(moduleBase, lastRva) };
+	}
+
+
+	// helper function for creating a CallHooks command
+	static commands::CallHooks MakeCallHooksCommand(hook::Type::Enum type, const void* moduleBase, const ModuleCache::FindHookData& hookData)
+	{
+		return MakeCallHooksCommand(type, moduleBase, hookData.firstRva, hookData.lastRva);
+	}
+
+
+	static LiveModule::CompileResult Compile(ModuleCache* moduleCache, const symbols::ObjPath& normalizedObjPath, symbols::Compiland* compiland, const LiveModule::PerProcessData* processData, size_t processCount, unsigned int flags, LiveModule::UpdateType::Enum updateType)
 	{
 		const std::wstring compilerPath = GetCompilerPath(compiland);
 
@@ -332,10 +347,10 @@ namespace
 			logging::LogNoFormat<logging::Channel::DEV>("\n");
 		}
 
-		// send compiler output to main executable
 		{
 			CriticalSection::ScopedLock lock(&g_compileOutputCS);
 
+			// log to local UI
 			logging::LogNoFormat<logging::Channel::USER>(compilerOutput);
 
 			const size_t outputLength = processContext->stdoutData.length();
@@ -343,6 +358,7 @@ namespace
 			{
 				if (updateType != LiveModule::UpdateType::NO_CLIENT_COMMUNICATION)
 				{
+					// send log to host DLL
 					for (size_t p = 0u; p < processCount; ++p)
 					{
 						const DuplexPipe* pipe = processData[p].liveProcess->GetPipe();
@@ -350,6 +366,25 @@ namespace
 						commands::LogOutput cmd {};
 						pipe->SendCommandAndWaitForAck(cmd, compilerOutput, (outputLength + 1u) * sizeof(wchar_t));
 					}
+
+					// send log to all registered hooks in case compilation was not successful
+					if (exitCode != 0u)
+					{
+						const ModuleCache::FindHookData& hookData = moduleCache->FindHooksInSectionBackwards(ModuleCache::SEARCH_ALL_MODULES, ImmutableString(LPP_COMPILE_ERROR_MESSAGE_SECTION));
+						if ((hookData.firstRva != 0u) && (hookData.lastRva != 0u))
+						{
+							const size_t count = hookData.data->processes.size();
+							for (size_t p = 0u; p < count; ++p)
+							{
+								const ModuleCache::ProcessData& hookProcessData = hookData.data->processes[p];
+
+								void* moduleBase = hookProcessData.moduleBase;
+								const DuplexPipe* pipe = hookProcessData.pipe;
+
+								pipe->SendCommandAndWaitForAck(MakeCallHooksCommand(hook::Type::COMPILE_ERROR_MESSAGE, moduleBase, hookData), compilerOutput, (outputLength + 1u) * sizeof(wchar_t));
+							}
+						}
+					}
 				}
 			}
 		}
@@ -825,7 +860,7 @@ namespace
 				const DuplexPipe* pipe = processData.pipe;
 
 				LC_LOG_USER("Calling compile start hooks (PID: %d)", pid);
-				pipe->SendCommandAndWaitForAck(commands::CallHooks { hook::MakeFunction(moduleBase, hookData.firstRva), hook::MakeFunction(moduleBase, hookData.lastRva) }, nullptr, 0u);
+				pipe->SendCommandAndWaitForAck(MakeCallHooksCommand(hook::Type::COMPILE_START, moduleBase, hookData), nullptr, 0u);
 			}
 		}
 	}
@@ -851,8 +886,8 @@ namespace
 				void* moduleBase = processData.moduleBase;
 				const DuplexPipe* pipe = processData.pipe;
 
-				LC_LOG_USER("Calling compile success hooks (PID: %d)", pid);				
-				pipe->SendCommandAndWaitForAck(commands::CallHooks { hook::MakeFunction(moduleBase, hookData.firstRva), hook::MakeFunction(moduleBase, hookData.lastRva) }, nullptr, 0u);
+				LC_LOG_USER("Calling compile success hooks (PID: %d)", pid);
+				pipe->SendCommandAndWaitForAck(MakeCallHooksCommand(hook::Type::COMPILE_SUCCESS, moduleBase, hookData), nullptr, 0u);
 			}
 		}
 	}
@@ -879,7 +914,7 @@ namespace
 				const DuplexPipe* pipe = processData.pipe;
 
 				LC_LOG_USER("Calling compile error hooks (PID: %d)", pid);
-				pipe->SendCommandAndWaitForAck(commands::CallHooks { hook::MakeFunction(moduleBase, hookData.firstRva), hook::MakeFunction(moduleBase, hookData.lastRva) }, nullptr, 0u);
+				pipe->SendCommandAndWaitForAck(MakeCallHooksCommand(hook::Type::COMPILE_ERROR, moduleBase, hookData), nullptr, 0u);
 			}
 		}
 	}
@@ -1641,6 +1676,7 @@ LiveModule::ErrorType::Enum LiveModule::Update(FileAttributeCache* fileCache, Di
 			CompileResult compileResult;
 		};
 
+		ModuleCache* moduleCache = m_moduleCache;
 		double wholeCompileTime = 0.0;
 
 		// now figure out which files can be compiled in parallel.
@@ -1663,10 +1699,10 @@ LiveModule::ErrorType::Enum LiveModule::Update(FileAttributeCache* fileCache, Di
 
 				if (compilerOptions::CreatesPrecompiledHeader(compiland->commandLine.c_str()))
 				{
-					auto task = scheduler::CreateTask(taskRoot, [i, objPath, compiland, processData, processCount, updateType]()
+					auto task = scheduler::CreateTask(taskRoot, [i, moduleCache, objPath, compiland, processData, processCount, updateType]()
 					{
 						telemetry::Scope compileScope("Compile");
-						const CompileResult& result = Compile(objPath, compiland, processData, processCount, 0u, updateType);
+						const CompileResult& result = Compile(moduleCache, objPath, compiland, processData, processCount, 0u, updateType);
 						return LocalCompileResult{ i, compileScope.ReadSeconds(), result };
 					});
 					scheduler::RunTask(task);
@@ -1743,11 +1779,11 @@ LiveModule::ErrorType::Enum LiveModule::Update(FileAttributeCache* fileCache, Di
 
 				if (compilerOptions::UsesC7DebugFormat(compiland->commandLine.c_str()))
 				{
-					auto task = scheduler::CreateTask(taskRoot, [i, objPath, compiland, processData, processCount, updateType]()
+					auto task = scheduler::CreateTask(taskRoot, [i, moduleCache, objPath, compiland, processData, processCount, updateType]()
 					{
 						telemetry::Scope compileScope("Compile");
 
-						const CompileResult& result = Compile(objPath, compiland, processData, processCount, 0u, updateType);
+						const CompileResult& result = Compile(moduleCache, objPath, compiland, processData, processCount, 0u, updateType);
 						return LocalCompileResult{ i, compileScope.ReadSeconds(), result };
 					});
 					scheduler::RunTask(task);
@@ -1834,7 +1870,7 @@ LiveModule::ErrorType::Enum LiveModule::Update(FileAttributeCache* fileCache, Di
 
 					telemetry::Scope compileScope("Compile");
 
-					const CompileResult& result = Compile(objPath, compiland, processData, processCount, 0u, updateType);
+					const CompileResult& result = Compile(moduleCache, objPath, compiland, processData, processCount, 0u, updateType);
 					OnCompiledFile(objPath, compiland, result, compileScope.ReadSeconds(), forceAmalgamationPartsLinkage);
 
 					if (result.exitCode != 0u)
@@ -1862,11 +1898,11 @@ LiveModule::ErrorType::Enum LiveModule::Update(FileAttributeCache* fileCache, Di
 					symbols::Compiland* compiland = availableModifiedFiles[fileIndex].compiland;
 
 					// this PDB file is being written to by one compiland only, we can compile that without any extra options
-					auto task = scheduler::CreateTask(taskRoot, [fileIndex, objPath, compiland, processData, processCount, updateType]()
+					auto task = scheduler::CreateTask(taskRoot, [fileIndex, moduleCache, objPath, compiland, processData, processCount, updateType]()
 					{
 						telemetry::Scope compileScope("Compile");
 
-						const CompileResult& result = Compile(objPath, compiland, processData, processCount, 0u, updateType);
+						const CompileResult& result = Compile(moduleCache, objPath, compiland, processData, processCount, 0u, updateType);
 						return LocalCompileResult{ fileIndex, compileScope.ReadSeconds(), result };
 					});
 					scheduler::RunTask(task);
@@ -1882,11 +1918,11 @@ LiveModule::ErrorType::Enum LiveModule::Update(FileAttributeCache* fileCache, Di
 						const symbols::ObjPath& objPath = availableModifiedFiles[fileIndex].objPath;
 						symbols::Compiland* compiland = availableModifiedFiles[fileIndex].compiland;
 
-						auto task = scheduler::CreateTask(taskRoot, [fileIndex, objPath, compiland, processData, processCount, updateType]()
+						auto task = scheduler::CreateTask(taskRoot, [fileIndex, moduleCache, objPath, compiland, processData, processCount, updateType]()
 						{
 							telemetry::Scope compileScope("Compile");
 
-							const CompileResult& result = Compile(objPath, compiland, processData, processCount, CompileFlags::SERIALIZE_PDB_ACCESS, updateType);
+							const CompileResult& result = Compile(moduleCache, objPath, compiland, processData, processCount, CompileFlags::SERIALIZE_PDB_ACCESS, updateType);
 							return LocalCompileResult{ fileIndex, compileScope.ReadSeconds(), result };
 						});
 						scheduler::RunTask(task);
@@ -3549,7 +3585,7 @@ LiveModule::ErrorType::Enum LiveModule::Update(FileAttributeCache* fileCache, Di
 					const DuplexPipe* pipe = hookProcessData.pipe;
 
 					LC_LOG_USER("Calling pre-patch hooks (PID: %d)", pid);
-					pipe->SendCommandAndWaitForAck(commands::CallHooks { hook::MakeFunction(moduleBase, hookData.firstRva), hook::MakeFunction(moduleBase, hookData.lastRva) }, nullptr, 0u);
+					pipe->SendCommandAndWaitForAck(MakeCallHooksCommand(hook::Type::PREPATCH, moduleBase, hookData), nullptr, 0u);
 				}
 
 				compiledModulePatch->RegisterPrePatchHooks(hookData.data->index, hookData.firstRva, hookData.lastRva);
@@ -4112,7 +4148,7 @@ LiveModule::ErrorType::Enum LiveModule::Update(FileAttributeCache* fileCache, Di
 					const DuplexPipe* pipe = hookData.data->processes[p].pipe;
 
 					LC_LOG_USER("Calling post-patch hooks (PID: %d)", pid);
-					pipe->SendCommandAndWaitForAck(commands::CallHooks { hook::MakeFunction(moduleBase, hookData.firstRva), hook::MakeFunction(moduleBase, hookData.lastRva) }, nullptr, 0u);
+					pipe->SendCommandAndWaitForAck(MakeCallHooksCommand(hook::Type::POSTPATCH, moduleBase, hookData), nullptr, 0u);
 				}
 
 				compiledModulePatch->RegisterPostPatchHooks(hookData.data->index, hookData.firstRva, hookData.lastRva);
@@ -4270,7 +4306,7 @@ bool LiveModule::InstallCompiledPatches(LiveProcess* liveProcess, void* original
 		LC_LOG_DEV("Calling pre-patch hooks");
 		{
 			void* hookModule = processModuleBases[patchData.prePatchHookModuleIndex];
-			pipe->SendCommandAndWaitForAck(commands::CallHooks { hook::MakeFunction(hookModule, patchData.firstPrePatchHook), hook::MakeFunction(hookModule, patchData.lastPrePatchHook) }, nullptr, 0u);
+			pipe->SendCommandAndWaitForAck(MakeCallHooksCommand(hook::Type::PREPATCH, hookModule, patchData.firstPrePatchHook, patchData.lastPrePatchHook), nullptr, 0u);
 		}
 
 
@@ -4355,7 +4391,7 @@ bool LiveModule::InstallCompiledPatches(LiveProcess* liveProcess, void* original
 		LC_LOG_DEV("Calling post-patch hooks");
 		{
 			void* hookModule = processModuleBases[patchData.postPatchHookModuleIndex];
-			pipe->SendCommandAndWaitForAck(commands::CallHooks { hook::MakeFunction(hookModule, patchData.firstPostPatchHook), hook::MakeFunction(hookModule, patchData.lastPostPatchHook) }, nullptr, 0u);
+			pipe->SendCommandAndWaitForAck(MakeCallHooksCommand(hook::Type::POSTPATCH, hookModule, patchData.firstPostPatchHook, patchData.lastPostPatchHook), nullptr, 0u);
 		}
 
 
diff --git a/Engine/Source/Developer/Windows/LiveCodingServer/Private/External/LC_Symbols.cpp b/Engine/Source/Developer/Windows/LiveCodingServer/Private/External/LC_Symbols.cpp
--- a/Engine/Source/Developer/Windows/LiveCodingServer/Private/External/LC_Symbols.cpp
+++ b/Engine/Source/Developer/Windows/LiveCodingServer/Private/External/LC_Symbols.cpp
@@ -301,7 +301,9 @@ namespace
 			return nullptr;
 		}
 
-		symbols::Provider* provider = new symbols::Provider { diaDataSource, diaSession, globalScope };
+		const file::Attributes& attributes = file::GetAttributes(filename);
+		const uint64_t lastModification = file::GetLastModificationTime(attributes);
+		symbols::Provider* provider = new symbols::Provider { diaDataSource, diaSession, globalScope, lastModification };
 		return provider;
 	}
 
@@ -613,7 +615,7 @@ namespace symbols
 	}
 
 
-	CompilandDB* GatherCompilands(Provider* provider, const DiaCompilandDB* diaCompilandDb, unsigned int splitAmalgamatedFilesThreshold, uint32_t compilandOptions)
+	CompilandDB* GatherCompilands(const Provider* provider, const DiaCompilandDB* diaCompilandDb, unsigned int splitAmalgamatedFilesThreshold, uint32_t compilandOptions)
 	{
 		telemetry::Scope telemetryScope("Gathering compilands");
 
@@ -840,6 +842,14 @@ namespace symbols
 
 						compilandPath = testPath;
 					}
+
+					// ignore compilands that are newer than the module itself.
+					// we cannot use those for reconstructing symbol information, because they weren't linked into the executable.
+					if (cacheData.lastModificationTime > provider->lastModificationTime)
+					{
+						LC_WARNING_USER("Ignoring compiland %S because it is newer than the module it belongs to.", compilandPath.c_str());
+						continue;
+					}
 				}
 
 				const std::wstring normalizedCompilandPath = file::NormalizePath(compilandPath.c_str());
diff --git a/Engine/Source/Developer/Windows/LiveCodingServer/Private/External/LC_Symbols.h b/Engine/Source/Developer/Windows/LiveCodingServer/Private/External/LC_Symbols.h
--- a/Engine/Source/Developer/Windows/LiveCodingServer/Private/External/LC_Symbols.h
+++ b/Engine/Source/Developer/Windows/LiveCodingServer/Private/External/LC_Symbols.h
@@ -187,6 +187,7 @@ namespace symbols
 		IDiaDataSource* diaDataSource;
 		IDiaSession* diaSession;
 		IDiaSymbol* globalScope;
+		TimeStamp lastModificationTime;
 	};
 
 	struct DynamicInitializerDB
@@ -249,7 +250,7 @@ namespace symbols
 	};
 
 	// CACHE: stored only ONCE, after initial load, cannot change
-	CompilandDB* GatherCompilands(Provider* provider, const DiaCompilandDB* diaCompilandDb, unsigned int splitAmalgamatedFilesThreshold, uint32_t compilandOptions);
+	CompilandDB* GatherCompilands(const Provider* provider, const DiaCompilandDB* diaCompilandDb, unsigned int splitAmalgamatedFilesThreshold, uint32_t compilandOptions);
 
 	// CACHE: stored only ONCE, after initial load, cannot change
 	LibraryDB* GatherLibraries(const DiaCompilandDB* diaCompilandDb);
