Live coding: Add multiple modules in one call at startup, rather than adding each module individuallly.

diff --git a/Engine/Source/Developer/Windows/LiveCoding/Private/LiveCodingModule.cpp b/Engine/Source/Developer/Windows/LiveCoding/Private/LiveCodingModule.cpp
--- a/Engine/Source/Developer/Windows/LiveCoding/Private/LiveCodingModule.cpp
+++ b/Engine/Source/Developer/Windows/LiveCoding/Private/LiveCodingModule.cpp
@@ -263,46 +263,62 @@ void FLiveCodingModule::UpdateModules()
 	TArray<FModuleStatus> ModuleStatuses;
 	FModuleManager::Get().QueryModules(ModuleStatuses);
 
+	TArray<FString> EnableModules;
 	for (const FModuleStatus& ModuleStatus : ModuleStatuses)
 	{
 		if (ModuleStatus.bIsLoaded)
 		{
-			FString FullFilePath = FPaths::ConvertRelativePathToFull(ModuleStatus.FilePath);
-			ConfigureModule(FName(*ModuleStatus.Name), FullFilePath);
+			FName ModuleName(*ModuleStatus.Name);
+			if (!ConfiguredModules.Contains(ModuleName))
+			{
+				FString FullFilePath = FPaths::ConvertRelativePathToFull(ModuleStatus.FilePath);
+				if (ShouldPreloadModule(ModuleName, FullFilePath))
+				{
+					EnableModules.Add(FullFilePath);
+				}
+				else
+				{
+					LppEnableLazyLoadedModule(*FullFilePath);
+				}
+				ConfiguredModules.Add(ModuleName);
+			}
 		}
 	}
-#endif
-}
 
-void FLiveCodingModule::OnModulesChanged(FName ModuleName, EModuleChangeReason Reason)
-{
-#if !IS_MONOLITHIC
-	if (Reason == EModuleChangeReason::ModuleLoaded)
+	if (EnableModules.Num() > 0)
 	{
-		FModuleStatus Status;
-		if (FModuleManager::Get().QueryModule(ModuleName, Status))
+		TArray<const TCHAR*> EnableModuleFileNames;
+		for (const FString& EnableModule : EnableModules)
 		{
-			FString FullFilePath = FPaths::ConvertRelativePathToFull(Status.FilePath);
-			ConfigureModule(ModuleName, FullFilePath);
+			EnableModuleFileNames.Add(*EnableModule);
 		}
+		LppEnableModules(EnableModuleFileNames.GetData(), EnableModuleFileNames.Num());
 	}
 #endif
 }
 
-void FLiveCodingModule::ConfigureModule(const FName& Name, const FString& FullFilePath)
+void FLiveCodingModule::OnModulesChanged(FName ModuleName, EModuleChangeReason Reason)
 {
 #if !IS_MONOLITHIC
-	if (!ConfiguredModules.Contains(Name))
+	if (Reason == EModuleChangeReason::ModuleLoaded)
 	{
-		if (ShouldPreloadModule(Name, FullFilePath))
-		{
-			LppEnableModule(*FullFilePath);
-		}
-		else
+		FModuleStatus Status;
+		if (FModuleManager::Get().QueryModule(ModuleName, Status))
 		{
-			LppEnableLazyLoadedModule(*FullFilePath);
+			if (!ConfiguredModules.Contains(ModuleName))
+			{
+				FString FullFilePath = FPaths::ConvertRelativePathToFull(Status.FilePath);
+				if (ShouldPreloadModule(ModuleName, FullFilePath))
+				{
+					LppEnableModule(*FullFilePath);
+				}
+				else
+				{
+					LppEnableLazyLoadedModule(*FullFilePath);
+				}
+				ConfiguredModules.Add(ModuleName);
+			}
 		}
-		ConfiguredModules.Add(Name);
 	}
 #endif
 }
diff --git a/Engine/Source/Developer/Windows/LiveCoding/Private/LiveCodingModule.h b/Engine/Source/Developer/Windows/LiveCoding/Private/LiveCodingModule.h
--- a/Engine/Source/Developer/Windows/LiveCoding/Private/LiveCodingModule.h
+++ b/Engine/Source/Developer/Windows/LiveCoding/Private/LiveCodingModule.h
@@ -56,7 +56,6 @@ private:
 
 	void UpdateModules();
 
-	void ConfigureModule(const FName& Name, const FString& FullFilePath);
 	bool ShouldPreloadModule(const FName& Name, const FString& FullFilePath) const;
 };
 
