Live Coding: Enable live coding for modules during the first Tick() call, to prevent other modules loaded during initialization being added later.

diff --git a/Engine/Source/Developer/Windows/LiveCoding/Private/LiveCodingModule.cpp b/Engine/Source/Developer/Windows/LiveCoding/Private/LiveCodingModule.cpp
--- a/Engine/Source/Developer/Windows/LiveCoding/Private/LiveCodingModule.cpp
+++ b/Engine/Source/Developer/Windows/LiveCoding/Private/LiveCodingModule.cpp
@@ -27,6 +27,7 @@ FLiveCodingModule::FLiveCodingModule()
 	: bEnabledLastTick(false)
 	, bEnabledForSession(false)
 	, bStarted(false)
+	, bUpdateModulesInTick(false)
 	, FullEnginePluginsDir(FPaths::ConvertRelativePathToFull(FPaths::EnginePluginsDir()))
 	, FullProjectDir(FPaths::ConvertRelativePathToFull(FPaths::ProjectDir()))
 	, FullProjectPluginsDir(FPaths::ConvertRelativePathToFull(FPaths::ProjectPluginsDir()))
@@ -205,6 +206,12 @@ void FLiveCodingModule::Tick()
 		EnableForSession(Settings->bEnabled);
 		bEnabledLastTick = Settings->bEnabled;
 	}
+
+	if (bUpdateModulesInTick)
+	{
+		UpdateModules();
+		bUpdateModulesInTick = false;
+	}
 }
 
 bool FLiveCodingModule::StartLiveCoding()
@@ -246,8 +253,15 @@ bool FLiveCodingModule::StartLiveCoding()
 		}
 		LppSetBuildArguments(*Arguments);
 
-		// Configure all the current modules
-		UpdateModules();
+		// Configure all the current modules. For non-commandlets, schedule it to be done in the first Tick() so we can batch everything together.
+		if (IsRunningCommandlet())
+		{
+			UpdateModules();
+		}
+		else
+		{
+			bUpdateModulesInTick = true;
+		}
 
 		// Register a delegate to listen for new modules loaded from this point onwards
 		ModulesChangedDelegateHandle = FModuleManager::Get().OnModulesChanged().AddRaw(this, &FLiveCodingModule::OnModulesChanged);
@@ -261,46 +275,49 @@ bool FLiveCodingModule::StartLiveCoding()
 
 void FLiveCodingModule::UpdateModules()
 {
+	if (bEnabledForSession)
+	{
 #if IS_MONOLITHIC
-	wchar_t FullFilePath[WINDOWS_MAX_PATH];
-	verify(GetModuleFileName(hInstance, FullFilePath, ARRAY_COUNT(FullFilePath)));
-	LppEnableModule(FullFilePath);
+		wchar_t FullFilePath[WINDOWS_MAX_PATH];
+		verify(GetModuleFileName(hInstance, FullFilePath, ARRAY_COUNT(FullFilePath)));
+		LppEnableModule(FullFilePath);
 #else
-	TArray<FModuleStatus> ModuleStatuses;
-	FModuleManager::Get().QueryModules(ModuleStatuses);
+		TArray<FModuleStatus> ModuleStatuses;
+		FModuleManager::Get().QueryModules(ModuleStatuses);
 
-	TArray<FString> EnableModules;
-	for (const FModuleStatus& ModuleStatus : ModuleStatuses)
-	{
-		if (ModuleStatus.bIsLoaded)
+		TArray<FString> EnableModules;
+		for (const FModuleStatus& ModuleStatus : ModuleStatuses)
 		{
-			FName ModuleName(*ModuleStatus.Name);
-			if (!ConfiguredModules.Contains(ModuleName))
+			if (ModuleStatus.bIsLoaded)
 			{
-				FString FullFilePath = FPaths::ConvertRelativePathToFull(ModuleStatus.FilePath);
-				if (ShouldPreloadModule(ModuleName, FullFilePath))
+				FName ModuleName(*ModuleStatus.Name);
+				if (!ConfiguredModules.Contains(ModuleName))
 				{
-					EnableModules.Add(FullFilePath);
+					FString FullFilePath = FPaths::ConvertRelativePathToFull(ModuleStatus.FilePath);
+					if (ShouldPreloadModule(ModuleName, FullFilePath))
+					{
+						EnableModules.Add(FullFilePath);
+					}
+					else
+					{
+						LppEnableLazyLoadedModule(*FullFilePath);
+					}
+					ConfiguredModules.Add(ModuleName);
 				}
-				else
-				{
-					LppEnableLazyLoadedModule(*FullFilePath);
-				}
-				ConfiguredModules.Add(ModuleName);
 			}
 		}
-	}
 
-	if (EnableModules.Num() > 0)
-	{
-		TArray<const TCHAR*> EnableModuleFileNames;
-		for (const FString& EnableModule : EnableModules)
+		if (EnableModules.Num() > 0)
 		{
-			EnableModuleFileNames.Add(*EnableModule);
+			TArray<const TCHAR*> EnableModuleFileNames;
+			for (const FString& EnableModule : EnableModules)
+			{
+				EnableModuleFileNames.Add(*EnableModule);
+			}
+			LppEnableModules(EnableModuleFileNames.GetData(), EnableModuleFileNames.Num());
 		}
-		LppEnableModules(EnableModuleFileNames.GetData(), EnableModuleFileNames.Num());
-	}
 #endif
+	}
 }
 
 void FLiveCodingModule::OnModulesChanged(FName ModuleName, EModuleChangeReason Reason)
@@ -308,22 +325,14 @@ void FLiveCodingModule::OnModulesChanged(FName ModuleName, EModuleChangeReason R
 #if !IS_MONOLITHIC
 	if (Reason == EModuleChangeReason::ModuleLoaded)
 	{
-		FModuleStatus Status;
-		if (FModuleManager::Get().QueryModule(ModuleName, Status))
+		// Assume that Tick() won't be called if we're running a commandlet
+		if (IsRunningCommandlet())
 		{
-			if (!ConfiguredModules.Contains(ModuleName))
-			{
-				FString FullFilePath = FPaths::ConvertRelativePathToFull(Status.FilePath);
-				if (ShouldPreloadModule(ModuleName, FullFilePath))
-				{
-					LppEnableModule(*FullFilePath);
-				}
-				else
-				{
-					LppEnableLazyLoadedModule(*FullFilePath);
-				}
-				ConfiguredModules.Add(ModuleName);
-			}
+			UpdateModules();
+		}
+		else
+		{
+			bUpdateModulesInTick = true;
 		}
 	}
 #endif
diff --git a/Engine/Source/Developer/Windows/LiveCoding/Private/LiveCodingModule.h b/Engine/Source/Developer/Windows/LiveCoding/Private/LiveCodingModule.h
--- a/Engine/Source/Developer/Windows/LiveCoding/Private/LiveCodingModule.h
+++ b/Engine/Source/Developer/Windows/LiveCoding/Private/LiveCodingModule.h
@@ -39,6 +39,7 @@ private:
 	bool bEnabledLastTick;
 	bool bEnabledForSession;
 	bool bStarted;
+	bool bUpdateModulesInTick;
 	TSet<FName> ConfiguredModules;
 
 	const FString FullEnginePluginsDir;
