Live Coding: Fix crash when lazy loading a module with changes.

#rnx

diff --git a/Engine/Source/Developer/Windows/LiveCodingServer/Private/External/LC_LiveModule.cpp b/Engine/Source/Developer/Windows/LiveCodingServer/Private/External/LC_LiveModule.cpp
--- a/Engine/Source/Developer/Windows/LiveCodingServer/Private/External/LC_LiveModule.cpp
+++ b/Engine/Source/Developer/Windows/LiveCodingServer/Private/External/LC_LiveModule.cpp
@@ -459,7 +459,7 @@ namespace
 		for (std::pair<const ImmutableString, symbols::Compiland*>& Pair : compilands)
 		{
 			symbols::Compiland* Compiland = Pair.second;
-			if (Compiland->amalgamatedUniqueId == ~(uint32_t)0)
+			if (Compiland != nullptr && Compiland->amalgamatedUniqueId == ~(uint32_t)0)
 			{
 				const std::wstring& wideObjPath = string::ToWideString(Pair.first);
 				Directories.insert(file::GetDirectory(wideObjPath));
@@ -476,7 +476,7 @@ namespace
 		for (std::pair<const ImmutableString, symbols::Compiland*>& Pair : compilands)
 		{
 			symbols::Compiland* Compiland = Pair.second;
-			if (Compiland->amalgamatedUniqueId == ~(uint32_t)0)
+			if (Compiland != nullptr && Compiland->amalgamatedUniqueId == ~(uint32_t)0)
 			{
 				types::StringMap<uint32_t>::const_iterator Iter = objFileToCompilandId.find(Pair.first);
 				if (Iter == objFileToCompilandId.end())
