LiveCoding: Bump the maximum number of tasks in the scheduler. When enabling a large number of modules at editor startup, we can exceed this and leave all threads deadlocked.

#rnx

diff --git a/Engine/Source/Developer/Windows/LiveCodingServer/Private/External/LC_SchedulerQueue.h b/Engine/Source/Developer/Windows/LiveCodingServer/Private/External/LC_SchedulerQueue.h
--- a/Engine/Source/Developer/Windows/LiveCodingServer/Private/External/LC_SchedulerQueue.h
+++ b/Engine/Source/Developer/Windows/LiveCodingServer/Private/External/LC_SchedulerQueue.h
@@ -16,7 +16,9 @@ namespace scheduler
 	// simple multi-producer, multi-consumer queue
 	class TaskQueue
 	{
-		static const unsigned int TASK_COUNT = 1024u;
+		// BEGIN EPIC MOD - Increasing task count due to hangs when enabling for all editor modules
+		static const unsigned int TASK_COUNT = 4096u;// 1024u;
+		// END EPIC MOD
 		static const unsigned int ACCESS_MASK = TASK_COUNT - 1u;
 
 	public:
