Live Coding: Fix several issues related to using launch on / packaged builds with Live Coding.

* The original .uproject file is now compiled into monolithic executables when Live Coding is enabled. This allows invoking UBT with the original project file when the executable is staged to a different directory. This parameter can be overriden via the LiveCoding.SourceProject cvar.
* The original engine directory is also compiled into the executable. This allows finding the console executable path without having to enter it manually via the LiveCoding.ConsolePath cvar.
* If an exact match for a binary filename is not found, try to find a match by name only. Also required to support staged builds for 'Launch On', etc...
* Add a LiveCoding.Compile command to trigger a compile from the console.

diff --git a/Engine/Source/Developer/Windows/LiveCoding/Private/LiveCodingModule.cpp b/Engine/Source/Developer/Windows/LiveCoding/Private/LiveCodingModule.cpp
--- a/Engine/Source/Developer/Windows/LiveCoding/Private/LiveCodingModule.cpp
+++ b/Engine/Source/Developer/Windows/LiveCoding/Private/LiveCodingModule.cpp
@@ -23,6 +23,11 @@ bool GIsCompileActive = false;
 FString GLiveCodingConsolePath;
 FString GLiveCodingConsoleArguments;
 
+#if IS_MONOLITHIC
+extern const TCHAR* GLiveCodingEngineDir;
+extern const TCHAR* GLiveCodingProject;
+#endif
+
 FLiveCodingModule::FLiveCodingModule()
 	: bEnabledLastTick(false)
 	, bEnabledForSession(false)
@@ -47,19 +52,42 @@ void FLiveCodingModule::StartupModule()
 		ECVF_Cheat
 	);
 
+	CompileCommand = ConsoleManager.RegisterConsoleCommand(
+		TEXT("LiveCoding.Compile"),
+		TEXT("Initiates a live coding compile"),
+		FConsoleCommandDelegate::CreateRaw(this, &FLiveCodingModule::Compile),
+		ECVF_Cheat
+	);
+
+#if IS_MONOLITHIC
+	FString DefaultEngineDir = GLiveCodingEngineDir;
+#else
+	FString DefaultEngineDir = FPaths::EngineDir();
+#endif
 #if USE_DEBUG_LIVE_CODING_CONSOLE
 	static const TCHAR* DefaultConsolePath = TEXT("Binaries/Win64/LiveCodingConsole-Win64-Debug.exe");
 #else
 	static const TCHAR* DefaultConsolePath = TEXT("Binaries/Win64/LiveCodingConsole.exe");
 #endif 
-
 	ConsolePathVariable = ConsoleManager.RegisterConsoleVariable(
 		TEXT("LiveCoding.ConsolePath"),
-		FPaths::ConvertRelativePathToFull(FPaths::EngineDir() / DefaultConsolePath),
+		FPaths::ConvertRelativePathToFull(DefaultEngineDir / DefaultConsolePath),
 		TEXT("Path to the live coding console application"),
 		ECVF_Cheat
 	);
 
+#if IS_MONOLITHIC
+	FString SourceProject = (GLiveCodingProject != nullptr)? GLiveCodingProject : TEXT("");
+#else
+	FString SourceProject = FPaths::IsProjectFilePathSet() ? FPaths::GetProjectFilePath() : TEXT("");
+#endif
+	SourceProjectVariable = ConsoleManager.RegisterConsoleVariable(
+		TEXT("LiveCoding.SourceProject"),
+		FPaths::ConvertRelativePathToFull(SourceProject),
+		TEXT("Path to the project that this target was built from"),
+		ECVF_Cheat
+	);
+
 	EndFrameDelegateHandle = FCoreDelegates::OnEndFrame.AddRaw(this, &FLiveCodingModule::Tick);
 
 	ISettingsModule* SettingsModule = FModuleManager::GetModulePtr<ISettingsModule>("Settings");
@@ -105,7 +133,9 @@ void FLiveCodingModule::ShutdownModule()
 	FCoreDelegates::OnEndFrame.Remove(EndFrameDelegateHandle);
 
 	IConsoleManager& ConsoleManager = IConsoleManager::Get();
+	ConsoleManager.UnregisterConsoleObject(SourceProjectVariable);
 	ConsoleManager.UnregisterConsoleObject(ConsolePathVariable);
+	ConsoleManager.UnregisterConsoleObject(CompileCommand);
 	ConsoleManager.UnregisterConsoleObject(EnableCommand);
 }
 
@@ -233,6 +263,14 @@ bool FLiveCodingModule::StartLiveCoding()
 			return false;
 		}
 
+		// Get the source project filename
+		FString SourceProject = SourceProjectVariable->GetString();
+		if (SourceProject.Len() > 0 && !FPaths::FileExists(SourceProject))
+		{
+			UE_LOG(LogLiveCoding, Error, TEXT("Unable to start live coding session. Unable to find source project file '%s'."), *SourceProject);
+			return false;
+		}
+
 		UE_LOG(LogLiveCoding, Display, TEXT("Starting LiveCoding"));
 
 		// Enable external build system
@@ -247,9 +285,9 @@ bool FLiveCodingModule::StartLiveCoding()
 		Arguments += FString::Printf(TEXT("%s"), FPlatformMisc::GetUBTPlatform());
 		Arguments += FString::Printf(TEXT(" %s"), EBuildConfigurations::ToString(FApp::GetBuildConfiguration()));
 		Arguments += FString::Printf(TEXT(" -TargetType=%s"), FPlatformMisc::GetUBTTarget());
-		if(FPaths::IsProjectFilePathSet())
+		if(SourceProject.Len() > 0)
 		{
-			Arguments += FString::Printf(TEXT(" -Project=\"%s\""), *FPaths::ConvertRelativePathToFull(FPaths::GetProjectFilePath()));
+			Arguments += FString::Printf(TEXT(" -Project=\"%s\""), *FPaths::ConvertRelativePathToFull(SourceProject));
 		}
 		LppSetBuildArguments(*Arguments);
 
diff --git a/Engine/Source/Developer/Windows/LiveCoding/Private/LiveCodingModule.h b/Engine/Source/Developer/Windows/LiveCoding/Private/LiveCodingModule.h
--- a/Engine/Source/Developer/Windows/LiveCoding/Private/LiveCodingModule.h
+++ b/Engine/Source/Developer/Windows/LiveCoding/Private/LiveCodingModule.h
@@ -47,7 +47,9 @@ private:
 	const FString FullProjectPluginsDir;
 
 	IConsoleCommand* EnableCommand;
+	IConsoleCommand* CompileCommand;
 	IConsoleVariable* ConsolePathVariable;
+	IConsoleVariable* SourceProjectVariable;
 	FDelegateHandle EndFrameDelegateHandle;
 	FDelegateHandle ModulesChangedDelegateHandle;
 
diff --git a/Engine/Source/Developer/Windows/LiveCodingServer/Private/External/LC_ServerCommandThread.cpp b/Engine/Source/Developer/Windows/LiveCodingServer/Private/External/LC_ServerCommandThread.cpp
--- a/Engine/Source/Developer/Windows/LiveCodingServer/Private/External/LC_ServerCommandThread.cpp
+++ b/Engine/Source/Developer/Windows/LiveCodingServer/Private/External/LC_ServerCommandThread.cpp
@@ -641,11 +641,27 @@ void ServerCommandThread::CompileChanges(bool didAllProcessesMakeProgress)
 			std::wstring ModuleFileName = file::NormalizePath(*Pair.Key);
 			if(ValidModuleFileNames.find(ModuleFileName) == ValidModuleFileNames.end())
 			{
-				std::wstring ModuleName = file::GetFilename(ModuleFileName);
-				LC_ERROR_USER("Live coding is not enabled for %S.", ModuleName.c_str());
-				LC_ERROR_USER("Configure the list of enabled modules from the Live Coding section of the editor preferences window.");
-				GLiveCodingServer->GetCompileFinishedDelegate().ExecuteIfBound(ELiveCodingResult::Error, *FString::Printf(TEXT("Live coding not enabled for %s"), ModuleName.c_str()));
-				return;
+				// We couldn't find this exact module filename, but this could be a staged executable. See if we can just match the name.
+				std::wstring ModuleFileNameOnly = file::GetFilename(ModuleFileName);
+
+				bool bFoundNameMatch = false;
+				for (const LiveModule* liveModule : m_liveModules)
+				{
+					if (ModuleFileNameOnly == file::GetFilename(liveModule->GetModuleName()))
+					{
+						ModuleFileName = liveModule->GetModuleName();
+						bFoundNameMatch = true;
+						break;
+					}
+				}
+
+				if (!bFoundNameMatch)
+				{
+					LC_ERROR_USER("Live coding is not enabled for %S.", ModuleFileName.c_str());
+					LC_ERROR_USER("Configure the list of enabled modules from the Live Coding section of the editor preferences window.");
+					GLiveCodingServer->GetCompileFinishedDelegate().ExecuteIfBound(ELiveCodingResult::Error, *FString::Printf(TEXT("Live coding not enabled for %s"), ModuleFileName.c_str()));
+					return;
+				}
 			}
 
 			types::vector<LiveModule::ModifiedObjFile> ObjectFiles;
diff --git a/Engine/Source/Programs/UnrealBuildTool/Configuration/TargetDescriptor.cs b/Engine/Source/Programs/UnrealBuildTool/Configuration/TargetDescriptor.cs
--- a/Engine/Source/Programs/UnrealBuildTool/Configuration/TargetDescriptor.cs
+++ b/Engine/Source/Programs/UnrealBuildTool/Configuration/TargetDescriptor.cs
@@ -58,12 +58,6 @@ namespace UnrealBuildTool
 		[CommandLine("-WriteActions=")]
 		public List<FileReference> WriteActionFiles = new List<FileReference>();
 
-		/// <summary>
-		/// Path to the manifest for passing info about the output to live coding
-		/// </summary>
-		[CommandLine("-LiveCodingManifest=")]
-		public FileReference LiveCodingManifest = null;
-
 		/// <summary>
 		/// Constructor
 		/// </summary>
diff --git a/Engine/Source/Programs/UnrealBuildTool/Configuration/UEBuildTarget.cs b/Engine/Source/Programs/UnrealBuildTool/Configuration/UEBuildTarget.cs
--- a/Engine/Source/Programs/UnrealBuildTool/Configuration/UEBuildTarget.cs
+++ b/Engine/Source/Programs/UnrealBuildTool/Configuration/UEBuildTarget.cs
@@ -3418,6 +3418,14 @@ namespace UnrealBuildTool
 			if(Rules.bWithLiveCoding)
 			{
 				GlobalCompileEnvironment.Definitions.Add("WITH_LIVE_CODING=1");
+				if (Rules.LinkType == TargetLinkType.Monolithic)
+				{
+					GlobalCompileEnvironment.Definitions.Add(String.Format("UE_LIVE_CODING_ENGINE_DIR=\"{0}\"", UnrealBuildTool.EngineDirectory.FullName.Replace("\\", "\\\\")));
+					if(ProjectFile != null)
+					{
+						GlobalCompileEnvironment.Definitions.Add(String.Format("UE_LIVE_CODING_PROJECT=\"{0}\"", ProjectFile.FullName.Replace("\\", "\\\\")));
+					}
+				}
 			}
 			else
 			{
diff --git a/Engine/Source/Programs/UnrealBuildTool/Modes/BuildMode.cs b/Engine/Source/Programs/UnrealBuildTool/Modes/BuildMode.cs
--- a/Engine/Source/Programs/UnrealBuildTool/Modes/BuildMode.cs
+++ b/Engine/Source/Programs/UnrealBuildTool/Modes/BuildMode.cs
@@ -76,6 +76,12 @@ namespace UnrealBuildTool
 		[CommandLine("-WriteOutdatedActions=")]
 		public FileReference WriteOutdatedActionsFile = null;
 
+		/// <summary>
+		/// Path to the manifest for passing info about the output to live coding
+		/// </summary>
+		[CommandLine("-LiveCodingManifest=")]
+		public FileReference LiveCodingManifest = null;
+
 		/// <summary>
 		/// Main entry point
 		/// </summary>
@@ -194,7 +200,7 @@ namespace UnrealBuildTool
 					// Create the working set provider
 					using (ISourceFileWorkingSet WorkingSet = SourceFileWorkingSet.Create(UnrealBuildTool.RootDirectory, ProjectDirs))
 					{
-						Build(TargetDescriptors, BuildConfiguration, WorkingSet, Options, WriteOutdatedActionsFile);
+						Build(TargetDescriptors, BuildConfiguration, WorkingSet, Options, LiveCodingManifest, WriteOutdatedActionsFile);
 					}
 				}
 			}
@@ -214,9 +220,10 @@ namespace UnrealBuildTool
 		/// <param name="BuildConfiguration">Current build configuration</param>
 		/// <param name="WorkingSet">The source file working set</param>
 		/// <param name="Options">Additional options for the build</param>
+		/// <param name="LiveCodingManifest">Path to write the live coding manifest to</param>
 		/// <param name="WriteOutdatedActionsFile">Files to write the list of outdated actions to (rather than building them)</param>
 		/// <returns>Result from the compilation</returns>
-		public static void Build(List<TargetDescriptor> TargetDescriptors, BuildConfiguration BuildConfiguration, ISourceFileWorkingSet WorkingSet, BuildOptions Options, FileReference WriteOutdatedActionsFile)
+		public static void Build(List<TargetDescriptor> TargetDescriptors, BuildConfiguration BuildConfiguration, ISourceFileWorkingSet WorkingSet, BuildOptions Options, FileReference LiveCodingManifest, FileReference WriteOutdatedActionsFile)
 		{
 			// Create a makefile for each target
 			TargetMakefile[] Makefiles = new TargetMakefile[TargetDescriptors.Count];
@@ -225,15 +232,11 @@ namespace UnrealBuildTool
 				Makefiles[TargetIdx] = CreateMakefile(BuildConfiguration, TargetDescriptors[TargetIdx], WorkingSet);
 			}
 
-			// Output the manifest
-			for(int TargetIdx = 0; TargetIdx < TargetDescriptors.Count; TargetIdx++)
+			// Output the Live Coding manifest
+			if(LiveCodingManifest != null)
 			{
-				TargetDescriptor TargetDescriptor = TargetDescriptors[TargetIdx];
-				if(TargetDescriptor.LiveCodingManifest != null)
-				{
-					TargetMakefile Makefile = Makefiles[TargetIdx];
-					HotReload.WriteLiveCodingManifest(TargetDescriptor.LiveCodingManifest, Makefile.Actions);
-				}
+				List<Action> AllActions = Makefiles.SelectMany(x => x.Actions).ToList();
+				HotReload.WriteLiveCodingManifest(LiveCodingManifest, AllActions);
 			}
 
 			// Export the actions for each target
diff --git a/Engine/Source/Programs/UnrealBuildTool/System/ExternalExecution.cs b/Engine/Source/Programs/UnrealBuildTool/System/ExternalExecution.cs
--- a/Engine/Source/Programs/UnrealBuildTool/System/ExternalExecution.cs
+++ b/Engine/Source/Programs/UnrealBuildTool/System/ExternalExecution.cs
@@ -1108,7 +1108,7 @@ namespace UnrealBuildTool
 
 						using(Timeline.ScopeEvent("Buildng UnrealHeaderTool"))
 						{
-							BuildMode.Build(TargetDescriptors, BuildConfiguration, WorkingSet, BuildOptions.Quiet, null);
+							BuildMode.Build(TargetDescriptors, BuildConfiguration, WorkingSet, BuildOptions.Quiet, null, null);
 						}
 					}
 
diff --git a/Engine/Source/Runtime/Core/Public/Modules/ModuleManager.h b/Engine/Source/Runtime/Core/Public/Modules/ModuleManager.h
--- a/Engine/Source/Runtime/Core/Public/Modules/ModuleManager.h
+++ b/Engine/Source/Runtime/Core/Public/Modules/ModuleManager.h
@@ -755,6 +755,21 @@ class FDefaultGameModuleImpl
 	#define IMPLEMENT_FOREIGN_ENGINE_DIR() 
 #endif
 
+/**
+ * Macros for setting the source directories for live coding builds. This allows locally packaging a target and patching code into it.
+ */
+#ifdef UE_LIVE_CODING_ENGINE_DIR
+	#define IMPLEMENT_LIVE_CODING_ENGINE_DIR() const TCHAR* GLiveCodingEngineDir = TEXT(UE_LIVE_CODING_ENGINE_DIR);
+	#ifdef UE_LIVE_CODING_PROJECT
+		#define IMPLEMENT_LIVE_CODING_PROJECT() const TCHAR* GLiveCodingProject = TEXT(UE_LIVE_CODING_PROJECT);
+	#else
+		#define IMPLEMENT_LIVE_CODING_PROJECT() const TCHAR* GLiveCodingProject = nullptr;
+	#endif
+#else
+	#define IMPLEMENT_LIVE_CODING_ENGINE_DIR()
+	#define IMPLEMENT_LIVE_CODING_PROJECT()
+#endif
+
 /**
  * Macro for passing a list argument to a macro
  */
@@ -816,6 +831,8 @@ class FDefaultGameModuleImpl
 			/* For monolithic builds, we must statically define the game's name string (See Core.h) */ \
 			TCHAR GInternalProjectName[64] = TEXT( GameName ); \
 			IMPLEMENT_FOREIGN_ENGINE_DIR() \
+			IMPLEMENT_LIVE_CODING_ENGINE_DIR() \
+			IMPLEMENT_LIVE_CODING_PROJECT() \
 			IMPLEMENT_SIGNING_KEY_REGISTRATION() \
 			IMPLEMENT_ENCRYPTION_KEY_REGISTRATION() \
 			IMPLEMENT_GAME_MODULE(FDefaultGameModuleImpl, ModuleName) \
@@ -833,6 +850,8 @@ class FDefaultGameModuleImpl
 					FCString::Strncpy(GInternalProjectName, TEXT( GameName ), ARRAY_COUNT(GInternalProjectName)); \
 				} \
 			} AutoSet##ModuleName; \
+			IMPLEMENT_LIVE_CODING_ENGINE_DIR() \
+			IMPLEMENT_LIVE_CODING_PROJECT() \
 			PER_MODULE_BOILERPLATE \
 			PER_MODULE_BOILERPLATE_ANYLINK(FDefaultGameModuleImpl, ModuleName) \
 			FEngineLoop GEngineLoop;
@@ -852,6 +871,8 @@ class FDefaultGameModuleImpl
 			/* Implement the GIsGameAgnosticExe variable (See Core.h). */ \
 			bool GIsGameAgnosticExe = false; \
 			IMPLEMENT_FOREIGN_ENGINE_DIR() \
+			IMPLEMENT_LIVE_CODING_ENGINE_DIR() \
+			IMPLEMENT_LIVE_CODING_PROJECT() \
 			IMPLEMENT_SIGNING_KEY_REGISTRATION() \
 			IMPLEMENT_ENCRYPTION_KEY_REGISTRATION() \
 			IMPLEMENT_GAME_MODULE( ModuleImplClass, ModuleName ) \
@@ -864,6 +885,8 @@ class FDefaultGameModuleImpl
 			TCHAR GInternalProjectName[64] = TEXT( PREPROCESSOR_TO_STRING(UE_PROJECT_NAME) ); \
 			PER_MODULE_BOILERPLATE \
 			IMPLEMENT_FOREIGN_ENGINE_DIR() \
+			IMPLEMENT_LIVE_CODING_ENGINE_DIR() \
+			IMPLEMENT_LIVE_CODING_PROJECT() \
 			IMPLEMENT_SIGNING_KEY_REGISTRATION() \
 			IMPLEMENT_ENCRYPTION_KEY_REGISTRATION() \
 			IMPLEMENT_GAME_MODULE( ModuleImplClass, ModuleName ) \
