Live Coding: Print a warning if the PDB file for an image cannot be loaded.

diff --git a/Engine/Source/Developer/Windows/LiveCodingServer/Private/External/LC_Symbols.cpp b/Engine/Source/Developer/Windows/LiveCodingServer/Private/External/LC_Symbols.cpp
--- a/Engine/Source/Developer/Windows/LiveCodingServer/Private/External/LC_Symbols.cpp
+++ b/Engine/Source/Developer/Windows/LiveCodingServer/Private/External/LC_Symbols.cpp
@@ -280,6 +280,13 @@ namespace
 				LC_ERROR_USER("Cannot load PDB file for module %S. Error: 0x%X", filename, hr);
 			}
 
+			// BEGIN EPIC MOD - Show a warning if we don't have a PDB for this module. Since we only enable Live++ for modules that we built, we should always have a PDB.
+			if (hr == E_PDB_NOT_FOUND)
+			{
+				LC_WARNING_USER("No PDB file found for module %S. If this is a packaged build, make sure that debug files are being staged. Live coding will be disabled for this module.", file::GetFilename(filename).c_str());
+			}
+			// END EPIC MOD
+
 			return nullptr;
 		}
 
