Live Coding: Fix support for lazy loading modules. Now passes UBT a list of modules that can be updated, and UBT fails the build and writes out a list of others that would be modified. Live Coding console then loads those modules and retries.

diff --git a/Engine/Source/Developer/Windows/LiveCodingServer/Private/External/LC_AppSettings.cpp b/Engine/Source/Developer/Windows/LiveCodingServer/Private/External/LC_AppSettings.cpp
--- a/Engine/Source/Developer/Windows/LiveCodingServer/Private/External/LC_AppSettings.cpp
+++ b/Engine/Source/Developer/Windows/LiveCodingServer/Private/External/LC_AppSettings.cpp
@@ -61,7 +61,7 @@ namespace
 		if (!file::IsDirectory(attributes))
 		{
 			// this is not a directory, but a full path
-			LC_SUCCESS_USER("Using %S at path %S", type, path.c_str());
+// EPIC REMOVED			LC_SUCCESS_USER("Using %S at path %S", type, path.c_str());
 			return path;
 		}
 
diff --git a/Engine/Source/Developer/Windows/LiveCodingServer/Private/External/LC_ServerCommandThread.cpp b/Engine/Source/Developer/Windows/LiveCodingServer/Private/External/LC_ServerCommandThread.cpp
--- a/Engine/Source/Developer/Windows/LiveCodingServer/Private/External/LC_ServerCommandThread.cpp
+++ b/Engine/Source/Developer/Windows/LiveCodingServer/Private/External/LC_ServerCommandThread.cpp
@@ -11,6 +11,7 @@
 #include "LC_CommandMap.h"
 #include "LC_FileAttributeCache.h"
 #include "LiveCodingServer.h"
+#include "Containers/UnrealString.h"
 #include "LC_Shortcut.h"
 #include "LC_Key.h"
 #include "LC_ChangeNotification.h"
@@ -557,6 +558,48 @@ struct ClientProxyThread
 		m_pipe->SendCommandAndWaitForAck(commands::FinishedLazyLoadingModules(), nullptr, 0);
 	}
 };
+
+bool ServerCommandThread::EnableRequiredModules(const TArray<FString>& RequiredModules)
+{
+	bool bEnabledModule = false;
+	for (LiveProcess* liveProcess : m_liveProcesses)
+	{
+		types::vector<std::wstring> LoadModuleFileNames;
+		for (const FString& RequiredModule : RequiredModules)
+		{
+			std::wstring ModuleFileName = file::NormalizePath(*RequiredModule);
+			if (liveProcess->IsPendingLazyLoadedModule(ModuleFileName))
+			{
+				LoadModuleFileNames.push_back(ModuleFileName);
+			}
+		}
+		if (LoadModuleFileNames.size() > 0)
+		{
+			const std::wstring PipeName = primitiveNames::Pipe(m_processGroupName + L"_ClientProxy");
+
+			DuplexPipeServer ServerPipe;
+			ServerPipe.Create(PipeName.c_str());
+
+			DuplexPipeClient ClientPipe;
+			ClientPipe.Connect(PipeName.c_str());
+
+			ClientProxyThread ClientThread(liveProcess, &ClientPipe, LoadModuleFileNames);
+
+			CommandMap commandMap;
+			commandMap.RegisterAction<actions::EnableModules>();
+			commandMap.RegisterAction<actions::FinishedLazyLoadingModules>();
+			commandMap.HandleCommands(&ServerPipe, this);
+
+			for (const std::wstring& loadModuleFileName : LoadModuleFileNames)
+			{
+				liveProcess->SetLazyLoadedModuleAsLoaded(loadModuleFileName);
+			}
+
+			bEnabledModule = true;
+		}
+	}
+	return bEnabledModule;
+}
 // END EPIC MOD
 
 void ServerCommandThread::CompileChanges(bool didAllProcessesMakeProgress)
@@ -586,46 +629,29 @@ void ServerCommandThread::CompileChanges(bool didAllProcessesMakeProgress)
 
 		GLiveCodingServer->GetStatusChangeDelegate().ExecuteIfBound(L"Compiling changes for live coding...");
 
+		// Keep retrying the compile until we've added all the required modules
 		TMap<FString, TArray<FString>> ModuleToObjectFiles;
-		if (!CompileDelegate.Execute(Targets, ModuleToObjectFiles))
+		for (;;)
 		{
-			GLiveCodingServer->GetCompileFinishedDelegate().ExecuteIfBound(ELiveCodingResult::Error, L"Compilation error.");
-			return;
-		}
-
-		// Enable any lazy-loaded modules that we need
-		for (LiveProcess* liveProcess : m_liveProcesses)
-		{
-			types::vector<std::wstring> LoadModuleFileNames;
-			for(const TPair<FString, TArray<FString>>& Pair : ModuleToObjectFiles)
+			// Build a list of modules which are enabled for live coding
+			TArray<FString> ValidModules;
+			for (LiveModule* liveModule : m_liveModules)
 			{
-				std::wstring ModuleFileName = file::NormalizePath(*Pair.Key);
-				if (liveProcess->IsPendingLazyLoadedModule(ModuleFileName))
-				{
-					LoadModuleFileNames.push_back(ModuleFileName);
-				}
+				ValidModules.Add(liveModule->GetModuleName().c_str());
 			}
-			if (LoadModuleFileNames.size() > 0)
-			{
-				const std::wstring PipeName = primitiveNames::Pipe(m_processGroupName + L"_ClientProxy");
 
-				DuplexPipeServer ServerPipe;
-				ServerPipe.Create(PipeName.c_str());
-
-				DuplexPipeClient ClientPipe;
-				ClientPipe.Connect(PipeName.c_str());
-
-				ClientProxyThread ClientThread(liveProcess, &ClientPipe, LoadModuleFileNames);
-
-				CommandMap commandMap;
-				commandMap.RegisterAction<actions::EnableModules>();
-				commandMap.RegisterAction<actions::FinishedLazyLoadingModules>();
-				commandMap.HandleCommands(&ServerPipe, this);
+			// Execute the compile
+			TArray<FString> RequiredModules;
+			if (CompileDelegate.Execute(Targets, ValidModules, RequiredModules, ModuleToObjectFiles))
+			{
+				break;
+			}
 
-				for (const std::wstring& loadModuleFileName : LoadModuleFileNames)
-				{
-					liveProcess->SetLazyLoadedModuleAsLoaded(loadModuleFileName);
-				}
+			// Enable any lazy-loaded modules that we need
+			if (!EnableRequiredModules(RequiredModules))
+			{
+				GLiveCodingServer->GetCompileFinishedDelegate().ExecuteIfBound(ELiveCodingResult::Error, L"Compilation error.");
+				return;
 			}
 		}
 
@@ -1516,24 +1542,29 @@ bool ServerCommandThread::actions::EnableModules::Execute(const CommandType* com
 		g_dependencyAllocator.PrintStats();
 	}
 
-	if (loadModuleTaskCount > 0u)
+	// BEGIN EPIC MOD - Suppress output when lazy loading modules
+	if (commandThread->m_compileThread == nullptr || thread::GetId() != thread::GetId(commandThread->m_compileThread))
 	{
-		LC_SUCCESS_USER("Loaded %zu module(s) (%.3fs, %zu translation units)", loadModuleTaskCount, moduleLoadingScope.ReadSeconds(), loadedTranslationUnits);
-	}
+		if (loadModuleTaskCount > 0u)
+		{
+			LC_SUCCESS_USER("Loaded %zu module(s) (%.3fs, %zu translation units)", loadModuleTaskCount, moduleLoadingScope.ReadSeconds(), loadedTranslationUnits);
+		}
 
-	// EPIC REMOVED commandThread->PrewarmCompilerEnvironmentCache();
+		// EPIC REMOVED commandThread->PrewarmCompilerEnvironmentCache();
 
-	// tell user we are ready, but only once to not clutter the log
-	{
-		static bool showedOnce = false;
-		if (!showedOnce)
+		// tell user we are ready, but only once to not clutter the log
 		{
-			showedOnce = true;
-			const int shortcut = appSettings::g_compileShortcut->GetValue();
-			const std::wstring& shortcutText = shortcut::ConvertShortcutToText(shortcut);
-			LC_SUCCESS_USER("Live coding ready - Save changes and press %S to re-compile code", shortcutText.c_str());
+			static bool showedOnce = false;
+			if (!showedOnce)
+			{
+				showedOnce = true;
+				const int shortcut = appSettings::g_compileShortcut->GetValue();
+				const std::wstring& shortcutText = shortcut::ConvertShortcutToText(shortcut);
+				LC_SUCCESS_USER("Live coding ready - Save changes and press %S to re-compile code", shortcutText.c_str());
+			}
 		}
 	}
+	// END EPIC MOD
 
 	// remove virtual drives once we're finished
 	RemoveVirtualDrive();
diff --git a/Engine/Source/Developer/Windows/LiveCodingServer/Private/External/LC_ServerCommandThread.h b/Engine/Source/Developer/Windows/LiveCodingServer/Private/External/LC_ServerCommandThread.h
--- a/Engine/Source/Developer/Windows/LiveCodingServer/Private/External/LC_ServerCommandThread.h
+++ b/Engine/Source/Developer/Windows/LiveCodingServer/Private/External/LC_ServerCommandThread.h
@@ -134,6 +134,10 @@ private:
 	bool m_active = true;
 	// END EPIC MOD
 
+	// BEGIN EPIC MOD - Lazy loading modules
+	bool EnableRequiredModules(const TArray<FString>& RequiredModules);
+	// END EPIC MOD
+
 	// for triggering recompiles using the API
 	bool m_manualRecompileTriggered;
 	types::unordered_map<std::wstring, types::vector<LiveModule::ModifiedObjFile>> m_liveModuleToModifiedOrNewObjFiles;
diff --git a/Engine/Source/Developer/Windows/LiveCodingServer/Public/ILiveCodingServer.h b/Engine/Source/Developer/Windows/LiveCodingServer/Public/ILiveCodingServer.h
--- a/Engine/Source/Developer/Windows/LiveCodingServer/Public/ILiveCodingServer.h
+++ b/Engine/Source/Developer/Windows/LiveCodingServer/Public/ILiveCodingServer.h
@@ -44,7 +44,7 @@ public:
 	virtual FLogOutputDelegate& GetLogOutputDelegate() = 0;
 
 	typedef TMap<FString, TArray<FString>> FModuleToObjectFiles;
-	DECLARE_DELEGATE_RetVal_TwoParams(bool, FCompileDelegate, const TArray<FString>&, FModuleToObjectFiles&)
+	DECLARE_DELEGATE_RetVal_FourParams(bool, FCompileDelegate, const TArray<FString>&, const TArray<FString>&, TArray<FString>&, FModuleToObjectFiles&)
 	virtual FCompileDelegate& GetCompileDelegate() = 0;
 
 	DECLARE_DELEGATE(FCompileStartedDelegate);
diff --git a/Engine/Source/Programs/LiveCodingConsole/Private/LiveCodingConsole.cpp b/Engine/Source/Programs/LiveCodingConsole/Private/LiveCodingConsole.cpp
--- a/Engine/Source/Programs/LiveCodingConsole/Private/LiveCodingConsole.cpp
+++ b/Engine/Source/Programs/LiveCodingConsole/Private/LiveCodingConsole.cpp
@@ -179,7 +179,7 @@ private:
 		LogWidget->AppendLine(GetLogColor(Verbosity), MoveTemp(Text));
 	}
 
-	bool CompilePatch(const TArray<FString>& Targets, TMap<FString, TArray<FString>>& ModuleToObjectFiles)
+	bool CompilePatch(const TArray<FString>& Targets, const TArray<FString>& ValidModules, TArray<FString>& RequiredModules, TMap<FString, TArray<FString>>& ModuleToObjectFiles)
 	{
 		// Update the compile start time. This gets copied into the last patch time once a patch has been confirmed to have been applied.
 		NextPatchStartTime = FDateTime::UtcNow();
@@ -188,21 +188,30 @@ private:
 		FString Executable = FPaths::EngineDir() / TEXT("Binaries/DotNET/UnrealBuildTool.exe");
 		FPaths::MakePlatformFilename(Executable);
 
+		// Write out the list of lazy-loaded modules for UBT to check
+		FString ModulesFileName = FPaths::ConvertRelativePathToFull(FPaths::EngineIntermediateDir() / TEXT("LiveCodingModules.txt"));
+		FFileHelper::SaveStringArrayToFile(ValidModules, *ModulesFileName);
+
+		// Delete the output file for non-whitelisted modules
+		FString ModulesOutputFileName = ModulesFileName + TEXT(".out");
+		IFileManager::Get().Delete(*ModulesOutputFileName);
+
+		// Delete any existing manifest
+		FString ManifestFileName = FPaths::ConvertRelativePathToFull(FPaths::EngineIntermediateDir() / TEXT("LiveCoding.json"));
+		IFileManager::Get().Delete(*ManifestFileName);
+
 		// Build the argument list
 		FString Arguments;
 		for (const FString& Target : Targets)
 		{
 			Arguments += FString::Printf(TEXT("-Target=\"%s\" "), *Target.Replace(TEXT("\""), TEXT("\"\"")));
 		}
-
-		FString ManifestFileName = FPaths::ConvertRelativePathToFull(FPaths::EngineIntermediateDir() / TEXT("LiveCoding.json"));
-		Arguments += FString::Printf(TEXT("-LiveCoding -LiveCodingManifest=\"%s\" -WaitMutex"), *ManifestFileName);
-
+		Arguments += FString::Printf(TEXT("-LiveCoding -LiveCodingModules=\"%s\" -LiveCodingManifest=\"%s\" -WaitMutex"), *ModulesFileName, *ManifestFileName);
 		AppendLogLine(ELiveCodingLogVerbosity::Info, *FString::Printf(TEXT("Running %s %s"), *Executable, *Arguments));
 
 		// Spawn UBT and wait for it to complete (or the compile button to be pressed)
 		FMonitoredProcess Process(*Executable, *Arguments, true);
-		Process.OnOutput().BindLambda([this](const FString& Text){ AppendLogLine(ELiveCodingLogVerbosity::Info, *Text); });
+		Process.OnOutput().BindLambda([this](const FString& Text){ AppendLogLine(ELiveCodingLogVerbosity::Info, *(FString(TEXT("  ")) + Text)); });
 		Process.Launch();
 		while(Process.Update())
 		{
@@ -214,9 +223,17 @@ private:
 			FPlatformProcess::Sleep(0.1f);
 		}
 
-		if (Process.GetReturnCode() != 0)
+		int ReturnCode = Process.GetReturnCode();
+		if (ReturnCode != 0)
 		{
-			AppendLogLine(ELiveCodingLogVerbosity::Failure, TEXT("Build failed."));
+			if (FPaths::FileExists(ModulesOutputFileName))
+			{
+				FFileHelper::LoadFileToStringArray(RequiredModules, *ModulesOutputFileName);
+			}
+			else
+			{
+				AppendLogLine(ELiveCodingLogVerbosity::Failure, TEXT("Build failed."));
+			}
 			return false;
 		}
 
diff --git a/Engine/Source/Programs/UnrealBuildTool/Configuration/TargetDescriptor.cs b/Engine/Source/Programs/UnrealBuildTool/Configuration/TargetDescriptor.cs
--- a/Engine/Source/Programs/UnrealBuildTool/Configuration/TargetDescriptor.cs
+++ b/Engine/Source/Programs/UnrealBuildTool/Configuration/TargetDescriptor.cs
@@ -58,6 +58,12 @@ namespace UnrealBuildTool
 		[CommandLine("-WriteActions=")]
 		public List<FileReference> WriteActionFiles = new List<FileReference>();
 
+		/// <summary>
+		/// Path to a file containing a list of modules that may be modified for live coding.
+		/// </summary>
+		[CommandLine("-LiveCodingModules=")]
+		public FileReference LiveCodingModules = null;
+
 		/// <summary>
 		/// Constructor
 		/// </summary>
diff --git a/Engine/Source/Programs/UnrealBuildTool/Modes/BuildMode.cs b/Engine/Source/Programs/UnrealBuildTool/Modes/BuildMode.cs
--- a/Engine/Source/Programs/UnrealBuildTool/Modes/BuildMode.cs
+++ b/Engine/Source/Programs/UnrealBuildTool/Modes/BuildMode.cs
@@ -610,6 +610,46 @@ namespace UnrealBuildTool
 			// Additional processing for hot reload
 			if (HotReloadMode == HotReloadMode.LiveCoding)
 			{
+				// Make sure we're not overwriting any lazy-loaded modules
+				if(TargetDescriptor.LiveCodingModules != null)
+				{
+					// Read the list of modules that we're allowed to build
+					string[] Lines = FileReference.ReadAllLines(TargetDescriptor.LiveCodingModules);
+
+					// Parse it out into a set of filenames
+					HashSet<FileReference> AllowedOutputFiles = new HashSet<FileReference>();
+					foreach (string Line in Lines)
+					{
+						string TrimLine = Line.Trim();
+						if (TrimLine.Length > 0)
+						{
+							AllowedOutputFiles.Add(new FileReference(TrimLine));
+						}
+					}
+
+					// Find all the binaries that we're actually going to build
+					HashSet<FileReference> OutputFiles = new HashSet<FileReference>();
+					foreach (Action Action in TargetActionsToExecute)
+					{
+						if (Action.ActionType == ActionType.Link)
+						{
+							OutputFiles.UnionWith(Action.ProducedItems.Where(x => x.HasExtension(".exe") || x.HasExtension(".dll")).Select(x => x.Location));
+						}
+					}
+
+					// Find all the files that will be built that aren't allowed
+					List<FileReference> ProtectedOutputFiles = OutputFiles.Where(x => !AllowedOutputFiles.Contains(x)).ToList();
+					if (ProtectedOutputFiles.Count > 0)
+					{
+						FileReference.WriteAllLines(new FileReference(TargetDescriptor.LiveCodingModules.FullName + ".out"), ProtectedOutputFiles.Select(x => x.ToString()));
+						foreach(FileReference ProtectedOutputFile in ProtectedOutputFiles)
+						{
+							Log.TraceInformation("Module {0} is not currently enabled for Live Coding", ProtectedOutputFile);
+						}
+						throw new CompilationResultException(CompilationResult.Canceled);
+					}
+				}
+
 				// Filter the prerequisite actions down to just the compile actions, then recompute all the actions to execute
 				PrerequisiteActions = new List<Action>(TargetActionsToExecute.Where(x => x.ActionType == ActionType.Compile));
 				TargetActionsToExecute = ActionGraph.GetActionsToExecute(Makefile.Actions, PrerequisiteActions, CppDependencies, History, BuildConfiguration.bIgnoreOutdatedImportLibraries);
