Fixes for tools and Ocean editor

diff --git a/Engine/Source/Developer/Windows/WindowsTargetPlatform/Private/GenericWindowsTargetPlatform.h b/Engine/Source/Developer/Windows/WindowsTargetPlatform/Private/GenericWindowsTargetPlatform.h
--- a/Engine/Source/Developer/Windows/WindowsTargetPlatform/Private/GenericWindowsTargetPlatform.h
+++ b/Engine/Source/Developer/Windows/WindowsTargetPlatform/Private/GenericWindowsTargetPlatform.h
@@ -13,12 +13,15 @@
 	#include "Engine/VolumeTexture.h"
 	#include "StaticMeshResources.h"
 	#include "RHI.h"
+	#include "AudioCompressionSettings.h"
 #endif // WITH_ENGINE
 
+
 #define LOCTEXT_NAMESPACE "TGenericWindowsTargetPlatform"
 
 namespace Windows
 {
+#if WITH_ENGINE
 	void CachePlatformAudioCookOverrides(FPlatformAudioCookOverrides& OutOverrides)
 	{
 		const TCHAR* CategoryName = TEXT("/Script/WindowsTargetPlatform.WindowsTargetSettings");
@@ -129,6 +132,7 @@ namespace Windows
 			OutOverrides.PlatformSampleRates.Add(ESoundwaveSampleRateSettings::Min, RetrievedSampleRate);
 		}
 	}
+#endif
 }
 
 /**
