Merging Live++ 1.5.0

#rnx

diff --git a/Engine/Source/Developer/Windows/LiveCoding/Private/External/LC_API.cpp b/Engine/Source/Developer/Windows/LiveCoding/Private/External/LC_API.cpp
new file mode 100644
--- /dev/null
+++ b/Engine/Source/Developer/Windows/LiveCoding/Private/External/LC_API.cpp
@@ -0,0 +1,231 @@
+// Copyright 2011-2019 Molecular Matters GmbH, all rights reserved.
+
+#include "LC_API.h"
+#include "LC_ClientStartupThread.h"
+#include "LC_CriticalSection.h"
+#include "LC_SyncPoint.h"
+#include "LC_Restart.h"
+#include "LC_RunMode.h"
+#include "LC_Logging.h"
+#include "LPP_API.h"
+
+namespace
+{
+	// startup thread
+	static ClientStartupThread* g_startupThread = nullptr;
+	
+	// critical section to ensure that startup thread is initialized only once
+	static CriticalSection g_ensureOneTimeStartup;
+
+	// default run mode
+	static RunMode::Enum g_runMode = RunMode::DEFAULT;
+}
+
+
+void api::Startup(ClientStartupThread* startupThread)
+{
+	g_startupThread = startupThread;
+	restart::Startup();
+}
+
+
+void api::Shutdown(void)
+{
+	restart::Shutdown();
+	g_startupThread = nullptr;
+}
+
+
+LPP_DLL_API(const char*) LppGetVersion(void)
+{
+	return LPP_VERSION;
+}
+
+
+LPP_DLL_API(int) LppCheckVersion(const char* apiVersion)
+{
+	if (strcmp(apiVersion, LPP_VERSION) == 0)
+	{
+		return 1;
+	}
+	
+	LC_ERROR_USER("Version mismatch detected. API version: %s, DLL version: %s", LPP_VERSION, apiVersion);
+	return 0;
+}
+
+
+LPP_DLL_API(void) LppRegisterProcessGroup(const char* groupName)
+{
+	// now that we have the process group name, start Live++.
+	// ensure that initialization can happen only once, even if the user calls this more than once.
+	{
+		CriticalSection::ScopedLock lock(&g_ensureOneTimeStartup);
+
+		static bool firstTime = true;
+		if (!firstTime)
+		{
+			// this was already called once, bail out
+			return;
+		}
+
+		firstTime = false;
+	}
+
+	g_startupThread->Start(groupName, g_runMode);
+}
+
+
+LPP_DLL_API(void) LppSyncPoint(void)
+{
+	syncPoint::EnterTarget();
+}
+
+
+LPP_DLL_API(void) LppWaitForToken(void* token)
+{
+	if (!token)
+	{
+		// nullptr tokens are returned by Live++ when trying to enable modules which are not loaded into the host process.
+		// therefore, we need to handle this case gracefully.
+		return;
+	}
+
+	g_startupThread->WaitForToken(token);
+}
+
+
+LPP_DLL_API(void) LppTriggerRecompile(void)
+{
+	g_startupThread->TriggerRecompile();
+}
+
+
+LPP_DLL_API(void) LppLogMessage(const wchar_t* message)
+{
+	g_startupThread->LogMessage(message);
+}
+
+
+LPP_DLL_API(void) LppBuildPatch(const wchar_t* moduleNames[], const wchar_t* objPaths[], const wchar_t* amalgamatedObjPaths[], unsigned int count)
+{
+	g_startupThread->BuildPatch(moduleNames, objPaths, amalgamatedObjPaths, count);
+}
+
+
+LPP_DLL_API(void) LppInstallExceptionHandler(void)
+{
+	g_startupThread->InstallExceptionHandler();
+}
+
+
+LPP_DLL_API(void) LppUseExternalBuildSystem(void)
+{
+	g_runMode = RunMode::EXTERNAL_BUILD_SYSTEM;
+}
+
+
+LPP_DLL_API(int) LppWantsRestart(void)
+{
+	return restart::WasRequested();
+}
+
+
+LPP_DLL_API(void) LppRestart(lpp::RestartBehaviour behaviour, unsigned int exitCode)
+{
+	restart::Execute(behaviour, exitCode);
+}
+
+
+LPP_DLL_API(void*) LppEnableModule(const wchar_t* nameOfExeOrDll)
+{
+	return g_startupThread->EnableModule(nameOfExeOrDll);
+}
+
+
+LPP_DLL_API(void*) LppEnableModules(const wchar_t* namesOfExeOrDll[], unsigned int count)
+{
+	return g_startupThread->EnableModules(namesOfExeOrDll, count);
+}
+
+
+LPP_DLL_API(void*) LppEnableAllModules(const wchar_t* nameOfExeOrDll)
+{
+	return g_startupThread->EnableAllModules(nameOfExeOrDll);
+}
+
+
+LPP_DLL_API(void*) LppDisableModule(const wchar_t* nameOfExeOrDll)
+{
+	return g_startupThread->DisableModule(nameOfExeOrDll);
+}
+
+
+LPP_DLL_API(void*) LppDisableModules(const wchar_t* namesOfExeOrDll[], unsigned int count)
+{
+	return g_startupThread->DisableModules(namesOfExeOrDll, count);
+}
+
+
+LPP_DLL_API(void*) LppDisableAllModules(const wchar_t* nameOfExeOrDll)
+{
+	return g_startupThread->DisableAllModules(nameOfExeOrDll);
+}
+
+
+// BEGIN EPIC MOD - Adding ShowConsole command
+LPP_DLL_API(void) LppShowConsole()
+{
+	g_startupThread->ShowConsole();
+}
+// END EPIC MOD
+
+
+// BEGIN EPIC MOD - Adding SetVisible command
+LPP_DLL_API(void) LppSetVisible(bool visible)
+{
+	g_startupThread->SetVisible(visible);
+}
+// END EPIC MOD
+
+
+
+// BEGIN EPIC MOD - Adding SetActive command
+LPP_DLL_API(void) LppSetActive(bool active)
+{
+	g_startupThread->SetActive(active);
+}
+// END EPIC MOD
+
+
+// BEGIN EPIC MOD - Adding SetBuildArguments command
+LPP_DLL_API(void) LppSetBuildArguments(const wchar_t* arguments)
+{
+	g_startupThread->SetBuildArguments(arguments);
+}
+// END EPIC MOD
+
+// BEGIN EPIC MOD - Support for lazy-loading modules
+LPP_DLL_API(void) LppEnableLazyLoadedModule(const wchar_t* nameOfExeOrDll)
+{
+	HMODULE baseAddress = GetModuleHandle(nameOfExeOrDll);
+	g_startupThread->EnableLazyLoadedModule(nameOfExeOrDll, baseAddress);
+}
+// END EPIC MOD
+
+
+LPP_DLL_API(void) LppApplySettingBool(const char* settingName, int value)
+{
+	g_startupThread->ApplySettingBool(settingName, value);
+}
+
+
+LPP_DLL_API(void) LppApplySettingInt(const char* settingName, int value)
+{
+	g_startupThread->ApplySettingInt(settingName, value);
+}
+
+
+LPP_DLL_API(void) LppApplySettingString(const char* settingName, const wchar_t* value)
+{
+	g_startupThread->ApplySettingString(settingName, value);
+}
diff --git a/Engine/Source/Developer/Windows/LiveCoding/Private/External/LC_API.h b/Engine/Source/Developer/Windows/LiveCoding/Private/External/LC_API.h
new file mode 100644
--- /dev/null
+++ b/Engine/Source/Developer/Windows/LiveCoding/Private/External/LC_API.h
@@ -0,0 +1,44 @@
+// Copyright 2011-2019 Molecular Matters GmbH, all rights reserved.
+
+#pragma once
+
+#include "LPP_API.h"
+
+class ClientStartupThread;
+
+
+// internal API, initialized by entryPoint
+namespace api
+{
+	void Startup(ClientStartupThread* startupThread);
+	void Shutdown(void);
+}
+
+
+// external Live++ API exported into DLL
+// BEGIN EPIC MOD - Internalizing API
+#define LPP_DLL_API(_rv) _rv __cdecl
+//#define LPP_DLL_API(_rv) extern "C" __declspec(dllexport) _rv __cdecl
+// END EPIC MOD - Internalizing API
+
+LPP_DLL_API(const char*) LppGetVersion(void);
+LPP_DLL_API(int) LppCheckVersion(const char* apiVersion);
+LPP_DLL_API(void) LppRegisterProcessGroup(const char* groupName);
+LPP_DLL_API(void) LppSyncPoint(void);
+LPP_DLL_API(void) LppWaitForToken(void* token);
+LPP_DLL_API(void) LppTriggerRecompile(void);
+LPP_DLL_API(void) LppLogMessage(const wchar_t* message);
+LPP_DLL_API(void) LppBuildPatch(const wchar_t* moduleNames[], const wchar_t* objPaths[], const wchar_t* amalgamatedObjPaths[], unsigned int count);
+LPP_DLL_API(void) LppInstallExceptionHandler(void);
+LPP_DLL_API(void) LppUseExternalBuildSystem(void);
+LPP_DLL_API(int) LppWantsRestart(void);
+LPP_DLL_API(void) LppRestart(lpp::RestartBehaviour behaviour, unsigned int exitCode);
+LPP_DLL_API(void*) LppEnableModule(const wchar_t* nameOfExeOrDll);
+LPP_DLL_API(void*) LppEnableModules(const wchar_t* namesOfExeOrDll[], unsigned int count);
+LPP_DLL_API(void*) LppEnableAllModules(const wchar_t* nameOfExeOrDll);
+LPP_DLL_API(void*) LppDisableModule(const wchar_t* nameOfExeOrDll);
+LPP_DLL_API(void*) LppDisableModules(const wchar_t* namesOfExeOrDll[], unsigned int count);
+LPP_DLL_API(void*) LppDisableAllModules(const wchar_t* nameOfExeOrDll);
+LPP_DLL_API(void) LppApplySettingBool(const char* settingName, int value);
+LPP_DLL_API(void) LppApplySettingInt(const char* settingName, int value);
+LPP_DLL_API(void) LppApplySettingString(const char* settingName, const wchar_t* value);
diff --git a/Engine/Source/Developer/Windows/LiveCoding/Private/External/LC_ClientStartupThread.cpp b/Engine/Source/Developer/Windows/LiveCoding/Private/External/LC_ClientStartupThread.cpp
--- a/Engine/Source/Developer/Windows/LiveCoding/Private/External/LC_ClientStartupThread.cpp
+++ b/Engine/Source/Developer/Windows/LiveCoding/Private/External/LC_ClientStartupThread.cpp
@@ -12,6 +12,8 @@
 #include "LC_Event.h"
 #include "LC_CriticalSection.h"
 #include "LC_PrimitiveNames.h"
+#include "LC_Environment.h"
+#include "LC_MemoryStream.h"
 #include "LC_Logging.h"
 #include "Misc/Paths.h"
 #include "Misc/ConfigCacheIni.h"
@@ -237,6 +239,18 @@ void ClientStartupThread::TriggerRecompile(void)
 }
 
 
+void ClientStartupThread::LogMessage(const wchar_t* message)
+{
+	// wait for the startup thread to finish initialization
+	Join();
+
+	if (m_userCommandThread)
+	{
+		m_userCommandThread->LogMessage(message);
+	}
+}
+
+
 void ClientStartupThread::BuildPatch(const wchar_t* moduleNames[], const wchar_t* objPaths[], const wchar_t* amalgamatedObjPaths[], unsigned int count)
 {
 	// wait for the startup thread to finish initialization
@@ -488,7 +502,48 @@ unsigned int ClientStartupThread::ThreadFunction(const std::wstring& processGrou
 	m_userCommandThread->Start(processGroupName, m_startEvent, m_pipeClientCS);
 
 	// register this process with Live++
-	m_pipeClient->SendCommandAndWaitForAck(commands::RegisterProcess { process::GetBase(), process::GetId(), commandThreadId, &JumpToSelf }, nullptr, 0u);
+	{
+		// try getting the previous process ID from the environment in case the process was restarted
+		unsigned int restartedProcessId = 0u;
+		const std::wstring& processIdStr = environment::GetVariable(L"LPP_PROCESS_RESTART_ID", nullptr);
+		if (processIdStr.length() != 0u)
+		{
+			restartedProcessId = static_cast<unsigned int>(std::stoi(processIdStr));
+			environment::RemoveVariable(L"LPP_PROCESS_RESTART_ID");
+		}
+
+		// store the current process ID in an environment variable.
+		// upon restart, the environment block is inherited by the new process and can be used to map the process IDs of
+		// restarted processes to their previous IDs.
+		{
+			const unsigned int processID = process::GetId();
+			environment::SetVariable(L"LPP_PROCESS_RESTART_ID", std::to_wstring(processID).c_str());
+		}
+
+		const std::wstring& imagePath = process::GetImagePath();
+		const std::wstring& commandLine = process::GetCommandLine();
+		const std::wstring& workingDirectory = process::GetWorkingDirectory();
+		process::Environment* environment = process::CreateEnvironment(::GetCurrentProcess());
+
+		const commands::RegisterProcess command =
+		{
+			process::GetBase(), process::GetId(), restartedProcessId, commandThreadId, &JumpToSelf,
+			(imagePath.size() + 1u) * sizeof(wchar_t), 
+			(commandLine.size() + 1u) * sizeof(wchar_t),
+			(workingDirectory.size() + 1u) * sizeof(wchar_t),
+			environment->size
+		};
+
+		memoryStream::Writer payload(command.imagePathSize + command.commandLineSize + command.workingDirectorySize + command.environmentSize);
+		payload.Write(imagePath.data(), command.imagePathSize);
+		payload.Write(commandLine.data(), command.commandLineSize);
+		payload.Write(workingDirectory.data(), command.workingDirectorySize);
+		payload.Write(environment->data, environment->size);
+
+		m_pipeClient->SendCommandAndWaitForAck(command, payload.GetData(), payload.GetSize());
+
+		process::DestroyEnvironment(environment);
+	}
 
 	// handle commands until registration is finished
 	{
diff --git a/Engine/Source/Developer/Windows/LiveCoding/Private/External/LC_ClientStartupThread.h b/Engine/Source/Developer/Windows/LiveCoding/Private/External/LC_ClientStartupThread.h
--- a/Engine/Source/Developer/Windows/LiveCoding/Private/External/LC_ClientStartupThread.h
+++ b/Engine/Source/Developer/Windows/LiveCoding/Private/External/LC_ClientStartupThread.h
@@ -37,6 +37,7 @@ public:
 
 	void WaitForToken(void* token);
 	void TriggerRecompile(void);
+	void LogMessage(const wchar_t* message);
 	void BuildPatch(const wchar_t* moduleNames[], const wchar_t* objPaths[], const wchar_t* amalgamatedObjPaths[], unsigned int count);
 
 	void InstallExceptionHandler(void);
diff --git a/Engine/Source/Developer/Windows/LiveCoding/Private/External/LC_ClientUserCommandThread.cpp b/Engine/Source/Developer/Windows/LiveCoding/Private/External/LC_ClientUserCommandThread.cpp
--- a/Engine/Source/Developer/Windows/LiveCoding/Private/External/LC_ClientUserCommandThread.cpp
+++ b/Engine/Source/Developer/Windows/LiveCoding/Private/External/LC_ClientUserCommandThread.cpp
@@ -361,6 +361,18 @@ void ClientUserCommandThread::TriggerRecompile(void)
 }
 
 
+void ClientUserCommandThread::LogMessage(const wchar_t* message)
+{
+	const size_t lengthWithoutNull = wcslen(message);
+	const size_t payloadSize = sizeof(wchar_t) * (lengthWithoutNull + 1u);
+
+	ProxyCommand<commands::LogMessage>* proxy = new ProxyCommand<commands::LogMessage>(false, payloadSize);
+	proxy->m_payload.Write(message, payloadSize);
+
+	PushUserCommand(proxy);
+}
+
+
 void ClientUserCommandThread::BuildPatch(const wchar_t* moduleNames[], const wchar_t* objPaths[], const wchar_t* amalgamatedObjPaths[], unsigned int count)
 {
 	const size_t perFileSize = sizeof(wchar_t) * MAX_PATH * 3u;
diff --git a/Engine/Source/Developer/Windows/LiveCoding/Private/External/LC_ClientUserCommandThread.h b/Engine/Source/Developer/Windows/LiveCoding/Private/External/LC_ClientUserCommandThread.h
--- a/Engine/Source/Developer/Windows/LiveCoding/Private/External/LC_ClientUserCommandThread.h
+++ b/Engine/Source/Developer/Windows/LiveCoding/Private/External/LC_ClientUserCommandThread.h
@@ -60,6 +60,7 @@ public:
 
 	void WaitForToken(void* token);
 	void TriggerRecompile(void);
+	void LogMessage(const wchar_t* message);
 	void BuildPatch(const wchar_t* moduleNames[], const wchar_t* objPaths[], const wchar_t* amalgamatedObjPaths[], unsigned int count);
 
 	// BEGIN EPIC MOD - Adding ShowConsole command
diff --git a/Engine/Source/Developer/Windows/LiveCoding/Private/External/LC_CommandMap.cpp b/Engine/Source/Developer/Windows/LiveCoding/Private/External/LC_CommandMap.cpp
--- a/Engine/Source/Developer/Windows/LiveCoding/Private/External/LC_CommandMap.cpp
+++ b/Engine/Source/Developer/Windows/LiveCoding/Private/External/LC_CommandMap.cpp
@@ -55,6 +55,7 @@ CommandMap::CommandMap(void)
 	RegisterDefaultAction<commands::CompilationFinished>(m_actions);
 	RegisterDefaultAction<commands::DisconnectClient>(m_actions);
 	RegisterDefaultAction<commands::TriggerRecompile>(m_actions);	
+	RegisterDefaultAction<commands::LogMessage>(m_actions);
 	RegisterDefaultAction<commands::BuildPatch>(m_actions);
 	RegisterDefaultAction<commands::HandleException>(m_actions);
 	RegisterDefaultAction<commands::HandleExceptionFinished>(m_actions);
diff --git a/Engine/Source/Developer/Windows/LiveCoding/Private/External/LC_Commands.h b/Engine/Source/Developer/Windows/LiveCoding/Private/External/LC_Commands.h
--- a/Engine/Source/Developer/Windows/LiveCoding/Private/External/LC_Commands.h
+++ b/Engine/Source/Developer/Windows/LiveCoding/Private/External/LC_Commands.h
@@ -37,9 +37,17 @@ namespace commands
 		static const uint32_t ID = Acknowledge::ID + 1u;
 
 		void* processBase;
-		unsigned int processId;
-		unsigned int threadId;		// thread ID of Live++ thread running in host
-		const void* jumpToSelf;		// address of jump-to-self instruction in host
+		unsigned int processId;				// current process ID
+		unsigned int restartedProcessId;	// process ID of the previous, restarted process. 0 if non-existent
+		unsigned int threadId;				// thread ID of Live++ thread running in host
+		const void* jumpToSelf;				// address of jump-to-self instruction in host
+
+		size_t imagePathSize;
+		size_t commandLineSize;
+		size_t workingDirectorySize;
+		size_t environmentSize;
+
+		// image path, command line, working directory and environment follow as payload
 	};
 
 	// tell the DLL that registration has finished
@@ -175,10 +183,16 @@ namespace commands
 		static const uint32_t ID = DisconnectClient::ID + 1u;
 	};
 
+	// tell Live++ to log a message
+	struct LogMessage
+	{
+		static const uint32_t ID = TriggerRecompile::ID + 1u;
+	};
+
 	// tell Live++ to build a patch using an array of object files
 	struct BuildPatch
 	{
-		static const uint32_t ID = TriggerRecompile::ID + 1u;
+		static const uint32_t ID = LogMessage::ID + 1u;
 
 		unsigned int fileCount;
 
diff --git a/Engine/Source/Developer/Windows/LiveCoding/Private/External/LC_EntryPoint.cpp b/Engine/Source/Developer/Windows/LiveCoding/Private/External/LC_EntryPoint.cpp
--- a/Engine/Source/Developer/Windows/LiveCoding/Private/External/LC_EntryPoint.cpp
+++ b/Engine/Source/Developer/Windows/LiveCoding/Private/External/LC_EntryPoint.cpp
@@ -2,20 +2,13 @@
 
 #include "LC_EntryPoint.h"
 #include "LC_ClientStartupThread.h"
-#include "LC_CriticalSection.h"
-#include "LC_RunMode.h"
-#include "LPP_API.h"
+#include "LC_API.h"
 
 
 namespace
 {
-	// startup thread. initialized when LppRegisterProcessGroup is called
+	// startup thread
 	static ClientStartupThread* g_startupThread = nullptr;
-
-	// critical section to ensure that startup thread is initialized only once
-	static CriticalSection g_ensureOneTimeStartup;
-
-	static RunMode::Enum g_runMode = RunMode::DEFAULT;
 }
 
 
@@ -23,11 +16,14 @@ namespace
 void Startup(HINSTANCE instance)
 {
 	g_startupThread = new ClientStartupThread(instance);
+	api::Startup(g_startupThread);
 }
 
 
 void Shutdown(void)
 {
+	api::Shutdown();
+
 	// wait for the startup thread to finish its work and clean up
 	g_startupThread->Join();
 	delete g_startupThread;
@@ -51,179 +47,6 @@ BOOL WINAPI DllMain(_In_ HINSTANCE hinstDLL, _In_ DWORD dwReason, _In_ LPVOID /*
 			break;
 	}
 
-	return Windows::TRUE;
+	return TRUE;
 }
 #endif
-// END EPIC MOD
-
-
-// exported Live++ API
-// BEGIN EPIC MOD - Internalizing API
-#define LPP_DLL_API
-//#define LPP_DLL_API extern "C" __declspec(dllexport)
-// END EPIC MOD - Internalizing API
-
-
-LPP_DLL_API void __cdecl LppRegisterProcessGroup(const char* groupName)
-{
-	// now that we have the process group name, start Live++.
-	// ensure that initialization can happen only once, even if the user calls this more than once.
-	{
-		CriticalSection::ScopedLock lock(&g_ensureOneTimeStartup);
-
-		static bool firstTime = true;
-		if (!firstTime)
-		{
-			// this was already called once, bail out
-			return;
-		}
-
-		firstTime = false;
-	}
-
-	g_startupThread->Start(groupName, g_runMode);
-}
-
-
-LPP_DLL_API void* __cdecl LppEnableModule(const wchar_t* nameOfExeOrDll)
-{
-	// hand command creation to the user command thread
-	return g_startupThread->EnableModule(nameOfExeOrDll);
-}
-
-
-LPP_DLL_API void* __cdecl LppEnableModules(const wchar_t* namesOfExeOrDll[], unsigned int count)
-{
-	// hand command creation to the user command thread
-	return g_startupThread->EnableModules(namesOfExeOrDll, count);
-}
-
-
-LPP_DLL_API void* __cdecl LppEnableAllModules(const wchar_t* nameOfExeOrDll)
-{
-	// hand command creation to the user command thread
-	return g_startupThread->EnableAllModules(nameOfExeOrDll);
-}
-
-
-LPP_DLL_API void* __cdecl LppDisableModule(const wchar_t* nameOfExeOrDll)
-{
-	// hand command creation to the user command thread
-	return g_startupThread->DisableModule(nameOfExeOrDll);
-}
-
-
-LPP_DLL_API void* __cdecl LppDisableModules(const wchar_t* namesOfExeOrDll[], unsigned int count)
-{
-	// hand command creation to the user command thread
-	return g_startupThread->DisableModules(namesOfExeOrDll, count);
-}
-
-
-LPP_DLL_API void* __cdecl LppDisableAllModules(const wchar_t* nameOfExeOrDll)
-{
-	// hand command creation to the user command thread
-	return g_startupThread->DisableAllModules(nameOfExeOrDll);
-}
-
-
-LPP_DLL_API void __cdecl LppWaitForToken(void* token)
-{
-	if (!token)
-	{
-		// nullptr tokens are returned by Live++ when trying to enable modules which are not loaded into the host process.
-		// therefore, we need to handle this case gracefully.
-		return;
-	}
-
-	g_startupThread->WaitForToken(token);
-}
-
-
-LPP_DLL_API void __cdecl LppTriggerRecompile(void)
-{
-	g_startupThread->TriggerRecompile();
-}
-
-
-LPP_DLL_API void __cdecl LppBuildPatch(const wchar_t* moduleNames[], const wchar_t* objPaths[], const wchar_t* amalgamatedObjPaths[], unsigned int count)
-{
-	g_startupThread->BuildPatch(moduleNames, objPaths, amalgamatedObjPaths, count);
-}
-
-
-LPP_DLL_API void __cdecl LppInstallExceptionHandler(void)
-{
-	g_startupThread->InstallExceptionHandler();
-}
-
-
-LPP_DLL_API void __cdecl LppUseExternalBuildSystem(void)
-{
-	g_runMode = RunMode::EXTERNAL_BUILD_SYSTEM;
-}
-
-
-// BEGIN EPIC MOD - Adding ShowConsole command
-LPP_DLL_API void __cdecl LppShowConsole()
-{
-	g_startupThread->ShowConsole();
-}
-// END EPIC MOD
-
-
-// BEGIN EPIC MOD - Adding SetVisible command
-LPP_DLL_API void __cdecl LppSetVisible(bool visible)
-{
-	g_startupThread->SetVisible(visible);
-}
-// END EPIC MOD
-
-
-
-// BEGIN EPIC MOD - Adding SetActive command
-LPP_DLL_API void __cdecl LppSetActive(bool active)
-{
-	g_startupThread->SetActive(active);
-}
-// END EPIC MOD
-
-
-// BEGIN EPIC MOD - Adding SetBuildArguments command
-LPP_DLL_API void __cdecl LppSetBuildArguments(const wchar_t* arguments)
-{
-	g_startupThread->SetBuildArguments(arguments);
-}
-// END EPIC MOD
-
-// BEGIN EPIC MOD - Support for lazy-loading modules
-LPP_DLL_API void __cdecl LppEnableLazyLoadedModule(const wchar_t* nameOfExeOrDll)
-{
-	HMODULE baseAddress = GetModuleHandle(nameOfExeOrDll);
-	g_startupThread->EnableLazyLoadedModule(nameOfExeOrDll, baseAddress);
-}
-// END EPIC MOD
-
-
-LPP_DLL_API void __cdecl LppApplySettingBool(const char* settingName, int value)
-{
-	// hand command creation to the user command thread
-	g_startupThread->ApplySettingBool(settingName, value);
-}
-
-
-LPP_DLL_API void __cdecl LppApplySettingInt(const char* settingName, int value)
-{
-	// hand command creation to the user command thread
-	g_startupThread->ApplySettingInt(settingName, value);
-}
-
-
-LPP_DLL_API void __cdecl LppApplySettingString(const char* settingName, const wchar_t* value)
-{
-	// hand command creation to the user command thread
-	g_startupThread->ApplySettingString(settingName, value);
-}
-
-
-#undef LPP_DLL_API
diff --git a/Engine/Source/Developer/Windows/LiveCoding/Private/External/LC_Environment.cpp b/Engine/Source/Developer/Windows/LiveCoding/Private/External/LC_Environment.cpp
new file mode 100644
--- /dev/null
+++ b/Engine/Source/Developer/Windows/LiveCoding/Private/External/LC_Environment.cpp
@@ -0,0 +1,54 @@
+// Copyright 2011-2019 Molecular Matters GmbH, all rights reserved.
+
+#include "LC_Environment.h"
+#include "LC_Logging.h"
+#include "Windows/WindowsHWrapper.h"
+
+
+namespace
+{
+	static const DWORD MAX_VARIABLE_SIZE = 1024u;
+}
+
+
+void environment::RemoveVariable(const wchar_t* variable)
+{
+	const BOOL result = ::SetEnvironmentVariableW(variable, NULL);
+	if (result == 0)
+	{
+		LC_LOG_DEV("Could not remove environment variable %S (Error: %d)", variable, ::GetLastError());
+	}
+}
+
+
+void environment::SetVariable(const wchar_t* variable, const wchar_t* value)
+{
+	const BOOL result = ::SetEnvironmentVariableW(variable, value);
+	if (result == 0)
+	{
+		LC_LOG_DEV("Could not set environment variable %S to value %S (Error: %d)", variable, value, ::GetLastError());
+	}
+}
+
+
+std::wstring environment::GetVariable(const wchar_t* variable, const wchar_t* defaultValue)
+{
+	std::wstring value;
+	value.resize(MAX_VARIABLE_SIZE);
+
+	const DWORD result = ::GetEnvironmentVariableW(variable, &value[0], MAX_VARIABLE_SIZE);
+	if (result == 0)
+	{
+		// if a default value was provided, return that instead
+		if (defaultValue)
+		{
+			return std::wstring(defaultValue);
+		}
+		else
+		{
+			return std::wstring();
+		}
+	}
+
+	return value;
+}
diff --git a/Engine/Source/Developer/Windows/LiveCoding/Private/External/LC_Environment.h b/Engine/Source/Developer/Windows/LiveCoding/Private/External/LC_Environment.h
new file mode 100644
--- /dev/null
+++ b/Engine/Source/Developer/Windows/LiveCoding/Private/External/LC_Environment.h
@@ -0,0 +1,18 @@
+// Copyright 2011-2019 Molecular Matters GmbH, all rights reserved.
+
+#pragma once
+
+#include "CoreTypes.h"
+#include <string>
+
+namespace environment
+{
+	// Removes a variable from the environment of the calling process.
+	void RemoveVariable(const wchar_t* variable);
+
+	// Sets a variable in the environment of the calling process.
+	void SetVariable(const wchar_t* variable, const wchar_t* value);
+
+	// Gets a variable from the environment of the calling process.
+	std::wstring GetVariable(const wchar_t* variable, const wchar_t* defaultValue);
+}
diff --git a/Engine/Source/Developer/Windows/LiveCoding/Private/External/LC_Event.cpp b/Engine/Source/Developer/Windows/LiveCoding/Private/External/LC_Event.cpp
--- a/Engine/Source/Developer/Windows/LiveCoding/Private/External/LC_Event.cpp
+++ b/Engine/Source/Developer/Windows/LiveCoding/Private/External/LC_Event.cpp
@@ -26,43 +26,21 @@ Event::~Event(void)
 }
 
 
-void Event::Signal(void)
-{
-	::SetEvent(m_event);
-}
-
-
 void Event::Reset(void)
 {
 	::ResetEvent(m_event);
 }
 
 
-void Event::Wait(void)
+void Event::Signal(void)
 {
-	const DWORD result = ::WaitForSingleObject(m_event, INFINITE);
-	switch (result)
-	{
-		case WAIT_OBJECT_0:
-			// event was successfully signaled
-			break;
-
-		case WAIT_TIMEOUT:
-			// the operation timed out, which should never happen with a timeout of INFINITE
-			LC_ERROR_DEV("Event timed out.");
-			break;
-
-		case WAIT_ABANDONED:
-			LC_ERROR_DEV("Wait() was called on a stale event which was not released by the owning thread.");
-			break;
+	::SetEvent(m_event);
+}
 
-		case WAIT_FAILED:
-			LC_ERROR_DEV("Failed to Wait() on an event.");
-			break;
 
-		default:
-			break;
-	}
+bool Event::Wait(void)
+{
+	return WaitTimeout(INFINITE);
 }
 
 
@@ -72,9 +50,15 @@ bool Event::WaitTimeout(unsigned int milliSeconds)
 	switch (result)
 	{
 		case WAIT_OBJECT_0:
+			// event was successfully signaled
 			return true;
 
 		case WAIT_TIMEOUT:
+			// the operation timed out, which should never happen with a timeout of INFINITE
+			if (milliSeconds == INFINITE)
+			{
+				LC_ERROR_DEV("Event timed out.");
+			}
 			return false;
 
 		case WAIT_ABANDONED:
@@ -89,3 +73,9 @@ bool Event::WaitTimeout(unsigned int milliSeconds)
 			return false;
 	}
 }
+
+
+bool Event::TryWait(void)
+{
+	return WaitTimeout(0u);
+}
diff --git a/Engine/Source/Developer/Windows/LiveCoding/Private/External/LC_Event.h b/Engine/Source/Developer/Windows/LiveCoding/Private/External/LC_Event.h
--- a/Engine/Source/Developer/Windows/LiveCoding/Private/External/LC_Event.h
+++ b/Engine/Source/Developer/Windows/LiveCoding/Private/External/LC_Event.h
@@ -22,13 +22,22 @@ public:
 	Event(const wchar_t* name, Type::Enum type);
 	~Event(void);
 
-	void Signal(void);
+	// Resets the event.
 	void Reset(void);
-	void Wait(void);
 
-	// returns true if the event was signaled
+	// Signals the event.
+	void Signal(void);
+
+	// Waits until the event becomes signaled, blocking.
+	bool Wait(void);
+
+	// Waits until the event becomes signaled, blocking until the timeout is reached.
+	// Returns whether the event was signaled.
 	bool WaitTimeout(unsigned int milliSeconds);
 
+	// Returns whether the event was signaled, non-blocking.
+	bool TryWait(void);
+
 private:
 	Windows::HANDLE m_event;
 };
diff --git a/Engine/Source/Developer/Windows/LiveCoding/Private/External/LC_MemoryStream.cpp b/Engine/Source/Developer/Windows/LiveCoding/Private/External/LC_MemoryStream.cpp
--- a/Engine/Source/Developer/Windows/LiveCoding/Private/External/LC_MemoryStream.cpp
+++ b/Engine/Source/Developer/Windows/LiveCoding/Private/External/LC_MemoryStream.cpp
@@ -21,6 +21,12 @@ void memoryStream::Reader::Read(void* data, size_t size)
 }
 
 
+void memoryStream::Reader::Seek(size_t offset)
+{
+	m_offset = offset;
+}
+
+
 memoryStream::Writer::Writer(size_t capacity)
 	: m_data(new char[capacity])
 	, m_capacity(capacity)
diff --git a/Engine/Source/Developer/Windows/LiveCoding/Private/External/LC_MemoryStream.h b/Engine/Source/Developer/Windows/LiveCoding/Private/External/LC_MemoryStream.h
--- a/Engine/Source/Developer/Windows/LiveCoding/Private/External/LC_MemoryStream.h
+++ b/Engine/Source/Developer/Windows/LiveCoding/Private/External/LC_MemoryStream.h
@@ -22,6 +22,8 @@ namespace memoryStream
 			return data;
 		}
 
+		void Seek(size_t offset);
+
 	private:
 		const void* m_data;
 		size_t m_capacity;
diff --git a/Engine/Source/Developer/Windows/LiveCoding/Private/External/LC_Platform.h b/Engine/Source/Developer/Windows/LiveCoding/Private/External/LC_Platform.h
--- a/Engine/Source/Developer/Windows/LiveCoding/Private/External/LC_Platform.h
+++ b/Engine/Source/Developer/Windows/LiveCoding/Private/External/LC_Platform.h
@@ -6,17 +6,24 @@
 #include "Misc/AssertionMacros.h"
 #include "LC_Preprocessor.h"
 
-#define LC_FILE								__FILE__
-#define LC_LINE								__LINE__
-#define LC_FUNCTION_NAME					__FUNCTION__
-#define LC_FUNCTION_SIGNATURE				__FUNCSIG__
-#define LC_UNUSED(_value)					(void)_value
-#define LC_ASSERT(_condition, _msg)			checkf(_condition, TEXT("%s"), TEXT(_msg))
-
-#define LC_DISABLE_COPY(_name)				_name(const _name&) = delete
-#define LC_DISABLE_MOVE(_name)				_name(_name&&) = delete
-#define LC_DISABLE_ASSIGNMENT(_name)		_name& operator=(const _name&) = delete
-#define LC_DISABLE_MOVE_ASSIGNMENT(_name)	_name& operator=(_name&&) = delete
+#if DO_CHECK
+#	define LC_BREAKPOINT()						((::IsDebuggerPresent() == 0) ? (void)true : __debugbreak())
+#else
+#	define LC_BREAKPOINT()						(void)true
+#endif
+
+#define LC_ASSERT(_condition, _msg)				checkf(_condition, TEXT("%s"), TEXT(_msg))
+
+#define LC_FILE									__FILE__
+#define LC_LINE									__LINE__
+#define LC_FUNCTION_NAME						__FUNCTION__
+#define LC_FUNCTION_SIGNATURE					__FUNCSIG__
+#define LC_UNUSED(_value)						(void)(_value)
+
+#define LC_DISABLE_COPY(_name)					_name(const _name&) = delete
+#define LC_DISABLE_MOVE(_name)					_name(_name&&) = delete
+#define LC_DISABLE_ASSIGNMENT(_name)			_name& operator=(const _name&) = delete
+#define LC_DISABLE_MOVE_ASSIGNMENT(_name)		_name& operator=(_name&&) = delete
 
 #if PLATFORM_64BITS
 #	define LC_64_BIT 1
diff --git a/Engine/Source/Developer/Windows/LiveCoding/Private/External/LC_PointerUtil.h b/Engine/Source/Developer/Windows/LiveCoding/Private/External/LC_PointerUtil.h
--- a/Engine/Source/Developer/Windows/LiveCoding/Private/External/LC_PointerUtil.h
+++ b/Engine/Source/Developer/Windows/LiveCoding/Private/External/LC_PointerUtil.h
@@ -18,6 +18,14 @@ namespace pointer
 	}
 
 
+	template <typename To, typename From>
+	To FromInteger(From integer)
+	{
+		// widen integer to width of pointer type first before converting to pointer
+		return reinterpret_cast<To>(static_cast<uintptr_t>(integer));
+	}
+
+
 	template <typename T>
 	T AlignBottom(void* anyPointer, size_t alignment)
 	{
diff --git a/Engine/Source/Developer/Windows/LiveCoding/Private/External/LC_PrimitiveNames.cpp b/Engine/Source/Developer/Windows/LiveCoding/Private/External/LC_PrimitiveNames.cpp
--- a/Engine/Source/Developer/Windows/LiveCoding/Private/External/LC_PrimitiveNames.cpp
+++ b/Engine/Source/Developer/Windows/LiveCoding/Private/External/LC_PrimitiveNames.cpp
@@ -8,12 +8,16 @@
 #define LPP_JOB					LPP L"_JOB"
 #define LPP_MUTEX				LPP L"_IPM"
 #define LPP_MEMORY				LPP L"_NSM"
+#define LPP_EVENT				LPP L"_EVT"
 #define LPP_SERVER_READY		LPP L"_SR"
 #define LPP_COMPILE				LPP L"_CMP"
 #define LPP_PIPE				L"\\\\.\\pipe\\UE_LC"
 #define LPP_EXCEPTION_PIPE		L"\\\\.\\pipe\\UE_LC_EXC"
 #define LPP_HEARTBEAT_MUTEX		LPP_MUTEX L"_HB"
 #define LPP_HEARTBEAT_MEMORY	LPP_MEMORY L"_HB"
+#define LPP_RESTART_REQUESTED	LPP_EVENT L"_RST_REQ"
+#define LPP_RESTART_PREPARED	LPP_EVENT L"_RST_PREP"
+#define LPP_RESTART				LPP_EVENT L"_RST"
 
 
 std::wstring primitiveNames::JobGroup(const std::wstring& processGroupName)
@@ -126,13 +130,53 @@ std::wstring primitiveNames::HeartBeatNamedSharedMemory(const std::wstring& proc
 }
 
 
+std::wstring primitiveNames::RequestRestart(unsigned int processId)
+{
+	std::wstring name;
+	name.reserve(128u);
+
+	name += LPP_RESTART_REQUESTED;
+	name += std::to_wstring(processId);
+
+	return name;
+}
+
+
+std::wstring primitiveNames::PreparedRestart(unsigned int processId)
+{
+	std::wstring name;
+	name.reserve(128u);
+
+	name += LPP_RESTART_PREPARED;
+	name += std::to_wstring(processId);
+
+	return name;
+}
+
+
+std::wstring primitiveNames::Restart(unsigned int processId)
+{
+	std::wstring name;
+	name.reserve(128u);
+
+	name += LPP_RESTART;
+	name += std::to_wstring(processId);
+
+	return name;
+}
+
+
 #undef LPP
 #undef LPP_JOB
 #undef LPP_MUTEX
 #undef LPP_MEMORY
+#undef LPP_EVENT
 #undef LPP_SERVER_READY
 #undef LPP_COMPILE
 #undef LPP_PIPE
 #undef LPP_EXCEPTION_PIPE
 #undef LPP_HEARTBEAT_MUTEX
 #undef LPP_HEARTBEAT_MEMORY
+#undef LPP_RESTART_REQUESTED
+#undef LPP_RESTART_PREPARED
+#undef LPP_RESTART
diff --git a/Engine/Source/Developer/Windows/LiveCoding/Private/External/LC_PrimitiveNames.h b/Engine/Source/Developer/Windows/LiveCoding/Private/External/LC_PrimitiveNames.h
--- a/Engine/Source/Developer/Windows/LiveCoding/Private/External/LC_PrimitiveNames.h
+++ b/Engine/Source/Developer/Windows/LiveCoding/Private/External/LC_PrimitiveNames.h
@@ -23,4 +23,9 @@ namespace primitiveNames
 	// heart beat
 	std::wstring HeartBeatMutex(const std::wstring& processGroupName, unsigned int processId);
 	std::wstring HeartBeatNamedSharedMemory(const std::wstring& processGroupName, unsigned int processId);
+
+	// restart
+	std::wstring RequestRestart(unsigned int processId);
+	std::wstring PreparedRestart(unsigned int processId);
+	std::wstring Restart(unsigned int processId);
 }
diff --git a/Engine/Source/Developer/Windows/LiveCoding/Private/External/LC_Process.cpp b/Engine/Source/Developer/Windows/LiveCoding/Private/External/LC_Process.cpp
--- a/Engine/Source/Developer/Windows/LiveCoding/Private/External/LC_Process.cpp
+++ b/Engine/Source/Developer/Windows/LiveCoding/Private/External/LC_Process.cpp
@@ -4,6 +4,8 @@
 #include "LC_Memory.h"
 #include "LC_PointerUtil.h"
 #include "LC_VirtualMemory.h"
+#include "LC_WindowsInternals.h"
+#include "LC_WindowsInternalFunctions.h"
 #include "LC_Logging.h"
 #include <Psapi.h>
 
@@ -13,353 +15,6 @@
 #pragma warning(disable:6335) // warning C6335: Leaking process information handle 'context->pi.hProcess'.
 // END EPIC MODS
 
-// internal ntdll.dll definitions
-typedef LONG NTSTATUS;
-typedef LONG KPRIORITY;
-
-enum NT_SYSTEM_INFORMATION_CLASS
-{
-	SystemBasicInformation = 0,
-	SystemPerformanceInformation = 2,
-	SystemTimeOfDayInformation = 3,
-	SystemProcessInformation = 5,
-	SystemProcessorPerformanceInformation = 8,
-	SystemHandleInformation = 16,
-	SystemInterruptInformation = 23,
-	SystemExceptionInformation = 33,
-	SystemRegistryQuotaInformation = 37,
-	SystemLookasideInformation = 45,
-	SystemProcessIdInformation = 0x58
-};
-
-enum NT_KWAIT_REASON
-{
-	Executive,
-	FreePage,
-	PageIn,
-	PoolAllocation,
-	DelayExecution,
-	Suspended,
-	UserRequest,
-	WrExecutive,
-	WrFreePage,
-	WrPageIn,
-	WrPoolAllocation,
-	WrDelayExecution,
-	WrSuspended,
-	WrUserRequest,
-	WrEventPair,
-	WrQueue,
-	WrLpcReceive,
-	WrLpcReply,
-	WrVirtualMemory,
-	WrPageOut,
-	WrRendezvous,
-	Spare2,
-	Spare3,
-	Spare4,
-	Spare5,
-	Spare6,
-	WrKernel,
-	MaximumWaitReason
-};
-
-struct NT_CLIENT_ID
-{
-	HANDLE UniqueProcess;
-	HANDLE UniqueThread;
-};
-
-struct NT_SYSTEM_THREAD_INFORMATION
-{
-	LARGE_INTEGER KernelTime;
-	LARGE_INTEGER UserTime;
-	LARGE_INTEGER CreateTime;
-	ULONG WaitTime;
-	PVOID StartAddress;
-	NT_CLIENT_ID ClientId;
-	KPRIORITY Priority;
-	LONG BasePriority;
-	ULONG ContextSwitches;
-	ULONG ThreadState;
-	NT_KWAIT_REASON WaitReason;
-};
-
-struct NT_UNICODE_STRING
-{
-	USHORT Length;
-	USHORT MaximumLength;
-	PWSTR Buffer;
-};
-
-struct NT_SYSTEM_PROCESS_INFORMATION
-{
-	ULONG uNext;
-	ULONG uThreadCount;
-	LARGE_INTEGER WorkingSetPrivateSize; // since VISTA
-	ULONG HardFaultCount; // since WIN7
-	ULONG NumberOfThreadsHighWatermark; // since WIN7
-	ULONGLONG CycleTime; // since WIN7
-	LARGE_INTEGER CreateTime;
-	LARGE_INTEGER UserTime;
-	LARGE_INTEGER KernelTime;
-	NT_UNICODE_STRING ImageName;
-	KPRIORITY BasePriority;
-	HANDLE uUniqueProcessId;
-	HANDLE InheritedFromUniqueProcessId;
-	ULONG HandleCount;
-	ULONG SessionId;
-	ULONG_PTR UniqueProcessKey; // since VISTA (requires SystemExtendedProcessInformation)
-	SIZE_T PeakVirtualSize;
-	SIZE_T VirtualSize;
-	ULONG PageFaultCount;
-	SIZE_T PeakWorkingSetSize;
-	SIZE_T WorkingSetSize;
-	SIZE_T QuotaPeakPagedPoolUsage;
-	SIZE_T QuotaPagedPoolUsage;
-	SIZE_T QuotaPeakNonPagedPoolUsage;
-	SIZE_T QuotaNonPagedPoolUsage;
-	SIZE_T PagefileUsage;
-	SIZE_T PeakPagefileUsage;
-	SIZE_T PrivatePageCount;
-	LARGE_INTEGER ReadOperationCount;
-	LARGE_INTEGER WriteOperationCount;
-	LARGE_INTEGER OtherOperationCount;
-	LARGE_INTEGER ReadTransferCount;
-	LARGE_INTEGER WriteTransferCount;
-	LARGE_INTEGER OtherTransferCount;
-	NT_SYSTEM_THREAD_INFORMATION Threads[1];
-};
-
-enum NT_PROCESS_INFORMATION_CLASS
-{
-	ProcessBasicInformation,
-	ProcessQuotaLimits,
-	ProcessIoCounters,
-	ProcessVmCounters,
-	ProcessTimes,
-	ProcessBasePriority,
-	ProcessRaisePriority,
-	ProcessDebugPort,
-	ProcessExceptionPort,
-	ProcessAccessToken,
-	ProcessLdtInformation,
-	ProcessLdtSize,
-	ProcessDefaultHardErrorMode,
-	ProcessIoPortHandlers,          // Note: this is kernel mode only
-	ProcessPooledUsageAndLimits,
-	ProcessWorkingSetWatch,
-	ProcessUserModeIOPL,
-	ProcessEnableAlignmentFaultFixup,
-	ProcessPriorityClass,
-	ProcessWx86Information,
-	ProcessHandleCount,
-	ProcessAffinityMask,
-	ProcessPriorityBoost,
-	ProcessDeviceMap,
-	ProcessSessionInformation,
-	ProcessForegroundInformation,
-	ProcessWow64Information,
-	ProcessImageFileName,
-	ProcessLUIDDeviceMapsEnabled,
-	ProcessBreakOnTermination,
-	ProcessDebugObjectHandle,
-	ProcessDebugFlags,
-	ProcessHandleTracing,
-	ProcessIoPriority,
-	ProcessExecuteFlags,
-	ProcessTlsInformation,
-	ProcessCookie,
-	ProcessImageInformation,
-	ProcessCycleTime,
-	ProcessPagePriority,
-	ProcessInstrumentationCallback,
-	ProcessThreadStackAllocation,
-	ProcessWorkingSetWatchEx,
-	ProcessImageFileNameWin32,
-	ProcessImageFileMapping,
-	ProcessAffinityUpdateMode,
-	ProcessMemoryAllocationMode,
-	ProcessGroupInformation,
-	ProcessTokenVirtualizationEnabled,
-	ProcessConsoleHostProcess,
-	ProcessWindowInformation,
-	MaxProcessInfoClass
-};
-
-struct RTL_USER_PROCESS_PARAMETERS
-{
-	BYTE Reserved1[16];
-	PVOID Reserved2[10];
-	NT_UNICODE_STRING ImagePathName;
-	NT_UNICODE_STRING CommandLine;
-};
-
-struct NT_LDR_DATA_TABLE_ENTRY
-{
-	LIST_ENTRY InLoadOrderLinks;
-	LIST_ENTRY InMemoryOrderLinks;
-	LIST_ENTRY InInitializationOrderLinks;
-	PVOID DllBase;
-	PVOID EntryPoint;
-	ULONG SizeOfImage;
-	NT_UNICODE_STRING FullDllName;
-	NT_UNICODE_STRING BaseDllName;
-	ULONG Flags;
-	USHORT LoadCount;
-	USHORT ObsoleteLoadCount;
-	USHORT TlsIndex;
-	LIST_ENTRY HashLinks;
-	ULONG TimeDateStamp;
-};
-
-struct NT_PEB_LDR_DATA
-{
-	ULONG Length;
-	BOOLEAN Initialized;
-	PVOID SsHandle;
-	LIST_ENTRY InLoadOrderModuleList;
-	LIST_ENTRY InMemoryOrderModuleList;
-	LIST_ENTRY InInitializationOrderModuleList;
-	PVOID EntryInProgress;
-	BOOLEAN ShutdownInProgress;
-	HANDLE ShutdownThreadId;
-};
-
-typedef VOID (NTAPI *NT_PS_POST_PROCESS_INIT_ROUTINE)(VOID);
-
-// not the real full definition, but enough for our purposes
-struct NT_PEB
-{
-	BYTE Reserved1[2];
-	BYTE BeingDebugged;
-	BYTE Reserved2[1];
-	PVOID Reserved3[2];
-	NT_PEB_LDR_DATA* Ldr;
-	RTL_USER_PROCESS_PARAMETERS* ProcessParameters;
-	PVOID Reserved4[3];
-	PVOID AtlThunkSListPtr;
-	PVOID Reserved5;
-	ULONG Reserved6;
-	PVOID Reserved7;
-	ULONG Reserved8;
-	ULONG AtlThunkSListPtr32;
-	PVOID Reserved9[45];
-	BYTE Reserved10[96];
-	NT_PS_POST_PROCESS_INIT_ROUTINE* PostProcessInitRoutine;
-	BYTE Reserved11[128];
-	PVOID Reserved12[1];
-	ULONG SessionId;
-};
-
-struct NT_PROCESS_BASIC_INFORMATION
-{
-	NTSTATUS ExitStatus;
-	NT_PEB* PebBaseAddress;
-	ULONG_PTR AffinityMask;
-	KPRIORITY BasePriority;
-	HANDLE UniqueProcessId;
-	HANDLE InheritedFromUniqueProcessId;
-};
-
-#define NT_SUCCESS(Status)				((NTSTATUS)(Status) >= 0)
-#define STATUS_INFO_LENGTH_MISMATCH		((NTSTATUS)0xC0000004L)
-
-
-// these are undocumented functions, found in ntdll.dll.
-// we don't call them directly, but use them for extracting their signature.
-extern "C" NTSTATUS NtSuspendProcess(HANDLE ProcessHandle);
-extern "C" NTSTATUS NtResumeProcess(HANDLE ProcessHandle);
-extern "C" NTSTATUS NtWriteVirtualMemory(HANDLE ProcessHandle, PVOID BaseAddress, PVOID Buffer, ULONG NumberOfBytesToWrite, PULONG NumberOfBytesWritten);
-extern "C" NTSTATUS NtQuerySystemInformation(NT_SYSTEM_INFORMATION_CLASS SystemInformationClass, PVOID SystemInformation, ULONG SystemInformationLength, PULONG ReturnLength);
-extern "C" NTSTATUS NtQueryInformationProcess(HANDLE ProcessHandle, NT_PROCESS_INFORMATION_CLASS ProcessInformationClass, PVOID ProcessInformation, ULONG ProcessInformationLength, PULONG ReturnLength);
-extern "C" NTSTATUS NtContinue(CONTEXT* ThreadContext, BOOLEAN RaiseAlert);
-
-
-namespace
-{
-	// helper class that allows us to call an undocumented function in any Windows DLL, as long as it is exported and we know its signature
-	template <typename T>
-	class UndocumentedFunction {};
-
-	template <typename R, typename... Args>
-	class UndocumentedFunction<R(Args...)>
-	{
-		typedef R (NTAPI *Function)(Args...);
-
-	public:
-		UndocumentedFunction(const char* moduleName, const char* functionName)
-			: m_moduleName(moduleName)
-			, m_functionName(functionName)
-			, m_function(nullptr)
-		{
-			HMODULE module = ::GetModuleHandleA(moduleName);
-			if (!module)
-			{
-				LC_ERROR_USER("Cannot get handle for module %s", moduleName);
-				return;
-			}
-
-			m_function = reinterpret_cast<Function>(reinterpret_cast<uintptr_t>(::GetProcAddress(module, functionName)));
-			if (!m_function)
-			{
-				LC_ERROR_USER("Cannot get address of function %s in module %s", functionName, moduleName);
-			}
-		}
-
-		R operator()(Args... args) const
-		{
-			const NTSTATUS status = m_function(std::forward<Args>(args)...);
-			if (!NT_SUCCESS(status))
-			{
-				LC_ERROR_USER("Call to function %s in module %s failed. Error: 0x%X", m_functionName, m_moduleName, status);
-			}
-
-			return status;
-		}
-
-	private:
-		const char* m_moduleName;
-		const char* m_functionName;
-		Function m_function;
-	};
-
-
-	static uint32_t ConvertToExecutableProtection(uint32_t currentProtection)
-	{
-		// cut off PAGE_GUARD, PAGE_NOCACHE, PAGE_WRITECOMBINE, and PAGE_REVERT_TO_FILE_MAP
-		const uint32_t extraBits = currentProtection & 0xFFFFFF00u;
-		const uint32_t pageProtection = currentProtection & 0x000000FFu;
-
-		switch (pageProtection)
-		{
-			case PAGE_NOACCESS:
-			case PAGE_READONLY:
-			case PAGE_READWRITE:
-			case PAGE_WRITECOPY:
-				return (pageProtection << 4u) | extraBits;
-
-			case PAGE_EXECUTE:
-			case PAGE_EXECUTE_READ:
-			case PAGE_EXECUTE_READWRITE:
-			case PAGE_EXECUTE_WRITECOPY:
-			default:
-				return currentProtection;
-		}
-	}
-}
-
-
-namespace undocumentedFunctions
-{
-	static UndocumentedFunction<decltype(NtSuspendProcess)> NtSuspendProcess("ntdll.dll", "NtSuspendProcess");
-	static UndocumentedFunction<decltype(NtResumeProcess)> NtResumeProcess("ntdll.dll", "NtResumeProcess");
-	static UndocumentedFunction<decltype(NtWriteVirtualMemory)> NtWriteVirtualMemory("ntdll.dll", "NtWriteVirtualMemory");
-	static UndocumentedFunction<decltype(NtQuerySystemInformation)> NtQuerySystemInformation("ntdll.dll", "NtQuerySystemInformation");
-	static UndocumentedFunction<decltype(NtQueryInformationProcess)> NtQueryInformationProcess("ntdll.dll", "NtQueryInformationProcess");
-	static UndocumentedFunction<decltype(NtContinue)> NtContinue("ntdll.dll", "NtContinue");
-}
-
 
 namespace process
 {
@@ -476,10 +131,21 @@ namespace process
 			LC_LOG_DEV("Command line: %S", commandLineBuffer ? commandLineBuffer : L"none");
 			LC_LOG_DEV("Working directory: %S", workingDirectory ? workingDirectory : L"none");
 			LC_LOG_DEV("Custom environment block: %S", environmentBlock ? L"yes" : L"no");
+			LC_LOG_DEV("Flags: %u", flags);
+		}
+
+		DWORD creationFlags = CREATE_UNICODE_ENVIRONMENT;
+		if (flags & SpawnFlags::NO_WINDOW)
+		{
+			creationFlags |= CREATE_NO_WINDOW;
+		}
+		if (flags & SpawnFlags::SUSPENDED)
+		{
+			creationFlags |= CREATE_SUSPENDED;
 		}
 
 		// the environment block is not written to by CreateProcess, so it is safe to const_cast (it's a Win32 API mistake)
-		const BOOL success = ::CreateProcessW(exePath, commandLineBuffer, NULL, NULL, Windows::TRUE, CREATE_NO_WINDOW, const_cast<void*>(environmentBlock), workingDirectory, &startupInfo, &context->pi);
+		const BOOL success = ::CreateProcessW(exePath, commandLineBuffer, NULL, NULL, Windows::TRUE, creationFlags, const_cast<void*>(environmentBlock), workingDirectory, &startupInfo, &context->pi);
 		if (success == 0)
 		{
 			LC_ERROR_USER("Could not spawn process %S. Error: %d", exePath, ::GetLastError());
@@ -498,6 +164,12 @@ namespace process
 	}
 
 
+	void ResumeMainThread(Context* context)
+	{
+		::ResumeThread(context->pi.hThread);
+	}
+
+
 	unsigned int Wait(Context* context)
 	{
 		// wait until process terminates
@@ -520,6 +192,18 @@ namespace process
 	}
 
 
+	unsigned int Wait(Handle handle)
+	{
+		// wait until process terminates
+		::WaitForSingleObject(handle, INFINITE);
+
+		DWORD exitCode = 0xFFFFFFFFu;
+		::GetExitCodeProcess(handle, &exitCode);
+
+		return exitCode;
+	}
+
+
 	void Destroy(Context*& context)
 	{
 		::CloseHandle(context->pi.hProcess);
@@ -561,6 +245,21 @@ namespace process
 	}
 
 
+	std::wstring GetWorkingDirectory(void)
+	{
+		wchar_t workingDirectory[MAX_PATH + 1] = {};
+		::GetCurrentDirectoryW(MAX_PATH + 1u, workingDirectory);
+
+		return std::wstring(workingDirectory);
+	}
+
+
+	std::wstring GetCommandLine(void)
+	{
+		return std::wstring(::GetCommandLineW());
+	}
+
+
 	void* GetBase(void)
 	{
 		return ::GetModuleHandle(NULL);
@@ -619,7 +318,7 @@ namespace process
 			// than in previous versions of Windows.
 			// this bug was reported here:
 			// https://developercommunity.visualstudio.com/content/problem/228061/writeprocessmemory-slowdown-on-windows-10.html
-			undocumentedFunctions::NtWriteVirtualMemory(handle, destAddress, const_cast<PVOID>(srcBuffer), static_cast<ULONG>(size), NULL);
+			windowsInternal::NtWriteVirtualMemory(handle, destAddress, const_cast<PVOID>(srcBuffer), static_cast<ULONG>(size), NULL);
 		}
 		::VirtualProtectEx(handle, destAddress, size, oldProtect, &oldProtect);
 	}
@@ -644,8 +343,8 @@ namespace process
 				return nullptr;
 			}
 
-			MEMORY_BASIC_INFORMATION memoryInfo = {};
-			::VirtualQueryEx(handle, scan, &memoryInfo, sizeof(MEMORY_BASIC_INFORMATION));
+			::MEMORY_BASIC_INFORMATION memoryInfo = {};
+			::VirtualQueryEx(handle, scan, &memoryInfo, sizeof(::MEMORY_BASIC_INFORMATION));
 
 			if ((memoryInfo.RegionSize >= size) && (memoryInfo.State == MEM_FREE))
 			{
@@ -665,15 +364,15 @@ namespace process
 
 		for (const void* scan = address; /* nothing */; /* nothing */)
 		{
-			MEMORY_BASIC_INFORMATION memoryInfo = {};
-			const SIZE_T bytesInBuffer = ::VirtualQueryEx(handle, scan, &memoryInfo, sizeof(MEMORY_BASIC_INFORMATION));
+			::MEMORY_BASIC_INFORMATION memoryInfo = {};
+			const SIZE_T bytesInBuffer = ::VirtualQueryEx(handle, scan, &memoryInfo, sizeof(::MEMORY_BASIC_INFORMATION));
 			if (bytesInBuffer == 0u)
 			{
 				// could not query the protection, bail out
 				break;
 			}
 
-			const uint32_t executableProtection = ConvertToExecutableProtection(memoryInfo.Protect);
+			const uint32_t executableProtection = ConvertPageProtectionToExecutableProtection(memoryInfo.Protect);
 			if (executableProtection != memoryInfo.Protect)
 			{
 				// change this page into an executable one
@@ -702,19 +401,19 @@ namespace process
 
 	void Suspend(Handle handle)
 	{
-		undocumentedFunctions::NtSuspendProcess(handle);
+		windowsInternal::NtSuspendProcess(handle);
 	}
 
 
 	void Resume(Handle handle)
 	{
-		undocumentedFunctions::NtResumeProcess(handle);
+		windowsInternal::NtResumeProcess(handle);
 	}
 
 
 	void Continue(CONTEXT* threadContext)
 	{
-		undocumentedFunctions::NtContinue(threadContext, Windows::FALSE);
+		windowsInternal::NtContinue(threadContext, Windows::FALSE);
 	}
 
 
@@ -726,14 +425,14 @@ namespace process
 		// 2MB should be enough for getting the process information, even on systems with high load
 		ULONG bufferSize = 2048u * 1024u;
 		void* processSnapshot = nullptr;
-		NTSTATUS status = 0;
+		windowsInternal::NTSTATUS status = 0;
 
 		do
 		{
 			processSnapshot = ::malloc(bufferSize);
 
 			// try getting a process snapshot into the provided buffer
-			status = undocumentedFunctions::NtQuerySystemInformation(SystemProcessInformation, processSnapshot, bufferSize, NULL);
+			status = windowsInternal::NtQuerySystemInformation(windowsInternal::SystemProcessInformation, processSnapshot, bufferSize, NULL);
 
 			if (status == STATUS_INFO_LENGTH_MISMATCH)
 			{
@@ -756,17 +455,17 @@ namespace process
 
 		// find the process information for the given process ID
 		{
-			NT_SYSTEM_PROCESS_INFORMATION* processInfo = static_cast<NT_SYSTEM_PROCESS_INFORMATION*>(processSnapshot);
+			windowsInternal::NT_SYSTEM_PROCESS_INFORMATION* processInfo = static_cast<windowsInternal::NT_SYSTEM_PROCESS_INFORMATION*>(processSnapshot);
 
 			while (processInfo != nullptr)
 			{
-				if (processInfo->uUniqueProcessId == reinterpret_cast<HANDLE>(static_cast<DWORD_PTR>(processId)))
+				if (processInfo->UniqueProcessId == reinterpret_cast<HANDLE>(static_cast<DWORD_PTR>(processId)))
 				{
 					// we found the process we're looking for
 					break;
 				}
 
-				if (processInfo->uNext == 0u)
+				if (processInfo->NextEntryOffset == 0u)
 				{
 					// we couldn't find our process
 					LC_ERROR_USER("Cannot enumerate threads, process not found (PID: %d)", processId);
@@ -777,14 +476,14 @@ namespace process
 				else
 				{
 					// walk to the next process info
-					processInfo = pointer::Offset<NT_SYSTEM_PROCESS_INFORMATION*>(processInfo, processInfo->uNext);
+					processInfo = pointer::Offset<windowsInternal::NT_SYSTEM_PROCESS_INFORMATION*>(processInfo, processInfo->NextEntryOffset);
 				}
 			}
 
 			// record all threads belonging to the given process
 			if (processInfo)
 			{
-				for (ULONG i = 0u; i < processInfo->uThreadCount; ++i)
+				for (ULONG i = 0u; i < processInfo->NumberOfThreads; ++i)
 				{
 					const DWORD threadId = static_cast<DWORD>(reinterpret_cast<DWORD_PTR>(processInfo->Threads[i].ClientId.UniqueThread));
 					threadIds.push_back(threadId);
@@ -800,34 +499,35 @@ namespace process
 
 	std::vector<Module> EnumerateModules(Handle handle)
 	{
+		// 1024 modules should be enough for most processes
 		std::vector<Module> modules;
-		modules.reserve(256u);
-
-		NT_PROCESS_BASIC_INFORMATION pbi = {};
-		undocumentedFunctions::NtQueryInformationProcess(handle, ProcessBasicInformation, &pbi, sizeof(pbi), NULL);
+		modules.reserve(1024u);
 
-		NT_PEB processPEB = {};
-		ReadProcessMemory(handle, pbi.PebBaseAddress, &processPEB, sizeof(NT_PEB));
+		windowsInternal::NT_PROCESS_BASIC_INFORMATION pbi = {};
+		windowsInternal::NtQueryInformationProcess(handle, windowsInternal::ProcessBasicInformation, &pbi, sizeof(pbi), NULL);
 
-		NT_PEB_LDR_DATA loaderData = {};
-		ReadProcessMemory(handle, processPEB.Ldr, &loaderData, sizeof(NT_PEB_LDR_DATA));
+		const windowsInternal::NT_PEB processPEB = ReadProcessMemory<windowsInternal::NT_PEB>(handle, pbi.PebBaseAddress);
+		const windowsInternal::NT_PEB_LDR_DATA loaderData = ReadProcessMemory<windowsInternal::NT_PEB_LDR_DATA>(handle, processPEB.Ldr);
 
 		LIST_ENTRY* listHeader = loaderData.InLoadOrderModuleList.Flink;
 		LIST_ENTRY* currentNode = listHeader;
 		do
 		{
-			NT_LDR_DATA_TABLE_ENTRY entry = {};
-			ReadProcessMemory(handle, currentNode, &entry, sizeof(NT_LDR_DATA_TABLE_ENTRY));
+			const windowsInternal::NT_LDR_DATA_TABLE_ENTRY entry = ReadProcessMemory<windowsInternal::NT_LDR_DATA_TABLE_ENTRY>(handle, currentNode);
 
-			currentNode = entry.InLoadOrderLinks.Flink;
-
-			WCHAR fullDllName[MAX_PATH] = {};
+			wchar_t fullDllName[MAX_PATH] = {};
 			if (entry.FullDllName.Length > 0)
 			{
 				ReadProcessMemory(handle, entry.FullDllName.Buffer, fullDllName, entry.FullDllName.Length);
 			}
 
 			modules.emplace_back(Module { fullDllName, entry.DllBase, entry.SizeOfImage });
+
+			currentNode = entry.InLoadOrderLinks.Flink;
+			if (currentNode == nullptr)
+			{
+				break;
+			}
 		}
 		while (listHeader != currentNode);
 
@@ -835,6 +535,126 @@ namespace process
 	}
 
 
+	uint32_t ConvertPageProtectionToExecutableProtection(uint32_t protection)
+	{
+		// cut off PAGE_GUARD, PAGE_NOCACHE, PAGE_WRITECOMBINE, and PAGE_REVERT_TO_FILE_MAP
+		const uint32_t extraBits = protection & 0xFFFFFF00u;
+		const uint32_t pageProtection = protection & 0x000000FFu;
+
+		switch (pageProtection)
+		{
+			case PAGE_NOACCESS:
+			case PAGE_READONLY:
+			case PAGE_READWRITE:
+			case PAGE_WRITECOPY:
+				return (pageProtection << 4u) | extraBits;
+
+			case PAGE_EXECUTE:
+			case PAGE_EXECUTE_READ:
+			case PAGE_EXECUTE_READWRITE:
+			case PAGE_EXECUTE_WRITECOPY:
+			default:
+				return protection;
+		}
+	}
+
+
+	bool IsWoW64(Handle handle)
+	{
+		// a WoW64 process has a PEB32 instead of a real PEB.
+		// if we get a meaningful pointer to this PEB32, the process is running under WoW64.
+		ULONG_PTR peb32 = 0u;
+		windowsInternal::NtQueryInformationProcess(handle, windowsInternal::ProcessWow64Information, &peb32, sizeof(ULONG_PTR), NULL);
+
+		return (peb32 != 0u);
+	}
+
+
+	Environment* CreateEnvironment(Handle handle)
+	{
+		const void* processEnvironment = nullptr;
+
+		const bool isWow64 = IsWoW64(handle);
+		if (!isWow64)
+		{
+			// this is either a 32-bit process running on 32-bit Windows, or a 64-bit process running on 64-bit Windows.
+			// the environment can be retrieved directly from the process' PEB and process parameters.
+			windowsInternal::NT_PROCESS_BASIC_INFORMATION pbi = {};
+			windowsInternal::NtQueryInformationProcess(handle, windowsInternal::ProcessBasicInformation, &pbi, sizeof(pbi), NULL);
+
+			const windowsInternal::NT_PEB peb = ReadProcessMemory<windowsInternal::NT_PEB>(handle, pbi.PebBaseAddress);
+			const windowsInternal::RTL_USER_PROCESS_PARAMETERS parameters = ReadProcessMemory<windowsInternal::RTL_USER_PROCESS_PARAMETERS>(handle, peb.ProcessParameters);
+
+			processEnvironment = parameters.Environment;
+		}
+		else
+		{
+			// this is a process running under WoW64.
+			// we must get the environment from the PEB32 of the process, rather than the "real" PEB.
+			ULONG_PTR peb32Wow64 = 0u;
+			windowsInternal::NtQueryInformationProcess(handle, windowsInternal::ProcessWow64Information, &peb32Wow64, sizeof(ULONG_PTR), NULL);
+
+			const windowsInternal::NT_PEB32 peb32 = ReadProcessMemory<windowsInternal::NT_PEB32>(handle, pointer::FromInteger<const void*>(peb32Wow64));
+			const windowsInternal::RTL_USER_PROCESS_PARAMETERS32 parameters32 = ReadProcessMemory<windowsInternal::RTL_USER_PROCESS_PARAMETERS32>(handle, pointer::FromInteger<const void*>(peb32.ProcessParameters32));
+
+			processEnvironment = pointer::FromInteger<const void*>(parameters32.Environment);		
+		}
+
+		if (!processEnvironment)
+		{
+			return nullptr;
+		}
+
+		// query the size of the page(s) the environment is stored in
+		::MEMORY_BASIC_INFORMATION memoryInfo = {};
+		const SIZE_T bytesInBuffer = ::VirtualQueryEx(handle, processEnvironment, &memoryInfo, sizeof(::MEMORY_BASIC_INFORMATION));
+		if (bytesInBuffer == 0u)
+		{
+			// operation failed, bail out
+			return nullptr;
+		}
+
+		Environment* environment = new Environment;
+		environment->size = memoryInfo.RegionSize - (reinterpret_cast<uintptr_t>(processEnvironment) - reinterpret_cast<uintptr_t>(memoryInfo.BaseAddress));
+		environment->data = new char[environment->size];
+
+		ReadProcessMemory(handle, processEnvironment, environment->data, environment->size);
+
+		return environment;
+	}
+
+
+	// BEGIN EPIC MOD - Allow passing environment block for linker
+	Environment* CreateEnvironmentFromMap(const TMap<FString, FString>& Pairs)
+	{
+		std::vector<wchar_t> environmentData;
+		for (const TPair<FString, FString>& Pair : Pairs)
+		{
+			FString Variable = FString::Printf(TEXT("%s=%s"), *Pair.Key, *Pair.Value);
+			environmentData.insert(environmentData.end(), *Variable, *Variable + (Variable.Len() + 1));
+		}
+		environmentData.push_back('\0');
+
+		Environment* environment = new Environment;
+		environment->size = environmentData.size();
+		environment->data = new char[environmentData.size() * sizeof(wchar_t)];
+		memcpy(environment->data, environmentData.data(), environmentData.size() * sizeof(wchar_t));
+		return environment;
+	}
+	// END EPIC MOD
+
+
+	void DestroyEnvironment(Environment*& environment)
+	{
+		if (environment)
+		{
+			delete[] environment->data;
+		}
+
+		memory::DeleteAndNull(environment);
+	}
+
+
 	void DumpMemory(Handle handle, const void* address, size_t size)
 	{
 		uint8_t* memory = new uint8_t[size];
diff --git a/Engine/Source/Developer/Windows/LiveCoding/Private/External/LC_Process.h b/Engine/Source/Developer/Windows/LiveCoding/Private/External/LC_Process.h
--- a/Engine/Source/Developer/Windows/LiveCoding/Private/External/LC_Process.h
+++ b/Engine/Source/Developer/Windows/LiveCoding/Private/External/LC_Process.h
@@ -25,6 +25,12 @@ namespace process
 		uint32_t sizeOfImage;
 	};
 
+	struct Environment
+	{
+		size_t size;
+		char* data;
+	};
+
 	typedef HANDLE Handle;
 
 	// returns the process ID for the calling process
@@ -36,16 +42,24 @@ namespace process
 		enum Enum : uint32_t
 		{
 			NONE = 0u,
-			REDIRECT_STDOUT = 1u << 0u
+			REDIRECT_STDOUT = 1u << 0u,
+			NO_WINDOW = 1u << 1u,
+			SUSPENDED = 1u << 2u
 		};
 	};
 
 	// spawns a new process
 	Context* Spawn(const wchar_t* exePath, const wchar_t* workingDirectory, const wchar_t* commandLine, const void* environmentBlock, uint32_t flags);
 
+	// resumes a process that was spawned in a suspended state
+	void ResumeMainThread(Context* context);
+
 	// waits until a spawned process has exited
 	unsigned int Wait(Context* context);
 
+	// waits until a process has exited
+	unsigned int Wait(Handle handle);
+
 	// destroys a spawned process
 	void Destroy(Context*& context);
 
@@ -69,6 +83,12 @@ namespace process
 	// returns the path to the executable of the calling process
 	std::wstring GetImagePath(void);
 
+	// returns the working directory of the calling process
+	std::wstring GetWorkingDirectory(void);
+
+	// returns the command line of the calling process
+	std::wstring GetCommandLine(void);
+
 	// returns the size of a module loaded into the virtual address space of a given process
 	uint32_t GetImageSize(Handle handle, void* moduleBase);
 
@@ -131,6 +151,26 @@ namespace process
 	std::vector<Module> EnumerateModules(Handle handle);
 
 
+	// converts any combination of page protection flags (e.g. PAGE_NOACCESS, PAGE_GUARD, ...) to protection flags
+	// that specify an executable page (e.g. PAGE_EXECUTE).
+	uint32_t ConvertPageProtectionToExecutableProtection(uint32_t protection);
+
+
+	// returns whether a process runs under Wow64 (32-bit emulation on 64-bit versions of Windows)
+	bool IsWoW64(Handle handle);
+
+	
+	// reads the environment of any process
+	Environment* CreateEnvironment(Handle handle);
+
+	// BEGIN EPIC MOD - Allow passing environment block for linker
+	Environment* CreateEnvironmentFromMap(const TMap<FString, FString>& Pairs);
+	// END EPIC MOD
+
+	// destroys an environment
+	void DestroyEnvironment(Environment*& environment);
+
+
 	// dumps raw memory for a given process
 	void DumpMemory(Handle handle, const void* address, size_t size);
 }
diff --git a/Engine/Source/Developer/Windows/LiveCoding/Private/External/LC_Restart.cpp b/Engine/Source/Developer/Windows/LiveCoding/Private/External/LC_Restart.cpp
new file mode 100644
--- /dev/null
+++ b/Engine/Source/Developer/Windows/LiveCoding/Private/External/LC_Restart.cpp
@@ -0,0 +1,78 @@
+// Copyright 2011-2019 Molecular Matters GmbH, all rights reserved.
+
+#include "LC_Restart.h"
+#include "LC_Event.h"
+#include "LC_PrimitiveNames.h"
+#include "LC_Process.h"
+#include "LC_Memory.h"
+
+
+namespace
+{
+	static Event* g_requestRestart = nullptr;
+	static Event* g_restartPrepared = nullptr;
+	static Event* g_executeRestart = nullptr;
+}
+
+
+void restart::Startup(void)
+{
+	const unsigned int processId = process::GetId();
+
+	g_requestRestart = new Event(primitiveNames::RequestRestart(processId).c_str(), Event::Type::AUTO_RESET);
+	g_restartPrepared = new Event(primitiveNames::PreparedRestart(processId).c_str(), Event::Type::AUTO_RESET);
+	g_executeRestart = new Event(primitiveNames::Restart(processId).c_str(), Event::Type::AUTO_RESET);
+}
+
+
+void restart::Shutdown(void)
+{
+	memory::DeleteAndNull(g_requestRestart);
+	memory::DeleteAndNull(g_restartPrepared);
+	memory::DeleteAndNull(g_executeRestart);
+}
+
+
+int restart::WasRequested(void)
+{
+	// check if the host requested a restart
+	if (g_requestRestart->TryWait())
+	{
+		return 1;
+	}
+
+	return 0;
+}
+
+
+void restart::Execute(lpp::RestartBehaviour behaviour, unsigned int exitCode)
+{
+	// tell the host that we finished running any custom client code and prepared everything for a restart
+	g_restartPrepared->Signal();
+
+	// wait until the host tells us to restart
+	g_executeRestart->Wait();
+
+	switch (behaviour)
+	{
+		// https://docs.microsoft.com/en-us/windows/win32/api/processthreadsapi/nf-processthreadsapi-exitprocess
+		case lpp::LPP_RESTART_BEHAVIOUR_DEFAULT_EXIT:
+			ExitProcess(exitCode);
+			break;
+
+		// https://docs.microsoft.com/en-us/cpp/c-runtime-library/reference/exit-exit-exit?view=vs-2019
+		case lpp::LPP_RESTART_BEHAVIOUR_EXIT_WITH_FLUSH:
+			exit(static_cast<int>(exitCode));
+			break;
+
+		// https://docs.microsoft.com/en-us/cpp/c-runtime-library/reference/exit-exit-exit?view=vs-2019
+		case lpp::LPP_RESTART_BEHAVIOUR_EXIT:
+			_Exit(static_cast<int>(exitCode));
+			break;
+
+		// https://docs.microsoft.com/en-us/windows/win32/api/processthreadsapi/nf-processthreadsapi-terminateprocess
+		case lpp::LPP_RESTART_BEHAVIOUR_INSTANT_TERMINATION:
+			::TerminateProcess(::GetCurrentProcess(), exitCode);
+			break;
+	};
+}
diff --git a/Engine/Source/Developer/Windows/LiveCoding/Private/External/LC_Restart.h b/Engine/Source/Developer/Windows/LiveCoding/Private/External/LC_Restart.h
new file mode 100644
--- /dev/null
+++ b/Engine/Source/Developer/Windows/LiveCoding/Private/External/LC_Restart.h
@@ -0,0 +1,15 @@
+// Copyright 2011-2019 Molecular Matters GmbH, all rights reserved.
+
+#pragma once
+
+#include "LPP_API.h"
+
+
+namespace restart
+{
+	void Startup(void);
+	void Shutdown(void);
+
+	int WasRequested(void);
+	void Execute(lpp::RestartBehaviour behaviour, unsigned int exitCode);
+}
diff --git a/Engine/Source/Developer/Windows/LiveCoding/Private/External/LC_Semaphore.cpp b/Engine/Source/Developer/Windows/LiveCoding/Private/External/LC_Semaphore.cpp
--- a/Engine/Source/Developer/Windows/LiveCoding/Private/External/LC_Semaphore.cpp
+++ b/Engine/Source/Developer/Windows/LiveCoding/Private/External/LC_Semaphore.cpp
@@ -23,37 +23,15 @@ void Semaphore::Signal(void)
 }
 
 
-void Semaphore::Wait(void)
+bool Semaphore::Wait(void)
 {
-	const DWORD result = ::WaitForSingleObject(m_sema, INFINITE);
-	switch (result)
-	{
-		case WAIT_OBJECT_0:
-			// semaphore was successfully signaled
-			break;
-
-		case WAIT_TIMEOUT:
-			// the operation timed out, which should never happen with a timeout of INFINITE
-			LC_ERROR_DEV("Semaphore timed out.");
-			break;
-
-		case WAIT_ABANDONED:
-			LC_ERROR_DEV("Wait() was called on a stale semaphore which was not released by the owning thread.");
-			break;
-
-		case WAIT_FAILED:
-			LC_ERROR_DEV("Failed to Wait() on a semaphore.");
-			break;
-
-		default:
-			break;
-	}
+	return WaitTimeout(INFINITE);
 }
 
 
-bool Semaphore::TryWait(void)
+bool Semaphore::WaitTimeout(unsigned int milliSeconds)
 {
-	const DWORD result = ::WaitForSingleObject(m_sema, 0);
+	const DWORD result = ::WaitForSingleObject(m_sema, milliSeconds);
 	switch (result)
 	{
 		case WAIT_OBJECT_0:
@@ -61,6 +39,11 @@ bool Semaphore::TryWait(void)
 			return true;
 
 		case WAIT_TIMEOUT:
+			// the operation timed out, which should never happen with a timeout of INFINITE
+			if (milliSeconds == INFINITE)
+			{
+				LC_ERROR_DEV("Semaphore timed out.");
+			}
 			return false;
 
 		case WAIT_ABANDONED:
@@ -75,3 +58,9 @@ bool Semaphore::TryWait(void)
 			return false;
 	}
 }
+
+
+bool Semaphore::TryWait(void)
+{
+	return WaitTimeout(0u);
+}
diff --git a/Engine/Source/Developer/Windows/LiveCoding/Private/External/LC_Semaphore.h b/Engine/Source/Developer/Windows/LiveCoding/Private/External/LC_Semaphore.h
--- a/Engine/Source/Developer/Windows/LiveCoding/Private/External/LC_Semaphore.h
+++ b/Engine/Source/Developer/Windows/LiveCoding/Private/External/LC_Semaphore.h
@@ -10,8 +10,17 @@ public:
 	Semaphore(unsigned int initialValue, unsigned int maximumValue);
 	~Semaphore(void);
 
+	// Signals the semaphore.
 	void Signal(void);
-	void Wait(void);
+
+	// Waits until the semaphore becomes signaled, blocking.
+	bool Wait(void);
+
+	// Waits until the semaphore becomes signaled, blocking until the timeout is reached.
+	// Returns whether the semaphore was signaled.
+	bool WaitTimeout(unsigned int milliSeconds);
+
+	// Returns whether the semaphore was signaled, non-blocking.
 	bool TryWait(void);
 
 private:
diff --git a/Engine/Source/Developer/Windows/LiveCoding/Private/External/LC_SyncPoint.cpp b/Engine/Source/Developer/Windows/LiveCoding/Private/External/LC_SyncPoint.cpp
--- a/Engine/Source/Developer/Windows/LiveCoding/Private/External/LC_SyncPoint.cpp
+++ b/Engine/Source/Developer/Windows/LiveCoding/Private/External/LC_SyncPoint.cpp
@@ -8,31 +8,11 @@
 namespace
 {
 	static CriticalSection g_syncPointCS;
+	static volatile bool g_isSyncPointUsed = false;
 
 	static Semaphore g_enterUserSyncPoint(0u, 1u);
 	static Semaphore g_leaveUserSyncPoint(0u, 1u);
 	static Semaphore g_dllSyncPoint(0u, 1u);
-
-	static volatile bool g_isSyncPointUsed = false;
-}
-
-
-// export the main synchronization point as C function
-void __cdecl LppSyncPoint(void)
-{
-	// mark the sync point as being used as soon as we enter this function the first time
-	g_syncPointCS.Enter();
-	{
-		g_isSyncPointUsed = true;
-	}
-	g_syncPointCS.Leave();
-
-	if (g_dllSyncPoint.TryWait())
-	{
-		// DLL code is currently inside sync point. tell DLL code that we are here, and wait for it to leave.
-		g_enterUserSyncPoint.Signal();
-		g_leaveUserSyncPoint.Wait();
-	}
 }
 
 
@@ -61,3 +41,20 @@ void syncPoint::Leave(void)
 	// tell user code that we're finished
 	g_leaveUserSyncPoint.Signal();
 }
+
+
+void syncPoint::EnterTarget(void)
+{
+	// mark the sync point as being used as soon as we enter this function the first time
+	{
+		CriticalSection::ScopedLock lock(&g_syncPointCS);
+		g_isSyncPointUsed = true;
+	}
+
+	if (g_dllSyncPoint.TryWait())
+	{
+		// DLL code is currently inside sync point. tell DLL code that we are here, and wait for it to leave.
+		g_enterUserSyncPoint.Signal();
+		g_leaveUserSyncPoint.Wait();
+	}
+}
diff --git a/Engine/Source/Developer/Windows/LiveCoding/Private/External/LC_SyncPoint.h b/Engine/Source/Developer/Windows/LiveCoding/Private/External/LC_SyncPoint.h
--- a/Engine/Source/Developer/Windows/LiveCoding/Private/External/LC_SyncPoint.h
+++ b/Engine/Source/Developer/Windows/LiveCoding/Private/External/LC_SyncPoint.h
@@ -6,6 +6,12 @@
 
 namespace syncPoint
 {
+	// Called when a host process wants to enter the sync point
 	void Enter(void);
+
+	// Called when a host process wants to leave the sync point
 	void Leave(void);
+
+	// Called by the LPP API when user code in a target process wants to enter the sync point
+	void EnterTarget(void);
 }
diff --git a/Engine/Source/Developer/Windows/LiveCoding/Private/External/LC_WindowsInternalFunctions.cpp b/Engine/Source/Developer/Windows/LiveCoding/Private/External/LC_WindowsInternalFunctions.cpp
new file mode 100644
--- /dev/null
+++ b/Engine/Source/Developer/Windows/LiveCoding/Private/External/LC_WindowsInternalFunctions.cpp
@@ -0,0 +1,11 @@
+// Copyright 2011-2019 Molecular Matters GmbH, all rights reserved.
+
+#include "LC_WindowsInternalFunctions.h"
+
+
+windowsInternal::Function<decltype(NtSuspendProcess)> windowsInternal::NtSuspendProcess("ntdll.dll", "NtSuspendProcess");
+windowsInternal::Function<decltype(NtResumeProcess)> windowsInternal::NtResumeProcess("ntdll.dll", "NtResumeProcess");
+windowsInternal::Function<decltype(NtWriteVirtualMemory)> windowsInternal::NtWriteVirtualMemory("ntdll.dll", "NtWriteVirtualMemory");
+windowsInternal::Function<decltype(NtQuerySystemInformation)> windowsInternal::NtQuerySystemInformation("ntdll.dll", "NtQuerySystemInformation");
+windowsInternal::Function<decltype(NtQueryInformationProcess)> windowsInternal::NtQueryInformationProcess("ntdll.dll", "NtQueryInformationProcess");
+windowsInternal::Function<decltype(NtContinue)> windowsInternal::NtContinue("ntdll.dll", "NtContinue");
diff --git a/Engine/Source/Developer/Windows/LiveCoding/Private/External/LC_WindowsInternalFunctions.h b/Engine/Source/Developer/Windows/LiveCoding/Private/External/LC_WindowsInternalFunctions.h
new file mode 100644
--- /dev/null
+++ b/Engine/Source/Developer/Windows/LiveCoding/Private/External/LC_WindowsInternalFunctions.h
@@ -0,0 +1,100 @@
+// Copyright 2011-2019 Molecular Matters GmbH, all rights reserved.
+
+#pragma once
+
+#include "LC_WindowsInternals.h"
+#include "LC_Logging.h"
+
+
+namespace windowsInternal
+{
+	// helper class that allows us to call any function in any Windows DLL, as long as it is exported and we know its signature.
+	// base template.
+	template <typename T>
+	class Function {};
+
+	// partial specialization for matching any function signature.
+	template <typename R, typename... Args>
+	class Function<R (Args...)>
+	{
+		typedef R (NTAPI *PtrToFunction)(Args...);
+
+	public:
+		inline Function(const char* moduleName, const char* functionName);
+
+		inline R operator()(Args... args) const;
+
+	private:
+		// helper for letting us check the result for arbitrary return types.
+		// base template.
+		template <typename T>
+		inline void CheckResult(T) const {}
+
+		// explicit specialization for NTSTATUS return values
+		template <>
+		inline void CheckResult(NTSTATUS result) const
+		{
+			if (!NT_SUCCESS(result))
+			{
+				LC_ERROR_USER("Call to function %s in module %s failed. Error: 0x%X", m_functionName, m_moduleName, result);
+			}
+		}
+
+		const char* m_moduleName;
+		const char* m_functionName;
+		PtrToFunction m_function;
+	};
+}
+
+
+template <typename R, typename... Args>
+inline windowsInternal::Function<R (Args...)>::Function(const char* moduleName, const char* functionName)
+	: m_moduleName(moduleName)
+	, m_functionName(functionName)
+	, m_function(nullptr)
+{
+	HMODULE module = ::GetModuleHandleA(moduleName);
+	if (!module)
+	{
+		LC_ERROR_USER("Cannot get handle for module %s", moduleName);
+		return;
+	}
+
+	m_function = reinterpret_cast<PtrToFunction>(reinterpret_cast<uintptr_t>(::GetProcAddress(module, functionName)));
+	if (!m_function)
+	{
+		LC_ERROR_USER("Cannot get address of function %s in module %s", functionName, moduleName);
+	}
+}
+
+
+template <typename R, typename... Args>
+inline R windowsInternal::Function<R (Args...)>::operator()(Args... args) const
+{
+	const R result = m_function(args...);
+	CheckResult(result);
+
+	return result;
+}
+
+
+// these are undocumented functions found in ntdll.dll.
+// we don't call them directly, but use them for "extracting" their signature.
+extern "C" windowsInternal::NTSTATUS NtSuspendProcess(HANDLE ProcessHandle);
+extern "C" windowsInternal::NTSTATUS NtResumeProcess(HANDLE ProcessHandle);
+extern "C" windowsInternal::NTSTATUS NtWriteVirtualMemory(HANDLE ProcessHandle, PVOID BaseAddress, PVOID Buffer, ULONG NumberOfBytesToWrite, PULONG NumberOfBytesWritten);
+extern "C" windowsInternal::NTSTATUS NtQuerySystemInformation(windowsInternal::NT_SYSTEM_INFORMATION_CLASS SystemInformationClass, PVOID SystemInformation, ULONG SystemInformationLength, PULONG ReturnLength);
+extern "C" windowsInternal::NTSTATUS NtQueryInformationProcess(HANDLE ProcessHandle, windowsInternal::NT_PROCESS_INFORMATION_CLASS ProcessInformationClass, PVOID ProcessInformation, ULONG ProcessInformationLength, PULONG ReturnLength);
+extern "C" windowsInternal::NTSTATUS NtContinue(CONTEXT* ThreadContext, BOOLEAN RaiseAlert);
+
+
+// cache important undocumented functions
+namespace windowsInternal
+{
+	extern Function<decltype(NtSuspendProcess)> NtSuspendProcess;
+	extern Function<decltype(NtResumeProcess)> NtResumeProcess;
+	extern Function<decltype(NtWriteVirtualMemory)> NtWriteVirtualMemory;
+	extern Function<decltype(NtQuerySystemInformation)> NtQuerySystemInformation;
+	extern Function<decltype(NtQueryInformationProcess)> NtQueryInformationProcess;
+	extern Function<decltype(NtContinue)> NtContinue;
+}
diff --git a/Engine/Source/Developer/Windows/LiveCoding/Private/External/LC_WindowsInternals.h b/Engine/Source/Developer/Windows/LiveCoding/Private/External/LC_WindowsInternals.h
new file mode 100644
--- /dev/null
+++ b/Engine/Source/Developer/Windows/LiveCoding/Private/External/LC_WindowsInternals.h
@@ -0,0 +1,208 @@
+// Copyright 2011-2019 Molecular Matters GmbH, all rights reserved.
+
+#pragma once
+
+#include "Windows/WindowsHWrapper.h"
+
+// https://docs.microsoft.com/en-us/openspecs/windows_protocols/ms-erref/87fba13e-bf06-450e-83b1-9241dc81e781
+// found in <ntstatus.h>
+#define STATUS_INFO_LENGTH_MISMATCH		((windowsInternal::NTSTATUS)0xC0000004L)
+
+// found in <winternl.h>
+#define NT_SUCCESS(Status)				((windowsInternal::NTSTATUS)(Status) >= 0)
+
+
+namespace windowsInternal
+{
+	// most of these types are defined in <winternl.h>, but we cannot include that header because it lacks
+	// a few undocumented values/members that we need for our purposes.
+	// note that our definitions aren't complete either, we only define what we need.
+
+	// https://docs.microsoft.com/en-us/openspecs/windows_protocols/ms-erref/87fba13e-bf06-450e-83b1-9241dc81e781
+	typedef LONG NTSTATUS;
+
+	// https://www.geoffchappell.com/studies/windows/km/ntoskrnl/api/ex/sysinfo/class.htm
+	// found in <winternl.h>
+	enum NT_SYSTEM_INFORMATION_CLASS
+	{
+		SystemProcessInformation = 5
+	};
+
+	// https://docs.microsoft.com/en-us/openspecs/windows_protocols/ms-tsts/a11e7129-685b-4535-8d37-21d4596ac057
+	// found in <wdm.h> in WDK
+	struct NT_CLIENT_ID
+	{
+		HANDLE UniqueProcess;
+		HANDLE UniqueThread;
+	};
+
+	// https://docs.microsoft.com/en-us/openspecs/windows_protocols/ms-tsts/e82d73e4-cedb-4077-9099-d58f3459722f
+	struct NT_SYSTEM_THREAD_INFORMATION
+	{
+		LARGE_INTEGER KernelTime;
+		LARGE_INTEGER UserTime;
+		LARGE_INTEGER CreateTime;
+		ULONG WaitTime;
+		PVOID StartAddress;
+		NT_CLIENT_ID ClientId;
+		LONG Priority;
+		LONG BasePriority;
+		ULONG ContextSwitches;
+		ULONG ThreadState;
+		ULONG WaitReason;
+	};
+
+	// https://docs.microsoft.com/en-us/openspecs/windows_protocols/ms-tsts/c90753f2-f9f9-490d-846d-6bdd41eae7f8
+	// found in <winternl.h>
+	struct NT_UNICODE_STRING
+	{
+		USHORT Length;
+		USHORT MaximumLength;
+		PWSTR Buffer;
+	};
+
+	// https://www.geoffchappell.com/studies/windows/km/ntoskrnl/api/ex/sysinfo/process.htm
+	// found in <winternl.h>
+	struct NT_SYSTEM_PROCESS_INFORMATION
+	{
+		ULONG NextEntryOffset;
+		ULONG NumberOfThreads;
+		LARGE_INTEGER WorkingSetPrivateSize;
+		ULONG HardFaultCount;
+		ULONG NumberOfThreadsHighWatermark;
+		ULONGLONG CycleTime;
+		LARGE_INTEGER CreateTime;
+		LARGE_INTEGER UserTime;
+		LARGE_INTEGER KernelTime;
+		NT_UNICODE_STRING ImageName;
+		LONG BasePriority;
+		PVOID UniqueProcessId;
+		PVOID InheritedFromUniqueProcessId;
+		ULONG HandleCount;
+		ULONG SessionId;
+		ULONG_PTR UniqueProcessKey;
+		ULONG_PTR PeakVirtualSize;
+		ULONG_PTR VirtualSize;
+		ULONG PageFaultCount;
+		ULONG_PTR PeakWorkingSetSize;
+		ULONG_PTR WorkingSetSize;
+		ULONG_PTR QuotaPeakPagedPoolUsage;
+		ULONG_PTR QuotaPagedPoolUsage;
+		ULONG_PTR QuotaPeakNonPagedPoolUsage;
+		ULONG_PTR QuotaNonPagedPoolUsage;
+		ULONG_PTR PagefileUsage;
+		ULONG_PTR PeakPagefileUsage;
+		ULONG_PTR PrivatePageCount;
+		LARGE_INTEGER ReadOperationCount;
+		LARGE_INTEGER WriteOperationCount;
+		LARGE_INTEGER OtherOperationCount;
+		LARGE_INTEGER ReadTransferCount;
+		LARGE_INTEGER WriteTransferCount;
+		LARGE_INTEGER OtherTransferCount;
+		NT_SYSTEM_THREAD_INFORMATION Threads[1];	// variable size data
+	};
+
+	// https://docs.microsoft.com/en-us/windows/desktop/api/winternl/nf-winternl-ntqueryinformationprocess
+	enum NT_PROCESS_INFORMATION_CLASS
+	{
+		ProcessBasicInformation = 0,
+		ProcessWow64Information = 26
+	};
+
+	// https://docs.microsoft.com/en-us/windows/desktop/api/winternl/ns-winternl-_rtl_user_process_parameters
+	// https://www.nirsoft.net/kernel_struct/vista/RTL_USER_PROCESS_PARAMETERS.html
+	// found in <winternl.h>
+	struct RTL_USER_PROCESS_PARAMETERS
+	{
+		BYTE Reserved1[16];
+		PVOID Reserved2[10];
+		NT_UNICODE_STRING ImagePathName;
+		NT_UNICODE_STRING CommandLine;
+		PWSTR Environment;
+	};
+
+	// similar to RTL_USER_PROCESS_PARAMETERS, altered to behave as a struct containing 32-bit pointers in a 64-bit environment
+	struct RTL_USER_PROCESS_PARAMETERS32
+	{
+		char Reserved[72];
+		ULONG Environment;
+	};
+
+	// https://www.geoffchappell.com/studies/windows/win32/ntdll/structs/ldr_data_table_entry.htm
+	// found in <winternl.h>
+	struct NT_LDR_DATA_TABLE_ENTRY
+	{
+		LIST_ENTRY InLoadOrderLinks;
+		LIST_ENTRY InMemoryOrderLinks;
+		LIST_ENTRY InInitializationOrderLinks;
+		PVOID DllBase;
+		PVOID EntryPoint;
+		ULONG SizeOfImage;
+		NT_UNICODE_STRING FullDllName;
+		NT_UNICODE_STRING BaseDllName;
+		ULONG Flags;
+		USHORT LoadCount;
+		USHORT ObsoleteLoadCount;
+		USHORT TlsIndex;
+		LIST_ENTRY HashLinks;
+	};
+
+	// https://docs.microsoft.com/en-us/windows/desktop/api/winternl/ns-winternl-_peb_ldr_data
+	// https://www.geoffchappell.com/studies/windows/win32/ntdll/structs/peb_ldr_data.htm
+	struct NT_PEB_LDR_DATA
+	{
+		ULONG Length;
+		BOOLEAN Initialized;
+		PVOID SsHandle;
+		LIST_ENTRY InLoadOrderModuleList;
+		LIST_ENTRY InMemoryOrderModuleList;
+		LIST_ENTRY InInitializationOrderModuleList;
+		PVOID EntryInProgress;
+		BOOLEAN ShutdownInProgress;
+		HANDLE ShutdownThreadId;
+	};
+
+	// found in <winternl.h>
+	typedef VOID (NTAPI *NT_PS_POST_PROCESS_INIT_ROUTINE)(VOID);
+
+	// https://docs.microsoft.com/en-us/windows/desktop/api/winternl/ns-winternl-_peb
+	struct NT_PEB
+	{
+		BYTE Reserved1[2];
+		BYTE BeingDebugged;
+		BYTE Reserved2[1];
+		PVOID Reserved3[2];
+		NT_PEB_LDR_DATA* Ldr;
+		RTL_USER_PROCESS_PARAMETERS* ProcessParameters;
+		PVOID Reserved4[3];
+		PVOID AtlThunkSListPtr;
+		PVOID Reserved5;
+		ULONG Reserved6;
+		PVOID Reserved7;
+		ULONG Reserved8;
+		ULONG AtlThunkSListPtr32;
+		PVOID Reserved9[45];
+		BYTE Reserved10[96];
+		NT_PS_POST_PROCESS_INIT_ROUTINE* PostProcessInitRoutine;
+		BYTE Reserved11[128];
+		PVOID Reserved12[1];
+		ULONG SessionId;
+	};
+
+	// similar to NT_PEB, altered to behave as a struct containing 32-bit pointers in a 64-bit environment
+	struct NT_PEB32
+	{
+		char Reserved[16];
+		ULONG ProcessParameters32;
+	};
+
+	// https://docs.microsoft.com/en-us/windows/desktop/api/winternl/nf-winternl-ntqueryinformationprocess
+	struct NT_PROCESS_BASIC_INFORMATION
+	{
+		PVOID Reserved1;
+		NT_PEB* PebBaseAddress;
+		PVOID Reserved2[2];
+		ULONG_PTR UniqueProcessId;
+		PVOID Reserved3;
+	};
+}
diff --git a/Engine/Source/Developer/Windows/LiveCoding/Private/External/LPP_API.h b/Engine/Source/Developer/Windows/LiveCoding/Private/External/LPP_API.h
--- a/Engine/Source/Developer/Windows/LiveCoding/Private/External/LPP_API.h
+++ b/Engine/Source/Developer/Windows/LiveCoding/Private/External/LPP_API.h
@@ -52,3 +52,43 @@
 	__pragma(section(LPP_COMPILE_ERROR_MESSAGE_SECTION, read))																									\
 	__declspec(allocate(LPP_COMPILE_ERROR_MESSAGE_SECTION)) extern void (*LPP_IDENTIFIER(lpp_compile_error_message_hook_function))(const wchar_t*) = &_function
 
+/******************************************************************************/
+/* API                                                                        */
+/******************************************************************************/
+
+// version string
+#define LPP_VERSION "1.5.0"
+
+#ifdef __cplusplus
+#	define LPP_NS_BEGIN			namespace lpp {
+#	define LPP_NS_END			}
+#	define LPP_API				inline
+#else
+#	define LPP_NS_BEGIN
+#	define LPP_NS_END
+#	define LPP_API				static inline
+#endif
+
+// helper macros to call a function in a DLL with an arbitrary signature without a compiler warning
+#define LPP_CALL1(_module, _functionName, _return, _args1)								((_return (__cdecl*)(_args1))((uintptr_t)GetProcAddress(_module, _functionName)))
+#define LPP_CALL2(_module, _functionName, _return, _args1, _args2)						((_return (__cdecl*)(_args1, _args2))((uintptr_t)GetProcAddress(_module, _functionName)))
+#define LPP_CALL3(_module, _functionName, _return, _args1, _args2, _args3)				((_return (__cdecl*)(_args1, _args2, _args3))((uintptr_t)GetProcAddress(_module, _functionName)))
+#define LPP_CALL4(_module, _functionName, _return, _args1, _args2, _args3, _args4)		((_return (__cdecl*)(_args1, _args2, _args3, _args4))((uintptr_t)GetProcAddress(_module, _functionName)))
+
+
+LPP_NS_BEGIN
+
+enum RestartBehaviour
+{
+	LPP_RESTART_BEHAVIOUR_DEFAULT_EXIT,				// ExitProcess()
+	LPP_RESTART_BEHAVIOUR_EXIT_WITH_FLUSH,			// exit()
+	LPP_RESTART_BEHAVIOUR_EXIT,						// _Exit()
+	LPP_RESTART_BEHAVIOUR_INSTANT_TERMINATION		// TerminateProcess
+};
+
+LPP_NS_END
+
+
+#undef LPP_CALL1
+#undef LPP_CALL2
+#undef LPP_CALL3
diff --git a/Engine/Source/Developer/Windows/LiveCodingServer/Private/External/LC_AppSettings.cpp b/Engine/Source/Developer/Windows/LiveCodingServer/Private/External/LC_AppSettings.cpp
--- a/Engine/Source/Developer/Windows/LiveCodingServer/Private/External/LC_AppSettings.cpp
+++ b/Engine/Source/Developer/Windows/LiveCodingServer/Private/External/LC_AppSettings.cpp
@@ -257,6 +257,15 @@ void appSettings::Startup(const wchar_t* group)
 		false
 	);
 
+	g_showTimestamps = new SettingBool
+	(
+		group,
+		L"show_timestamps",
+		L"Show timestamps",
+		L"Specifies whether output will show timestamps",
+		false
+	);
+
 	g_wordWrapOutput = new SettingBool
 	(
 		group,
@@ -484,6 +493,7 @@ void appSettings::Shutdown(void)
 	delete g_compileShortcut;
 
 	delete g_showUndecoratedNames;
+	delete g_showTimestamps;
 	delete g_wordWrapOutput;
 	delete g_enableDevLog;
 	delete g_enableTelemetryLog;
@@ -630,7 +640,7 @@ void appSettings::UpdateAmalgamatedCppFileExtensions(void)
 
 void appSettings::ApplySettingBool(const char* const settingName, bool value)
 {
-	const unsigned int COUNT = 20u;
+	const unsigned int COUNT = 21u;
 	SettingBool* settings[COUNT] =
 	{
 		g_showFullPathInTitle,
@@ -640,6 +650,7 @@ void appSettings::ApplySettingBool(const char* const settingName, bool value)
 		g_minimizeOnClose,
 		g_keepTrayIcon,
 		g_showUndecoratedNames,
+		g_showTimestamps,
 		g_wordWrapOutput,
 		g_enableDevLog,
 		g_enableTelemetryLog,
@@ -793,6 +804,7 @@ extern SettingString* appSettings::g_playSoundOnError = nullptr;
 extern SettingShortcut* appSettings::g_compileShortcut = nullptr;
 
 extern SettingBool* appSettings::g_showUndecoratedNames = nullptr;
+extern SettingBool* appSettings::g_showTimestamps = nullptr;
 extern SettingBool* appSettings::g_wordWrapOutput = nullptr;
 extern SettingBool* appSettings::g_enableDevLog = nullptr;
 extern SettingBool* appSettings::g_enableTelemetryLog = nullptr;
diff --git a/Engine/Source/Developer/Windows/LiveCodingServer/Private/External/LC_AppSettings.h b/Engine/Source/Developer/Windows/LiveCodingServer/Private/External/LC_AppSettings.h
--- a/Engine/Source/Developer/Windows/LiveCodingServer/Private/External/LC_AppSettings.h
+++ b/Engine/Source/Developer/Windows/LiveCodingServer/Private/External/LC_AppSettings.h
@@ -78,6 +78,7 @@ namespace appSettings
 
 	// logging
 	extern SettingBool* g_showUndecoratedNames;
+	extern SettingBool* g_showTimestamps;
 	extern SettingBool* g_wordWrapOutput;
 	extern SettingBool* g_enableDevLog;
 	extern SettingBool* g_enableTelemetryLog;
diff --git a/Engine/Source/Developer/Windows/LiveCodingServer/Private/External/LC_COMThread.cpp b/Engine/Source/Developer/Windows/LiveCodingServer/Private/External/LC_COMThread.cpp
new file mode 100644
--- /dev/null
+++ b/Engine/Source/Developer/Windows/LiveCodingServer/Private/External/LC_COMThread.cpp
@@ -0,0 +1,185 @@
+// Copyright 2011-2019 Molecular Matters GmbH, all rights reserved.
+
+#include "LC_COMThread.h"
+#include "LC_Platform.h"
+#include "LC_Logging.h"
+#include "Windows/MinimalWindowsAPI.h"
+#include <objidl.h>
+
+#include "Windows/AllowWindowsPlatformAtomics.h"
+
+namespace
+{
+	// helper class to deal with RPC_E_CALL_REJECTED errors when trying to use automation
+	// https://web.archive.org/web/20131209031731/http://msdn.microsoft.com/en-us/library/ms228772(VS.80).aspx
+	// https://infosys.beckhoff.com/english.php?content=../content/1033/tc3_automationinterface/54043195771173899.html&id=
+	class COMMessageFilter : public IMessageFilter
+	{
+	public:
+		virtual HRESULT STDMETHODCALLTYPE QueryInterface(
+			/* [in] */ REFIID riid,
+			/* [iid_is][out] */ _COM_Outptr_ void __RPC_FAR *__RPC_FAR *ppvObject)
+		{
+			// always set out parameter to NULL, validating it first
+			if (!ppvObject)
+			{
+				return E_INVALIDARG;
+			}
+
+			*ppvObject = NULL;
+			if (riid == IID_IUnknown || riid == IID_IMessageFilter)
+			{
+				// increment the reference count and return the pointer
+				*ppvObject = this;
+				AddRef();
+
+				return NOERROR;
+			}
+
+			return E_NOINTERFACE;
+		}
+
+		virtual ULONG STDMETHODCALLTYPE AddRef(void)
+		{
+			::InterlockedIncrement(&m_refCount);
+			return m_refCount;
+		}
+
+		virtual ULONG STDMETHODCALLTYPE Release(void)
+		{
+			// decrement the object's internal counter and delete the interface if zero
+			ULONG refCount = ::InterlockedDecrement(&m_refCount);
+			if (0 == m_refCount)
+			{
+				delete this;
+			}
+
+			return refCount;
+		}
+
+		// Handle incoming thread requests
+		virtual DWORD STDMETHODCALLTYPE HandleInComingCall(
+			/* [annotation][in] */
+			_In_  DWORD dwCallType,
+			/* [annotation][in] */
+			_In_  HTASK htaskCaller,
+			/* [annotation][in] */
+			_In_  DWORD dwTickCount,
+			/* [annotation][in] */
+			_In_opt_  LPINTERFACEINFO lpInterfaceInfo)
+		{
+			LC_UNUSED(dwCallType);
+			LC_UNUSED(htaskCaller);
+			LC_UNUSED(dwTickCount);
+			LC_UNUSED(lpInterfaceInfo);
+
+			return SERVERCALL_ISHANDLED;
+		}
+
+		// Thread call was rejected, so try again
+		virtual DWORD STDMETHODCALLTYPE RetryRejectedCall(
+			/* [annotation][in] */
+			_In_  HTASK htaskCallee,
+			/* [annotation][in] */
+			_In_  DWORD dwTickCount,
+			/* [annotation][in] */
+			_In_  DWORD dwRejectType)
+		{
+			LC_UNUSED(htaskCallee);
+			LC_UNUSED(dwTickCount);
+
+			if (dwRejectType == SERVERCALL_RETRYLATER)
+			{
+				// retry the thread call immediately if return in rage [0, 99]
+				return 99;
+			}
+
+			// too busy, cancel call
+			return static_cast<DWORD>(-1);
+		}
+
+		virtual DWORD STDMETHODCALLTYPE MessagePending(
+			/* [annotation][in] */
+			_In_  HTASK htaskCallee,
+			/* [annotation][in] */
+			_In_  DWORD dwTickCount,
+			/* [annotation][in] */
+			_In_  DWORD dwPendingType)
+		{
+			LC_UNUSED(htaskCallee);
+			LC_UNUSED(dwTickCount);
+			LC_UNUSED(dwPendingType);
+
+			return PENDINGMSG_WAITDEFPROCESS;
+		}
+
+	private:
+		volatile ULONG m_refCount = 0ul;
+	};
+}
+
+
+COMThread::COMThread(void)
+	: m_function()
+	, m_functionAvailableEvent(nullptr, Event::Type::AUTO_RESET)
+	, m_functionFinishedExecutingEvent(nullptr, Event::Type::AUTO_RESET)
+	, m_leaveThreadEvent(nullptr, Event::Type::AUTO_RESET)
+	, m_internalThread(nullptr)
+{
+	// launch the thread that takes care of calling COM functions
+	m_internalThread = thread::Create("Live++ COM", 65536u, &COMThread::ThreadFunction, this);
+}
+
+
+COMThread::~COMThread(void)
+{
+	// make the internal thread break out of its loop and wait until it terminates
+	m_leaveThreadEvent.Signal();
+	m_functionAvailableEvent.Signal();
+
+	thread::Join(m_internalThread);
+	thread::Close(m_internalThread);
+}
+
+
+unsigned int COMThread::ThreadFunction(void)
+{
+	const HRESULT result = ::CoInitialize(NULL);
+	if (result != S_OK)
+	{
+		LC_ERROR_DEV("Could not initialize COM. Error: 0x%X", result);
+	}
+
+	for (;;)
+	{
+		// wait until function becomes available
+		m_functionAvailableEvent.Wait();
+
+		// gracefully exit the thread
+		if (m_leaveThreadEvent.TryWait())
+		{
+			m_functionFinishedExecutingEvent.Signal();
+			break;
+		}
+
+		// install a COM message filter and execute the function.
+		// the message filter deals with RPC_E_CALL_REJECTED errors automatically without having to wrap
+		// every call into error checks and retries.
+		COMMessageFilter* messageFilter = new COMMessageFilter;
+		IMessageFilter* oldFilter = nullptr;
+		::CoRegisterMessageFilter(messageFilter, &oldFilter);
+		{
+			(*m_function.function)(m_function.context, m_function.returnValueAddr);
+		}
+		::CoRegisterMessageFilter(oldFilter, nullptr);
+
+		// tell waiting thread that function finished executing
+		m_functionFinishedExecutingEvent.Signal();
+	}
+
+	::CoUninitialize();
+
+	return 0u;
+}
+
+#include "Windows/HideWindowsPlatformAtomics.h"
diff --git a/Engine/Source/Developer/Windows/LiveCodingServer/Private/External/LC_COMThread.h b/Engine/Source/Developer/Windows/LiveCodingServer/Private/External/LC_COMThread.h
new file mode 100644
--- /dev/null
+++ b/Engine/Source/Developer/Windows/LiveCodingServer/Private/External/LC_COMThread.h
@@ -0,0 +1,77 @@
+// Copyright 2011-2019 Molecular Matters GmbH, all rights reserved.
+
+#pragma once
+
+#include "LC_Event.h"
+#include "LC_Thread.h"
+
+
+// helper class that is the only one calling COM-related functions and methods.
+// this is needed because most of the COM VS automation interfaces only can be called from a single thread.
+class COMThread
+{
+	struct COMFunction
+	{
+		typedef void (FunctionType)(void*, void*);
+
+		FunctionType* function;
+		void* context;
+		void* returnValueAddr;
+	};
+
+public:
+	COMThread(void);
+	~COMThread(void);
+
+	// helper function that runs any given function in the internal COM thread and waits until execution has finished
+	template <typename F, typename... Args>
+	inline typename std::result_of<F(Args...)>::type CallInThread(F ptrToFunction, Args&&... args)
+	{
+		typedef typename std::result_of<F(Args...)>::type ReturnValue;
+
+		// a helper lambda that calls the given functions with the provided arguments.
+		// because this lambda needs to capture, it cannot be converted to a function pointer implicitly,
+		// and therefore cannot be used as a COM function directly.
+		// however, we can solve this by another indirection.
+		auto captureLambda = [ptrToFunction, args...]()
+		{
+			return (*ptrToFunction)(args...);
+		};
+
+		// helper to make the following easier to read
+		typedef decltype(captureLambda) CaptureLambdaType;
+
+		// here's the trick: we generate another capture-less lambda that has the same signature as our COM function,
+		// and internally cast the given object to its original type, calling the lambda with captures from within
+		// this capture-less lambda.
+		// it is OK to leave both the lambda as well as the return value on the stack. this function returns only
+		// after the COM function has been called by the internal thread.
+		ReturnValue result = {};
+		m_function.function = [](void* captureLambda, void* returnValueAddr)
+		{
+			CaptureLambdaType* realCaptureLambda = static_cast<CaptureLambdaType*>(captureLambda);
+			ReturnValue* realReturnValueAddr = static_cast<ReturnValue*>(returnValueAddr);
+			*realReturnValueAddr = (*realCaptureLambda)();
+		};
+
+		m_function.context = &captureLambda;
+		m_function.returnValueAddr = &result;
+
+		// signal event that function in thread should now be called
+		m_functionAvailableEvent.Signal();
+
+		// wait for signal that the function has finished executing
+		m_functionFinishedExecutingEvent.Wait();
+
+		return result;
+	}
+
+private:
+	unsigned int ThreadFunction(void);
+
+	COMFunction m_function;
+	Event m_functionAvailableEvent;
+	Event m_functionFinishedExecutingEvent;
+	Event m_leaveThreadEvent;
+	thread::Handle m_internalThread;
+};
diff --git a/Engine/Source/Developer/Windows/LiveCodingServer/Private/External/LC_Compiler.cpp b/Engine/Source/Developer/Windows/LiveCodingServer/Private/External/LC_Compiler.cpp
--- a/Engine/Source/Developer/Windows/LiveCodingServer/Private/External/LC_Compiler.cpp
+++ b/Engine/Source/Developer/Windows/LiveCodingServer/Private/External/LC_Compiler.cpp
@@ -7,6 +7,7 @@
 #include "LC_Environment.h"
 #include "LC_CriticalSection.h"
 #include "LC_Logging.h"
+#include "LC_TimeStamp.h"
 
 
 namespace
@@ -24,18 +25,18 @@ namespace
 		{
 			for (auto it = m_cache.begin(); it != m_cache.end(); ++it)
 			{
-				environment::Block* block = it->second;
-				environment::DestroyBlock(block);
+				process::Environment* env = it->second;
+				process::DestroyEnvironment(env);
 			}
 		}
 
-		void Insert(const wchar_t* key, environment::Block* value)
+		void Insert(const wchar_t* key, process::Environment* value)
 		{
 	        // BEGIN EPIC MOD - Allow passing environment block for linker
 			auto it = m_cache.find(key);
 			if (it != m_cache.end())
 			{
-				environment::DestroyBlock(it->second);
+				process::DestroyEnvironment(it->second);
 				it->second = value;
 				return;
 			}
@@ -44,7 +45,7 @@ namespace
 			m_cache[key] = value;
 		}
 
-		const environment::Block* Fetch(const wchar_t* key)
+		const process::Environment* Fetch(const wchar_t* key)
 		{
 			const auto it = m_cache.find(key);
 			if (it != m_cache.end())
@@ -56,7 +57,7 @@ namespace
 		}
 
 	private:
-		types::unordered_map<std::wstring, environment::Block*> m_cache;
+		types::unordered_map<std::wstring, process::Environment*> m_cache;
 	};
 
 	static CompilerEnvironmentCache g_compilerEnvironmentCache;
@@ -127,7 +128,7 @@ namespace
 }
 
 
-const environment::Block* compiler::CreateEnvironmentCacheEntry(const wchar_t* absolutePathToCompilerExe)
+const process::Environment* compiler::CreateEnvironmentCacheEntry(const wchar_t* absolutePathToCompilerExe)
 {
 	LC_LOG_DEV("Creating environment cache entry for %S", absolutePathToCompilerExe);
 
@@ -167,50 +168,73 @@ const environment::Block* compiler::CreateEnvironmentCacheEntry(const wchar_t* a
 		const file::Attributes& attributes = file::GetAttributes(pathToVcvars.c_str());
 		if (file::DoesExist(attributes))
 		{
-			// this is the correct vcvars*.bat
-
-			// quote path to batch file
-			std::wstring vcvarsBat(L"\"");
-			vcvarsBat += path;
-			vcvarsBat += relativePathsToVcvarsFile[i];
-			vcvarsBat += L"\"";
-
-			// now that we have the path to the vcvars*.bat to call, construct a command that first invokes
-			// the batch file and then outputs the environment variables to a file.
-			const std::wstring& tempFile = file::CreateTempFile();
-			const std::wstring& cmdPath = environment::GetVariable(L"COMSPEC");
-
-			// tell cmd.exe to execute commands, and quote all filenames involved
-			std::wstring commandLine(L"/c \"");
-			commandLine += vcvarsBat;
-			commandLine += L" && set > \"";
-			commandLine += tempFile;
-			commandLine += L"\"\"";
-
-			process::Context* vcvarsProcess = process::Spawn(cmdPath.c_str(), nullptr, commandLine.c_str(), nullptr, process::SpawnFlags::NONE);
-			const unsigned int exitCode = process::Wait(vcvarsProcess);
-			process::Destroy(vcvarsProcess);
-
-			if (exitCode == 0u)
+			// this is the correct vcvars*.bat.
+			// we need to invoke the command shell, run the .bat file, and extract the process' environment to cache it for later use.
+			// this is slightly more complicated than it needs to be, because we cannot simply run a command in the shell and grab
+			// the environment without knowing if the .bat has finished running. similarly, we cannot grab the environment once
+			// the shell process has terminated already.
+			// here's what we do:
+			// - create a unique, temporary file. this is used for letting the shell signal that it has finished executing
+			// - run the .bat file in the shell
+			// - tell the shell to delete the temporary file and pause
+			// - wait until the temporary file is deleted. it is now save to grab the environment
+			// - grab the environment and store it in our cache
+			// - terminate the shell process
+
+			// create an empty temporary file
+			const std::wstring& tempFileDeletedByCmd = file::GenerateTempFilename();
+			char emptyData = 0u;
+			file::CreateFileWithData(tempFileDeletedByCmd.c_str(), &emptyData, 0u);
+
+			// tell cmd.exe to execute commands, and quote all filenames involved.
+			// the whole command needs to be quoted as well.
+			const std::wstring& cmdPath = environment::GetVariable(L"COMSPEC", L"cmd");
+			std::wstring commandLine(L"/c \"\"");
+			commandLine += pathToVcvars;
+			commandLine += L"\" && del /q \"";
+			commandLine += tempFileDeletedByCmd;
+			commandLine += L"\" && pause \"";
+
+			process::Context* vcvarsProcess = process::Spawn(cmdPath.c_str(), nullptr, commandLine.c_str(), nullptr, process::SpawnFlags::NO_WINDOW);
+
+			// wait until the temporary file is deleted.
+			// busy waiting like this is not very nice, but happens only once or twice during startup, and is called from a separate thread anyway.
+			const uint64_t startTimestamp = timeStamp::Get();
+			bool shownWarning = false;
+			for (;;)
 			{
-				// the temporary file now holds the full environment block after vcvars*.bat has executed.
-				// load it and insert it into the cache.
-				environment::Block* block = environment::CreateBlockFromFile(tempFile.c_str());
+				const file::Attributes& tempAttributes = file::GetAttributes(tempFileDeletedByCmd.c_str());
+				if (file::DoesExist(tempAttributes))
+				{
+					thread::Sleep(10u);
+				}
+				else
 				{
-					CriticalSection::ScopedLock lock(&g_compilerCacheCS);
-					g_compilerEnvironmentCache.Insert(absolutePathToCompilerExe, block);
+					break;
+				}
 
-					if (block)
-					{
-						environment::DumpBlockData(vcvarsBat.c_str(), block);
-					}
+				// show a warning in case this takes longer than 5 seconds.
+				// this can happen for some users:
+				// https://developercommunity.visualstudio.com/content/problem/51179/vsdevcmdbat-or-vcvarsallbat-excecution-takes-a-ver.html
+				const uint64_t delta = timeStamp::Get() - startTimestamp;
+				if ((timeStamp::ToSeconds(delta) >= 5.0) && (!shownWarning))
+				{
+					LC_WARNING_USER("Prewarming compiler/linker environment for %S is taking suspiciously long.", pathToVcvars.c_str());
+					shownWarning = true;
 				}
+			}
 
-				return block;
+			// grab the environment from the process and insert it into the cache
+			process::Environment* environment = process::CreateEnvironment(vcvarsProcess->pi.hProcess);
+			{
+				CriticalSection::ScopedLock lock(&g_compilerCacheCS);
+				g_compilerEnvironmentCache.Insert(absolutePathToCompilerExe, environment);
 			}
 
-			LC_WARNING_USER("vcvars*.bat could not be invoked at %S", vcvarsBat.c_str());
-			return nullptr;
+			process::Terminate(vcvarsProcess->pi.hProcess);
+			process::Destroy(vcvarsProcess);
+
+			return environment;
 		}
 		else
 		{
@@ -223,28 +247,28 @@ const environment::Block* compiler::CreateEnvironmentCacheEntry(const wchar_t* a
 }
 
 
-const environment::Block* compiler::GetEnvironmentFromCache(const wchar_t* absolutePathToCompilerExe)
+const process::Environment* compiler::GetEnvironmentFromCache(const wchar_t* absolutePathToCompilerExe)
 {
 	CriticalSection::ScopedLock lock(&g_compilerCacheCS);
 	return g_compilerEnvironmentCache.Fetch(absolutePathToCompilerExe);
 }
 
 
-const environment::Block* compiler::UpdateEnvironmentCache(const wchar_t* absolutePathToCompilerExe)
+const process::Environment* compiler::UpdateEnvironmentCache(const wchar_t* absolutePathToCompilerExe)
 {
-	const environment::Block* block = GetEnvironmentFromCache(absolutePathToCompilerExe);
-	if (block)
+	const process::Environment* environment = GetEnvironmentFromCache(absolutePathToCompilerExe);
+	if (environment)
 	{
-		return block;
+		return environment;
 	}
 
 	return CreateEnvironmentCacheEntry(absolutePathToCompilerExe);
 }
 
 // BEGIN EPIC MOD - Allow passing environment block for linker
-void compiler::AddEnvironmentToCache(const wchar_t* absolutePathToCompilerExe, environment::Block* block)
+void compiler::AddEnvironmentToCache(const wchar_t* absolutePathToCompilerExe, process::Environment* environment)
 {
 	CriticalSection::ScopedLock lock(&g_compilerCacheCS);
-	g_compilerEnvironmentCache.Insert(absolutePathToCompilerExe, block);
+	g_compilerEnvironmentCache.Insert(absolutePathToCompilerExe, environment);
 }
 // END EPIC MOD
diff --git a/Engine/Source/Developer/Windows/LiveCodingServer/Private/External/LC_Compiler.h b/Engine/Source/Developer/Windows/LiveCodingServer/Private/External/LC_Compiler.h
--- a/Engine/Source/Developer/Windows/LiveCodingServer/Private/External/LC_Compiler.h
+++ b/Engine/Source/Developer/Windows/LiveCodingServer/Private/External/LC_Compiler.h
@@ -4,26 +4,26 @@
 
 #include "CoreTypes.h"
 
-namespace environment
+namespace process
 {
-	struct Block;
+	struct Environment;
 }
 
 
 namespace compiler
 {
 	// creates a new entry in the cache for the given compiler .exe, and returns it
-	const environment::Block* CreateEnvironmentCacheEntry(const wchar_t* absolutePathToCompilerExe);
+	const process::Environment* CreateEnvironmentCacheEntry(const wchar_t* absolutePathToCompilerExe);
 
 	// gets the environment for a given compiler .exe from the cache.
 	// returns nullptr if the environment is not yet in the cache.
-	const environment::Block* GetEnvironmentFromCache(const wchar_t* absolutePathToCompilerExe);
+	const process::Environment* GetEnvironmentFromCache(const wchar_t* absolutePathToCompilerExe);
 
 	// helper function that either creates a new entry in the cache if none exists yet,
 	// or returns the one found in the cache.
-	const environment::Block* UpdateEnvironmentCache(const wchar_t* absolutePathToCompilerExe);
+	const process::Environment* UpdateEnvironmentCache(const wchar_t* absolutePathToCompilerExe);
 	
 	// BEGIN EPIC MOD - Allow overriding environment block for tools
-	void AddEnvironmentToCache(const wchar_t* absolutePathToCompilerExe, environment::Block* block);
+	void AddEnvironmentToCache(const wchar_t* absolutePathToCompilerExe, process::Environment* environment);
 	// END EPIC MOD
 }
diff --git a/Engine/Source/Developer/Windows/LiveCodingServer/Private/External/LC_Environment.cpp b/Engine/Source/Developer/Windows/LiveCodingServer/Private/External/LC_Environment.cpp
deleted file mode 100644
diff --git a/Engine/Source/Developer/Windows/LiveCodingServer/Private/External/LC_Environment.h b/Engine/Source/Developer/Windows/LiveCodingServer/Private/External/LC_Environment.h
deleted file mode 100644
diff --git a/Engine/Source/Developer/Windows/LiveCodingServer/Private/External/LC_FileUtil.cpp b/Engine/Source/Developer/Windows/LiveCodingServer/Private/External/LC_FileUtil.cpp
--- a/Engine/Source/Developer/Windows/LiveCodingServer/Private/External/LC_FileUtil.cpp
+++ b/Engine/Source/Developer/Windows/LiveCodingServer/Private/External/LC_FileUtil.cpp
@@ -252,7 +252,7 @@ bool file::IsRelativePath(const wchar_t* path)
 }
 
 
-std::wstring file::CreateTempFile(void)
+std::wstring file::GenerateTempFilename(void)
 {
 	wchar_t path[MAX_PATH] = {};
 	const DWORD pathLength = ::GetTempPathW(MAX_PATH, path);
diff --git a/Engine/Source/Developer/Windows/LiveCodingServer/Private/External/LC_FileUtil.h b/Engine/Source/Developer/Windows/LiveCodingServer/Private/External/LC_FileUtil.h
--- a/Engine/Source/Developer/Windows/LiveCodingServer/Private/External/LC_FileUtil.h
+++ b/Engine/Source/Developer/Windows/LiveCodingServer/Private/External/LC_FileUtil.h
@@ -46,7 +46,7 @@ namespace file
 	bool IsRelativePath(const wchar_t* path);
 
 	// creates a unique, temporary absolute filename, e.g. [REDACTED_USER_HOME]\AppData\Local\Temp\ABCD.tmp
-	std::wstring CreateTempFile(void);
+	std::wstring GenerateTempFilename(void);
 
 	// creates a file, storing the given data
 	bool CreateFileWithData(const wchar_t* path, const void* data, size_t size);
diff --git a/Engine/Source/Developer/Windows/LiveCodingServer/Private/External/LC_ImmutableString.cpp b/Engine/Source/Developer/Windows/LiveCodingServer/Private/External/LC_ImmutableString.cpp
--- a/Engine/Source/Developer/Windows/LiveCodingServer/Private/External/LC_ImmutableString.cpp
+++ b/Engine/Source/Developer/Windows/LiveCodingServer/Private/External/LC_ImmutableString.cpp
@@ -13,9 +13,7 @@ namespace
 {
 	static inline uint32_t Hash(const char* key, size_t length)
 	{
-		// need to cast here because the xxhash API returns an unsigned int, even though it calculates a
-		// 32-bit hash internally.
-		return static_cast<uint32_t>(XXH32(key, length * sizeof(char), 0u));
+		return XXH32(key, length * sizeof(char), 0u);
 	}
 
 	static inline const char* Clone(const char* str, size_t length)
diff --git a/Engine/Source/Developer/Windows/LiveCodingServer/Private/External/LC_LiveModule.cpp b/Engine/Source/Developer/Windows/LiveCodingServer/Private/External/LC_LiveModule.cpp
--- a/Engine/Source/Developer/Windows/LiveCodingServer/Private/External/LC_LiveModule.cpp
+++ b/Engine/Source/Developer/Windows/LiveCodingServer/Private/External/LC_LiveModule.cpp
@@ -304,7 +304,7 @@ namespace
 		// create a temporary file that acts as a so-called response file for the compiler, and contains
 		// the whole compiler command-line. this is done because the latter can get very long, longer
 		// than the limit of 32k characters.
-		const std::wstring responseFilePath = file::CreateTempFile();
+		const std::wstring responseFilePath = file::GenerateTempFilename();
 		file::CreateFileWithData(responseFilePath.c_str(), compilerOptions.c_str(), compilerOptions.size() * sizeof(char));
 
 		std::wstring compilerCommandLine;
@@ -320,8 +320,8 @@ namespace
 		compilerCommandLine += responseFilePath;
 		compilerCommandLine += L"\"";
 
-		const environment::Block* envBlock = compiler::GetEnvironmentFromCache(compilerPath.c_str());
-		const void* envBlockData = envBlock ? environment::GetBlockData(envBlock) : nullptr;
+		const process::Environment* environment = compiler::GetEnvironmentFromCache(compilerPath.c_str());
+		const void* environmentData = environment ? environment->data : nullptr;
 		std::wstring workingDirectory = string::ToWideString(compiland->workingDirectory);
 
 		// if the working directory does not exist, use the compiler's directory instead.
@@ -336,7 +336,7 @@ namespace
 
 		LC_LOG_USER("Compiling %s %s", isPartOfAmalgamation ? "split file" : "file", normalizedObjPath.c_str());
 
-		process::Context* processContext = process::Spawn(compilerPath.c_str(), workingDirectory.c_str(), compilerCommandLine.c_str(), envBlockData, process::SpawnFlags::REDIRECT_STDOUT);
+		process::Context* processContext = process::Spawn(compilerPath.c_str(), workingDirectory.c_str(), compilerCommandLine.c_str(), environmentData, process::SpawnFlags::REDIRECT_STDOUT | process::SpawnFlags::NO_WINDOW);
 		const unsigned int exitCode = process::Wait(processContext);
 		const wchar_t* compilerOutput = processContext->stdoutData.c_str();
 
@@ -3031,7 +3031,7 @@ LiveModule::ErrorType::Enum LiveModule::Update(FileAttributeCache* fileCache, Di
 	// create a temporary file that acts as a so-called response file for the linker, and contains
 	// the whole linker command-line. this is done because the latter can get very long, longer
 	// than the limit of 32k characters.
-	const std::wstring responseFilePath = file::CreateTempFile();
+	const std::wstring responseFilePath = file::GenerateTempFilename();
 	file::CreateFileWithData(responseFilePath.c_str(), linkerOptions.c_str(), linkerOptions.size() * sizeof(wchar_t));
 
 	std::wstring linkerCommandLine = file::GetFilename(linkerPath);
@@ -3039,12 +3039,12 @@ LiveModule::ErrorType::Enum LiveModule::Update(FileAttributeCache* fileCache, Di
 	linkerCommandLine += responseFilePath;
 	linkerCommandLine += L"\"";
 
-	const environment::Block* linkerEnvBlock = compiler::GetEnvironmentFromCache(linkerPath.c_str());
-	const void* linkerEnvBlockData = linkerEnvBlock ? environment::GetBlockData(linkerEnvBlock) : nullptr;
+	const process::Environment* linkerEnvironment = compiler::GetEnvironmentFromCache(linkerPath.c_str());
+	const void* linkerEnvironmentData = linkerEnvironment ? linkerEnvironment->data : nullptr;
 
 	GLiveCodingServer->GetStatusChangeDelegate().ExecuteIfBound(L"Linking patch...");
 
-	process::Context* linkerProcessContext = process::Spawn(linkerPath.c_str(), linkerWorkingDirectory.c_str(), linkerCommandLine.c_str(), linkerEnvBlockData, process::SpawnFlags::REDIRECT_STDOUT);
+	process::Context* linkerProcessContext = process::Spawn(linkerPath.c_str(), linkerWorkingDirectory.c_str(), linkerCommandLine.c_str(), linkerEnvironmentData, process::SpawnFlags::REDIRECT_STDOUT | process::SpawnFlags::NO_WINDOW);
 	const unsigned int linkerExitCode = process::Wait(linkerProcessContext);
 
 	const double linkerTime = linkScope.ReadSeconds();
diff --git a/Engine/Source/Developer/Windows/LiveCodingServer/Private/External/LC_LiveProcess.cpp b/Engine/Source/Developer/Windows/LiveCodingServer/Private/External/LC_LiveProcess.cpp
--- a/Engine/Source/Developer/Windows/LiveCodingServer/Private/External/LC_LiveProcess.cpp
+++ b/Engine/Source/Developer/Windows/LiveCodingServer/Private/External/LC_LiveProcess.cpp
@@ -3,17 +3,28 @@
 #include "LC_LiveProcess.h"
 #include "LC_HeartBeat.h"
 #include "LC_CodeCave.h"
+#include "LC_Event.h"
+#include "LC_PrimitiveNames.h"
+#include "LC_VisualStudioAutomation.h"
 
 
-LiveProcess::LiveProcess(process::Handle processHandle, unsigned int processId, unsigned int commandThreadId, const void* jumpToSelf, const DuplexPipe* pipe)
+LiveProcess::LiveProcess(process::Handle processHandle, unsigned int processId, unsigned int commandThreadId, const void* jumpToSelf, const DuplexPipe* pipe,
+	const wchar_t* imagePath, const wchar_t* commandLine, const wchar_t* workingDirectory, const void* environment, size_t environmentSize)
 	: m_processHandle(processHandle)
 	, m_processId(processId)
 	, m_commandThreadId(commandThreadId)
 	, m_jumpToSelf(jumpToSelf)
 	, m_pipe(pipe)
+	, m_imagePath(imagePath)
+	, m_commandLine(commandLine)
+	, m_workingDirectory(workingDirectory)
+	, m_environment(environment, environmentSize)
 	, m_imagesTriedToLoad()
 	, m_heartBeatDelta(0ull)
+	, m_vsDebugger(nullptr)
+	, m_vsDebuggerThreads()
 	, m_codeCave(nullptr)
+	, m_restartState(RestartState::DEFAULT)
 {
 	m_imagesTriedToLoad.reserve(256u);
 }
@@ -40,6 +51,66 @@ bool LiveProcess::MadeProgress(void) const
 }
 
 
+void LiveProcess::HandleDebuggingPreCompile(void)
+{
+	if (!MadeProgress())
+	{
+		// this process did not make progress.
+		// try to find a debugger that's currently debugging our process.
+		m_vsDebugger = visualStudio::FindDebuggerForProcess(m_processId);
+		if (m_vsDebugger)
+		{
+			// found a debugger.
+			// enumerate all threads, freeze every thread but the command thread, and let the debugger resume.
+			// this "halts" the process but lets the command thread act on commands sent by us.
+			m_vsDebuggerThreads = visualStudio::EnumerateThreads(m_vsDebugger);
+			if (m_vsDebuggerThreads.size() > 0u)
+			{
+				visualStudio::FreezeThreads(m_vsDebugger, m_vsDebuggerThreads);
+				visualStudio::ThawThread(m_vsDebugger, m_vsDebuggerThreads, m_commandThreadId);
+				visualStudio::Resume(m_vsDebugger);
+
+				LC_SUCCESS_USER("Automating debugger attached to process (PID: %d)", m_processId);
+
+				return;
+			}
+		}
+
+		// no debugger could be found or an error occurred.
+		// continue by installing a code cave.
+		LC_LOG_USER("Failed to automate debugger attached to process (PID: %d), using fallback mechanism", m_processId);
+		LC_SUCCESS_USER("Waiting for client process (PID: %d), hit 'Continue' (F5 in Visual Studio) if being held in the debugger", m_processId);
+	}
+
+	// this process either made progress and is not held in the debugger, or we failed automating the debugger.
+	// "halt" this process by installing a code cave.
+	InstallCodeCave();
+}
+
+
+void LiveProcess::HandleDebuggingPostCompile(void)
+{
+	if (m_codeCave)
+	{
+		// we installed a code cave previously, remove it
+		UninstallCodeCave();
+	}
+	else if (m_vsDebugger)
+	{
+		// we automated the debugger previously. break into the debugger again and resume all threads.
+		// when debugging a C# project that calls into C++ code, the VS debugger sometimes creates new MTA threads in between our PreCompile and PostCompile calls.
+		// try getting a new list of threads and thaw them all as well.
+		visualStudio::Break(m_vsDebugger);
+		
+		types::vector<EnvDTE::ThreadPtr> newDebuggerThreads = visualStudio::EnumerateThreads(m_vsDebugger);
+		visualStudio::ThawThreads(m_vsDebugger, newDebuggerThreads);
+		visualStudio::ThawThreads(m_vsDebugger, m_vsDebuggerThreads);
+	}
+
+	m_vsDebugger = nullptr;
+}
+
+
 void LiveProcess::InstallCodeCave(void)
 {
 	m_codeCave = new CodeCave(m_processHandle, m_processId, m_commandThreadId, m_jumpToSelf);
@@ -72,6 +143,58 @@ bool LiveProcess::TriedToLoadImage(const executable::Header& imageHeader) const
 	return (m_imagesTriedToLoad.find(imageHeader) != m_imagesTriedToLoad.end());
 }
 
+
+bool LiveProcess::PrepareForRestart(void)
+{
+	// signal to the target process that a restart for this process was requested
+	Event requestRestart(primitiveNames::RequestRestart(m_processId).c_str(), Event::Type::AUTO_RESET);
+	requestRestart.Signal();
+
+	// the client code in the target is now inside the lpp::lppWantsRestart() code block.
+	// wait until it calls lpp::lppRestart() after finishing custom client code.
+	// give the client 10 seconds to finish up.
+	Event restartPrepared(primitiveNames::PreparedRestart(m_processId).c_str(), Event::Type::AUTO_RESET);
+	const bool success = restartPrepared.WaitTimeout(10u * 1000u);
+	if (success)
+	{
+		m_restartState = RestartState::SUCCESSFUL_PREPARE;
+		return true;
+	}
+	else
+	{
+		LC_ERROR_USER("Client did not respond to restart request within 10 seconds, aborting restart (PID: %d)", m_processId);
+		m_restartState = RestartState::FAILED_PREPARE;
+		return false;
+	}
+}
+
+
+void LiveProcess::Restart(void)
+{
+	if (m_restartState == RestartState::SUCCESSFUL_PREPARE)
+	{
+		// in case PrepareForRestart was successful, the client is now waiting for the signal to restart.
+		// tell the client to initiate a restart now.
+		Event executeRestart(primitiveNames::Restart(m_processId).c_str(), Event::Type::AUTO_RESET);
+		executeRestart.Signal();
+
+		// wait until the client terminates
+		process::Wait(m_processHandle);
+
+		// restart the target application
+		process::Spawn(m_imagePath.c_str(), m_workingDirectory.c_str(), m_commandLine.c_str(), m_environment.GetData(), process::SpawnFlags::NONE);
+
+		m_restartState = RestartState::SUCCESSFUL;
+	}
+}
+
+
+bool LiveProcess::WasSuccessfulRestart(void) const
+{
+	return (m_restartState == RestartState::SUCCESSFUL);
+}
+
+
 // BEGIN EPIC MOD - Allow lazy-loading modules
 void LiveProcess::AddLazyLoadedModule(const std::wstring moduleName, Windows::HMODULE moduleBase)
 {
@@ -108,4 +231,4 @@ Windows::HMODULE LiveProcess::GetLazyLoadedModuleBase(const std::wstring& module
 		return iter->second.m_moduleBase;
 	}
 }
-// END EPIC MOD
+// END EPIC MOD
\ No newline at end of file
diff --git a/Engine/Source/Developer/Windows/LiveCodingServer/Private/External/LC_LiveProcess.h b/Engine/Source/Developer/Windows/LiveCodingServer/Private/External/LC_LiveProcess.h
--- a/Engine/Source/Developer/Windows/LiveCodingServer/Private/External/LC_LiveProcess.h
+++ b/Engine/Source/Developer/Windows/LiveCodingServer/Private/External/LC_LiveProcess.h
@@ -5,7 +5,9 @@
 #include "CoreTypes.h"
 #include "LC_Process.h"
 #include "LC_Executable.h"
+#include "LC_MemoryBlock.h"
 #include "LC_Types.h"
+#include "LC_VisualStudio.h"
 
 
 class DuplexPipe;
@@ -14,24 +16,33 @@ class CodeCave;
 class LiveProcess
 {
 public:
-	LiveProcess(process::Handle processHandle, unsigned int processId, unsigned int commandThreadId, const void* jumpToSelf, const DuplexPipe* pipe);
-
+	LiveProcess(process::Handle processHandle, unsigned int processId, unsigned int commandThreadId, const void* jumpToSelf, const DuplexPipe* pipe,
+		const wchar_t* imagePath, const wchar_t* commandLine, const wchar_t* workingDirectory, const void* environment, size_t environmentSize);
 
 	void ReadHeartBeatDelta(const wchar_t* const processGroupName);
 
 	// returns whether this process made some progress, based on the heart beat received from the client
 	bool MadeProgress(void) const;
 
+	// handles any debugging mechanism that might currently debug our process before we compile changes
+	void HandleDebuggingPreCompile(void);
+
+	// handles any debugging mechanism that might currently debug our process after we compiled changes
+	void HandleDebuggingPostCompile(void);
 
 	void InstallCodeCave(void);
 	void UninstallCodeCave(void);
 
-
 	void AddLoadedImage(const executable::Header& imageHeader);
 	void RemoveLoadedImage(const executable::Header& imageHeader);
 	bool TriedToLoadImage(const executable::Header& imageHeader) const;
 
 
+	bool PrepareForRestart(void);
+	void Restart(void);
+	bool WasSuccessfulRestart(void) const;
+
+
 	inline process::Handle GetProcessHandle(void) const
 	{
 		return m_processHandle;
@@ -83,6 +94,11 @@ private:
 	const void* m_jumpToSelf;
 	const DuplexPipe* m_pipe;
 
+	std::wstring m_imagePath;
+	std::wstring m_commandLine;
+	std::wstring m_workingDirectory;
+	MemoryBlock m_environment;
+
 	// BEGIN EPIC MOD - Add build arguments
 	std::wstring m_buildArguments;
 	// END EPIC MOD
@@ -103,5 +119,32 @@ private:
 	types::unordered_set<executable::Header> m_imagesTriedToLoad;
 
 	uint64_t m_heartBeatDelta;
+
+#if LC_WITH_VISUAL_STUDIO_DTE
+	// for handling communication with the VS debugger
+	EnvDTE::DebuggerPtr m_vsDebugger;
+	types::vector<EnvDTE::ThreadPtr> m_vsDebuggerThreads;
+#endif
+
+	// fallback in case communication with the VS debugger is not possible
 	CodeCave* m_codeCave;
+
+	// restart
+	struct RestartState
+	{
+		enum Enum
+		{
+			DEFAULT,
+			FAILED_PREPARE,
+			SUCCESSFUL_PREPARE,
+			SUCCESSFUL
+		};
+	};
+
+	RestartState::Enum m_restartState;
+
+	LC_DISABLE_COPY(LiveProcess);
+	LC_DISABLE_MOVE(LiveProcess);
+	LC_DISABLE_ASSIGNMENT(LiveProcess);
+	LC_DISABLE_MOVE_ASSIGNMENT(LiveProcess);
 };
diff --git a/Engine/Source/Developer/Windows/LiveCodingServer/Private/External/LC_MemoryBlock.cpp b/Engine/Source/Developer/Windows/LiveCodingServer/Private/External/LC_MemoryBlock.cpp
new file mode 100644
--- /dev/null
+++ b/Engine/Source/Developer/Windows/LiveCodingServer/Private/External/LC_MemoryBlock.cpp
@@ -0,0 +1,17 @@
+// Copyright 2011-2019 Molecular Matters GmbH, all rights reserved.
+
+#include "LC_MemoryBlock.h"
+
+
+MemoryBlock::MemoryBlock(const void* data, size_t size)
+	: m_data(malloc(size))
+	, m_size(size)
+{
+	memcpy(m_data, data, size);
+}
+
+
+MemoryBlock::~MemoryBlock(void)
+{
+	free(m_data);
+}
diff --git a/Engine/Source/Developer/Windows/LiveCodingServer/Private/External/LC_MemoryBlock.h b/Engine/Source/Developer/Windows/LiveCodingServer/Private/External/LC_MemoryBlock.h
new file mode 100644
--- /dev/null
+++ b/Engine/Source/Developer/Windows/LiveCodingServer/Private/External/LC_MemoryBlock.h
@@ -0,0 +1,32 @@
+// Copyright 2011-2019 Molecular Matters GmbH, all rights reserved.
+
+#pragma once
+
+#include "CoreTypes.h"
+#include "LC_Platform.h"
+
+class MemoryBlock
+{
+public:
+	MemoryBlock(const void* data, size_t size);
+	~MemoryBlock(void);
+
+	inline const void* GetData(void) const
+	{
+		return m_data;
+	}
+
+	inline size_t GetSize(void) const
+	{
+		return m_size;
+	}
+
+private:
+	void* m_data;
+	size_t m_size;
+
+	LC_DISABLE_COPY(MemoryBlock);
+	LC_DISABLE_MOVE(MemoryBlock);
+	LC_DISABLE_ASSIGNMENT(MemoryBlock);
+	LC_DISABLE_MOVE_ASSIGNMENT(MemoryBlock);
+};
diff --git a/Engine/Source/Developer/Windows/LiveCodingServer/Private/External/LC_ServerCommandThread.cpp b/Engine/Source/Developer/Windows/LiveCodingServer/Private/External/LC_ServerCommandThread.cpp
--- a/Engine/Source/Developer/Windows/LiveCodingServer/Private/External/LC_ServerCommandThread.cpp
+++ b/Engine/Source/Developer/Windows/LiveCodingServer/Private/External/LC_ServerCommandThread.cpp
@@ -25,6 +25,8 @@
 #include "LC_Allocators.h"
 #include "LC_DuplexPipeClient.h"
 #include "LC_MemoryStream.h"
+#include "LC_NamedSharedMemory.h"
+#include "LC_VisualStudioAutomation.h"
 #include <mmsystem.h>
 
 // unreachable code
@@ -39,22 +41,6 @@ namespace
 {
 	static telemetry::Accumulator g_loadedModuleSize("Module size");
 
-	struct InitializeCOM
-	{
-		InitializeCOM(void)
-		{
-			const HRESULT success = ::CoInitialize(NULL);
-			if (success != S_OK)
-			{
-				LC_LOG_DEV("Could not initialize COM. Error: %d", success);
-			}
-		}
-
-		~InitializeCOM(void)
-		{
-			::CoUninitialize();
-		}
-	};
 
 
 	static void AddVirtualDrive(void)
@@ -111,7 +97,13 @@ ServerCommandThread::ServerCommandThread(MainFrame* mainFrame, const wchar_t* co
 	, m_commandThreads()
 	, m_manualRecompileTriggered(false)
 	, m_liveModuleToModifiedOrNewObjFiles()
+	, m_restartCS()
+	, m_restartJob(nullptr)
+	, m_restartedProcessCount(0u)
+	, m_restartedProcessIdToDebugger()
 {
+	visualStudio::Startup();
+
 	m_serverThread = thread::Create("Live coding server", 64u * 1024u, &ServerCommandThread::ServerThread, this);
 	m_compileThread = thread::Create("Live coding compilation", 64u * 1024u, &ServerCommandThread::CompileThread, this);
 
@@ -125,11 +117,89 @@ ServerCommandThread::ServerCommandThread(MainFrame* mainFrame, const wchar_t* co
 
 ServerCommandThread::~ServerCommandThread(void)
 {
-	// note that we deliberately do *nothing* here.
+	// note that we deliberately do almost *nothing* here.
 	// this is only called when Live++ is being torn down anyway, so we leave cleanup to the OS.
 	// otherwise we could run into races when trying to terminate the thread that might currently be doing
 	// some intensive work.
 	delete m_directoryCache;
+
+	visualStudio::Shutdown();
+}
+
+
+void ServerCommandThread::RestartTargets(void)
+{
+	// protect against concurrent compilation
+	m_restartCS.Enter();
+
+	// EPIC REMOVED: g_theApp.GetMainFrame()->SetBusy(true);
+	// EPIC REMOVED: g_theApp.GetMainFrame()->ChangeStatusBarText(L"Restarting target applications...");
+
+	LC_LOG_USER("---------- Restarting target applications ----------");
+
+	// prevent current Live++ instance from shutting down by associating it with a new job object to keep it alive
+	if (!m_restartJob)
+	{
+		m_restartJob = ::CreateJobObjectW(NULL, primitiveNames::JobGroup(m_processGroupName).c_str());
+		::AssignProcessToJobObject(m_restartJob, ::GetCurrentProcess());
+	}
+
+	// protect against m_liveProcesses being accessed when processes restart and register themselves with this Live++ instance
+	CriticalSection::ScopedLock lock(&m_actionCS);
+
+	// remove processes that were successfully restarted last time
+	for (auto processIt = m_liveProcesses.begin(); processIt != m_liveProcesses.end(); /* nothing */)
+	{
+		LiveProcess* liveProcess = *processIt;
+		if (liveProcess->WasSuccessfulRestart())
+		{
+			process::Handle processHandle = liveProcess->GetProcessHandle();
+			process::Close(processHandle);
+
+			// tell live modules to remove this process
+			const size_t moduleCount = m_liveModules.size();
+			for (size_t j = 0u; j < moduleCount; ++j)
+			{
+				LiveModule* liveModule = m_liveModules[j];
+				liveModule->UnregisterProcess(liveProcess);
+			}
+
+			delete liveProcess;
+
+			processIt = m_liveProcesses.erase(processIt);
+		}
+		else
+		{
+			++processIt;
+		}
+	}
+
+	// try preparing all processes for a restart
+	const size_t count = m_liveProcesses.size();
+	for (size_t i = 0u; i < count; ++i)
+	{
+		LiveProcess* liveProcess = m_liveProcesses[i];
+		const bool success = liveProcess->PrepareForRestart();
+		if (success)
+		{
+			++m_restartedProcessCount;
+
+			// check if a VS debugger is currently attached to the process about to restart
+			const unsigned int processId = liveProcess->GetProcessId();
+			EnvDTE::DebuggerPtr debugger = visualStudio::FindDebuggerAttachedToProcess(processId);
+			if (debugger)
+			{
+				m_restartedProcessIdToDebugger.emplace(processId, debugger);
+			}
+		}
+	}
+
+	// restart all successfully prepared processes
+	for (size_t i = 0u; i < count; ++i)
+	{
+		LiveProcess* liveProcess = m_liveProcesses[i];
+		liveProcess->Restart();
+	}
 }
 
 
@@ -433,7 +503,9 @@ void ServerCommandThread::PrewarmCompilerEnvironmentCache(void)
 
 unsigned int ServerCommandThread::ServerThread(void)
 {
-	InitializeCOM initCOM;
+	// keep named shared memory alive so that restarted processes don't try spawning new Live++ instances
+	NamedSharedMemory sharedMemory(primitiveNames::StartupNamedSharedMemory(m_processGroupName).c_str());
+	sharedMemory.Write(::GetCurrentProcessId());
 
 	// inter process event for telling client that server is ready
 	Event serverReadyEvent(primitiveNames::ServerReadyEvent(m_processGroupName).c_str(), Event::Type::AUTO_RESET);
@@ -870,6 +942,9 @@ unsigned int ServerCommandThread::CompileThread(void)
 
 	for (;;)
 	{
+		// protect against concurrent restarts
+		m_restartCS.Enter();
+
 		const int shortcutValue = appSettings::g_compileShortcut->GetValue();
 		keyShortcut.AssignCode(shortcut::GetVirtualKeyCode(shortcutValue));
 
@@ -967,14 +1042,14 @@ unsigned int ServerCommandThread::CompileThread(void)
 
 			if (!didAllProcessesMakeProgress)
 			{
-				// install a code cave for all processes.
-				// this ensures that if a process is currently being held in the debugger, the process will
-				// not make progress in terms of new instructions being executed after continuing it in the debugger.
+				// not all processes made progress.
+				// this usually means that at least one of them is currently being debugged.
+				// let each process handle this.
 				const size_t processCount = m_liveProcesses.size();
 				for (size_t i = 0u; i < processCount; ++i)
 				{
 					LiveProcess* liveProcess = m_liveProcesses[i];
-					liveProcess->InstallCodeCave();
+					liveProcess->HandleDebuggingPreCompile();
 				}
 
 				// don't allow the exception handler dialog to be shown when continuing in the debugger with F5
@@ -984,15 +1059,6 @@ unsigned int ServerCommandThread::CompileThread(void)
 			// wait until all command threads/clients are ready to go. we might not be getting commands
 			// from a client because it is being held in the debugger.
 			{
-				if (didAllProcessesMakeProgress)
-				{
-					LC_SUCCESS_USER("Waiting for client(s)");
-				}
-				else
-				{
-					LC_SUCCESS_USER("Waiting for client(s), hit 'Continue' (F5) if being held in the debugger");
-				}
-
 				CriticalSection::ScopedLock lock(&m_connectionCS);
 
 				const size_t count = m_commandThreads.size();
@@ -1039,12 +1105,11 @@ unsigned int ServerCommandThread::CompileThread(void)
 
 			if (!didAllProcessesMakeProgress)
 			{
-				// remove all code caves
 				const size_t processCount = m_liveProcesses.size();
 				for (size_t i = 0u; i < processCount; ++i)
 				{
 					LiveProcess* liveProcess = m_liveProcesses[i];
-					liveProcess->UninstallCodeCave();
+					liveProcess->HandleDebuggingPostCompile();
 				}
 
 				// remove the lock on the exception handler dialog
@@ -1062,11 +1127,10 @@ unsigned int ServerCommandThread::CompileThread(void)
 			m_manualRecompileTriggered = false;
 			m_liveModuleToModifiedOrNewObjFiles.clear();
 		}
-		else
-		{
-			// nothing to do for now, go to sleep a bit
-			thread::Sleep(10u);
-		}
+
+		m_restartCS.Leave();
+
+		thread::Sleep(10u);
 	}
 
 	return 0u;
@@ -1078,6 +1142,7 @@ unsigned int ServerCommandThread::CommandThread(DuplexPipeServer* pipe, Event* r
 	// handle incoming commands
 	CommandMap commandMap;
 	commandMap.RegisterAction<actions::TriggerRecompile>();
+	commandMap.RegisterAction<actions::LogMessage>();
 	commandMap.RegisterAction<actions::BuildPatch>();
 	commandMap.RegisterAction<actions::ReadyForCompilation>();
 	commandMap.RegisterAction<actions::DisconnectClient>();
@@ -1204,6 +1269,16 @@ bool ServerCommandThread::actions::TriggerRecompile::Execute(const CommandType*,
 }
 
 
+bool ServerCommandThread::actions::LogMessage::Execute(const CommandType*, const DuplexPipe* pipe, void*, const void* payload, size_t)
+{
+	LC_LOG_USER("%S", static_cast<const wchar_t*>(payload));
+
+	pipe->SendAck();
+
+	return true;
+}
+
+
 bool ServerCommandThread::actions::BuildPatch::Execute(const CommandType* command, const DuplexPipe* pipe, void* context, const void* payload, size_t payloadSize)
 {
 	pipe->SendAck();
@@ -1410,7 +1485,7 @@ bool ServerCommandThread::actions::EnableLazyLoadedModule::Execute(const Command
 }
 // END EPIC MOD
 
-bool ServerCommandThread::actions::RegisterProcess::Execute(const CommandType* command, const DuplexPipe* pipe, void* context, const void*, size_t)
+bool ServerCommandThread::actions::RegisterProcess::Execute(const CommandType* command, const DuplexPipe* pipe, void* context, const void* payload, size_t)
 {
 	pipe->SendAck();
 
@@ -1452,13 +1527,63 @@ bool ServerCommandThread::actions::RegisterProcess::Execute(const CommandType* c
 
 		if (registeredSuccessfully)
 		{
-			LiveProcess* liveProcess = new LiveProcess(processHandle, command->processId, command->threadId, command->jumpToSelf, pipe);
+			const wchar_t* imagePath = pointer::Offset<const wchar_t*>(payload, 0u);
+			const wchar_t* commandLine = pointer::Offset<const wchar_t*>(imagePath, command->imagePathSize);
+			const wchar_t* workingDirectory = pointer::Offset<const wchar_t*>(commandLine, command->commandLineSize);
+			const void* environment = pointer::Offset<const wchar_t*>(workingDirectory, command->workingDirectorySize);
+
+			LiveProcess* liveProcess = new LiveProcess(processHandle, command->processId, command->threadId, command->jumpToSelf, pipe, imagePath, commandLine, workingDirectory, environment, command->environmentSize);
 			commandThread->m_liveProcesses.push_back(liveProcess);
 			// BEGIN EPIC MOD - No built-in UI
 			// commandThread->m_mainFrame->UpdateWindowTitle();
 			// END EPIC MOD
 
-			LC_SUCCESS_USER("Registered process %S (PID: %d)", processPath.c_str(), command->processId);
+			if (command->restartedProcessId == 0u)
+			{
+				// this is a new process
+				LC_SUCCESS_USER("Registered process %S (PID: %d)", processPath.c_str(), command->processId);
+			}
+			else
+			{
+				// this process was restarted
+				LC_SUCCESS_USER("Registered restarted process %S (PID: %d, previous PID: %d)", processPath.c_str(), command->processId, command->restartedProcessId);
+
+				// reattach the debugger in case the previous process had a debugger attached
+				{
+					auto it = commandThread->m_restartedProcessIdToDebugger.find(command->restartedProcessId);
+					if (it != commandThread->m_restartedProcessIdToDebugger.end())
+					{
+						LC_LOG_USER("Reattaching debugger to PID %d", command->processId);
+
+						const EnvDTE::DebuggerPtr& debugger = it->second;
+						const bool success = visualStudio::AttachToProcess(debugger, command->processId);
+						if (!success)
+						{
+							LC_ERROR_USER("Failed to reattach debugger to PID %d", command->processId);
+						}
+
+						commandThread->m_restartedProcessIdToDebugger.erase(it);
+					}
+				}
+
+				--commandThread->m_restartedProcessCount;
+				if (commandThread->m_restartedProcessCount == 0u)
+				{
+					// finished restarting, remove the job that kept this instance alive
+					if (commandThread->m_restartJob)
+					{
+						::CloseHandle(commandThread->m_restartJob);
+						commandThread->m_restartJob = nullptr;
+
+						commandThread->m_restartCS.Leave();
+
+						LC_LOG_USER("---------- Restarting finished ----------");
+
+						// EPIC REMOVED: g_theApp.GetMainFrame()->ResetStatusBarText();
+						// EPIC REMOVED: g_theApp.GetMainFrame()->SetBusy(false);
+					}
+				}
+			}
 		}
 
 		// tell client we are finished
diff --git a/Engine/Source/Developer/Windows/LiveCodingServer/Private/External/LC_ServerCommandThread.h b/Engine/Source/Developer/Windows/LiveCodingServer/Private/External/LC_ServerCommandThread.h
--- a/Engine/Source/Developer/Windows/LiveCodingServer/Private/External/LC_ServerCommandThread.h
+++ b/Engine/Source/Developer/Windows/LiveCodingServer/Private/External/LC_ServerCommandThread.h
@@ -14,6 +14,7 @@
 #include "LC_RunMode.h"
 #include "LC_Types.h"
 #include "LC_LiveModule.h"
+#include "LC_VisualStudio.h"
 
 
 class MainFrame;
@@ -27,6 +28,8 @@ public:
 	ServerCommandThread(MainFrame* mainFrame, const wchar_t* const processGroupName, RunMode::Enum runMode);
 	~ServerCommandThread(void);
 
+	void RestartTargets(void);
+
 	std::wstring GetProcessImagePath(void) const;
 
 private:
@@ -69,6 +72,7 @@ private:
 			}
 
 		DECLARE_ACTION(TriggerRecompile);
+		DECLARE_ACTION(LogMessage);
 		DECLARE_ACTION(BuildPatch);
 		DECLARE_ACTION(HandleException);
 		DECLARE_ACTION(ReadyForCompilation);
@@ -141,4 +145,10 @@ private:
 	// for triggering recompiles using the API
 	bool m_manualRecompileTriggered;
 	types::unordered_map<std::wstring, types::vector<LiveModule::ModifiedObjFile>> m_liveModuleToModifiedOrNewObjFiles;
+
+	// restart mechanism
+	CriticalSection m_restartCS;
+	void* m_restartJob;
+	unsigned int m_restartedProcessCount;
+	types::unordered_map<unsigned int, EnvDTE::DebuggerPtr> m_restartedProcessIdToDebugger;
 };
diff --git a/Engine/Source/Developer/Windows/LiveCodingServer/Private/External/LC_SymbolResolution.cpp b/Engine/Source/Developer/Windows/LiveCodingServer/Private/External/LC_SymbolResolution.cpp
--- a/Engine/Source/Developer/Windows/LiveCodingServer/Private/External/LC_SymbolResolution.cpp
+++ b/Engine/Source/Developer/Windows/LiveCodingServer/Private/External/LC_SymbolResolution.cpp
@@ -14,31 +14,38 @@
 
 namespace symbolResolution
 {
-	void Startup(void)
+	// loads symbol information for this process in the constructor,
+	// and unloads it in the destructor.
+	// makes it easier to only load information when it is needed.
+	struct SymbolHelper
 	{
-		::SymSetOptions(SYMOPT_DEFERRED_LOADS | SYMOPT_LOAD_LINES | SYMOPT_UNDNAME | SYMOPT_INCLUDE_32BIT_MODULES);
-		const BOOL success = ::SymInitialize(::GetCurrentProcess(), "", true);
-		if (success == Windows::FALSE)
+		SymbolHelper(void)
 		{
-			const DWORD error = ::GetLastError();
-			LC_ERROR_DEV("SymInitialize failed with error 0x%X", error);
+			::SymSetOptions(SYMOPT_DEFERRED_LOADS | SYMOPT_LOAD_LINES | SYMOPT_UNDNAME | SYMOPT_INCLUDE_32BIT_MODULES);
+			const BOOL success = ::SymInitialize(::GetCurrentProcess(), "", true);
+			if (success == Windows::FALSE)
+			{
+				const DWORD error = ::GetLastError();
+				LC_ERROR_DEV("SymInitialize failed with error 0x%X", error);
+			}
 		}
-	}
-
 
-	void Shutdown(void)
-	{
-		const BOOL success = ::SymCleanup(::GetCurrentProcess());
-		if (success == Windows::FALSE)
+		~SymbolHelper(void)
 		{
-			const DWORD error = ::GetLastError();
-			LC_ERROR_DEV("SymCleanup failed with error 0x%X", error);
+			const BOOL success = ::SymCleanup(::GetCurrentProcess());
+			if (success == Windows::FALSE)
+			{
+				const DWORD error = ::GetLastError();
+				LC_ERROR_DEV("SymCleanup failed with error 0x%X", error);
+			}
 		}
-	}
-
+	};
 
 	SymbolInfo ResolveSymbolsForAddress(const void* const address)
 	{
+		// prepare symbol information
+		static SymbolHelper s_symbolHelper;
+
 		// retrieve function
 		DWORD64 displacement64 = 0;
 		ULONG64 buffer[(sizeof(SYMBOL_INFO) + MAX_SYM_NAME * sizeof(TCHAR) + sizeof(ULONG64) - 1) / sizeof(ULONG64)] = {};
diff --git a/Engine/Source/Developer/Windows/LiveCodingServer/Private/External/LC_SymbolResolution.h b/Engine/Source/Developer/Windows/LiveCodingServer/Private/External/LC_SymbolResolution.h
--- a/Engine/Source/Developer/Windows/LiveCodingServer/Private/External/LC_SymbolResolution.h
+++ b/Engine/Source/Developer/Windows/LiveCodingServer/Private/External/LC_SymbolResolution.h
@@ -7,9 +7,6 @@
 
 namespace symbolResolution
 {
-	void Startup(void);
-	void Shutdown(void);
-
 	// Resolves the symbol information for a given address
 	SymbolInfo ResolveSymbolsForAddress(const void* const address);
 }
diff --git a/Engine/Source/Developer/Windows/LiveCodingServer/Private/External/LC_Symbols.cpp b/Engine/Source/Developer/Windows/LiveCodingServer/Private/External/LC_Symbols.cpp
--- a/Engine/Source/Developer/Windows/LiveCodingServer/Private/External/LC_Symbols.cpp
+++ b/Engine/Source/Developer/Windows/LiveCodingServer/Private/External/LC_Symbols.cpp
@@ -559,35 +559,6 @@ namespace symbols
 	}
 
 
-	ModuleDB* GatherModules(DiaCompilandDB* diaCompilandDb)
-	{
-		telemetry::Scope telemetryScope("Gathering modules");
-
-		const size_t count = diaCompilandDb->symbols.size();
-
-		ModuleDB* database = new ModuleDB;
-		database->modules.reserve(count);
-
-		for (size_t i = 0u; i < count; ++i)
-		{
-			IDiaSymbol* diaSymbol = diaCompilandDb->symbols[i];
-
-			const dia::SymbolName& compilandPath = dia::GetSymbolName(diaSymbol);
-			const std::wstring& uppercaseCompilandPath = string::ToUpper(compilandPath.GetString());
-
-			const bool isDllPath = string::Contains(uppercaseCompilandPath.c_str(), L".DLL");
-			const bool isImport = string::Contains(uppercaseCompilandPath.c_str(), L"IMPORT:");
-			if (isDllPath && !isImport)
-			{
-				// store the module for now
-				database->modules.emplace_back(compilandPath.GetString());
-			}
-		}
-
-		return database;
-	}
-
-
 	UserDefinedTypesDB* GatherUserDefinedTypes(const DiaCompilandDB* diaCompilandDb, const Compiland* compiland)
 	{
 		telemetry::Scope telemetryScope("Gathering user-defined types");
diff --git a/Engine/Source/Developer/Windows/LiveCodingServer/Private/External/LC_Symbols.h b/Engine/Source/Developer/Windows/LiveCodingServer/Private/External/LC_Symbols.h
--- a/Engine/Source/Developer/Windows/LiveCodingServer/Private/External/LC_Symbols.h
+++ b/Engine/Source/Developer/Windows/LiveCodingServer/Private/External/LC_Symbols.h
@@ -230,10 +230,6 @@ namespace symbols
 	DiaCompilandDB* GatherDiaCompilands(Provider* provider);
 
 
-	// gathers all modules used by the executable
-	ModuleDB* GatherModules(DiaCompilandDB* diaCompilandDb);
-
-
 	// gathers user-defined types for a given compiland
 	UserDefinedTypesDB* GatherUserDefinedTypes(const DiaCompilandDB* diaCompilandDb, const Compiland* compiland);
 
diff --git a/Engine/Source/Developer/Windows/LiveCodingServer/Private/External/LC_Telemetry.cpp b/Engine/Source/Developer/Windows/LiveCodingServer/Private/External/LC_Telemetry.cpp
--- a/Engine/Source/Developer/Windows/LiveCodingServer/Private/External/LC_Telemetry.cpp
+++ b/Engine/Source/Developer/Windows/LiveCodingServer/Private/External/LC_Telemetry.cpp
@@ -1,6 +1,7 @@
 // Copyright 2011-2019 Molecular Matters GmbH, all rights reserved.
 
 #include "LC_Telemetry.h"
+#include "LC_TimeStamp.h"
 #include "LC_Logging.h"
 #include <ratio>
 #include <inttypes.h>
@@ -10,27 +11,18 @@
 
 namespace
 {
-	template <typename T>
-	static double ReadChrono(std::chrono::high_resolution_clock::time_point start)
+	static void Print(const char* name, uint64_t start)
 	{
-		const std::chrono::high_resolution_clock::time_point now = std::chrono::high_resolution_clock::now();
-		const std::chrono::duration<double, T> timeSpan = now - start;
-
-		return timeSpan.count();
-	}
-
-
-	static void Print(const char* name, std::chrono::high_resolution_clock::time_point start)
-	{
-		const double seconds = ReadChrono<std::ratio<1, 1>>(start);
-		LC_LOG_TELEMETRY("Scope \"%s\" took %.3fs (%.3fms)", name, seconds, seconds*1000.0);
+		const uint64_t now = timeStamp::Get();
+		const uint64_t delta = now - start;
+		LC_LOG_TELEMETRY("Scope \"%s\" took %.3fs (%.3fms)", name, timeStamp::ToSeconds(delta), timeStamp::ToMilliSeconds(delta));
 	}
 }
 
 
 telemetry::Scope::Scope(const char* name)
 	: m_name(name)
-	, m_start(std::chrono::high_resolution_clock::now())
+	, m_start(timeStamp::Get())
 	, m_cs()
 {
 }
@@ -49,25 +41,28 @@ telemetry::Scope::~Scope(void)
 
 double telemetry::Scope::ReadSeconds(void) const
 {
-	return ::ReadChrono<std::ratio<1, 1>>(m_start);
+	const uint64_t now = timeStamp::Get();
+	return timeStamp::ToSeconds(now - m_start);
 }
 
 
 double telemetry::Scope::ReadMilliSeconds(void) const
 {
-	return ::ReadChrono<std::milli>(m_start);
+	const uint64_t now = timeStamp::Get();
+	return timeStamp::ToMilliSeconds(now - m_start);
 }
 
 
 double telemetry::Scope::ReadMicroSeconds(void) const
 {
-	return ::ReadChrono<std::micro>(m_start);
+	const uint64_t now = timeStamp::Get();
+	return timeStamp::ToMicroSeconds(now - m_start);
 }
 
 
 void telemetry::Scope::Restart(void)
 {
-	m_start = std::chrono::high_resolution_clock::now();
+	m_start = timeStamp::Get();
 }
 
 
diff --git a/Engine/Source/Developer/Windows/LiveCodingServer/Private/External/LC_Telemetry.h b/Engine/Source/Developer/Windows/LiveCodingServer/Private/External/LC_Telemetry.h
--- a/Engine/Source/Developer/Windows/LiveCodingServer/Private/External/LC_Telemetry.h
+++ b/Engine/Source/Developer/Windows/LiveCodingServer/Private/External/LC_Telemetry.h
@@ -23,7 +23,7 @@ namespace telemetry
 
 	private:
 		const char* m_name;
-		std::chrono::high_resolution_clock::time_point m_start;
+		uint64_t m_start;
 		CriticalSection m_cs;
 	};
 
diff --git a/Engine/Source/Developer/Windows/LiveCodingServer/Private/External/LC_TimeDate.cpp b/Engine/Source/Developer/Windows/LiveCodingServer/Private/External/LC_TimeDate.cpp
new file mode 100644
--- /dev/null
+++ b/Engine/Source/Developer/Windows/LiveCodingServer/Private/External/LC_TimeDate.cpp
@@ -0,0 +1,38 @@
+// Copyright 2011-2019 Molecular Matters GmbH, all rights reserved.
+
+#include "LC_TimeDate.h"
+#include "Windows/WindowsHWrapper.h"
+
+TimeDate::TimeDate(unsigned short year, unsigned char month, unsigned char day, unsigned char hour, unsigned char minute, unsigned char second, unsigned short milliSecond)
+	: m_year(year)
+	, m_month(month)
+	, m_day(day)
+	, m_hour(hour)
+	, m_minute(minute)
+	, m_second(second)
+	, m_milliSecond(milliSecond)
+{
+}
+
+
+TimeDate TimeDate::GetCurrent(void)
+{
+	SYSTEMTIME localTime = {};
+	::GetLocalTime(&localTime);
+
+	return TimeDate(static_cast<unsigned short>(localTime.wYear), static_cast<unsigned char>(localTime.wMonth), static_cast<unsigned char>(localTime.wDay), static_cast<unsigned char>(localTime.wHour), static_cast<unsigned char>(localTime.wMinute), static_cast<unsigned char>(localTime.wSecond), static_cast<unsigned short>(localTime.wMilliseconds));
+}
+
+
+const char* TimeDate::ToTimeString(TimeDescription& desc) const
+{
+	_snprintf_s(desc, _TRUNCATE, "%02d:%02d:%02d,%04d", m_hour, m_minute, m_second, m_milliSecond);
+	return desc;
+}
+
+
+const char* TimeDate::ToDateString(DateDescription& desc) const
+{
+	_snprintf_s(desc, _TRUNCATE, "%d-%02d-%02d", m_year, m_month, m_day);
+	return desc;
+}
diff --git a/Engine/Source/Developer/Windows/LiveCodingServer/Private/External/LC_TimeDate.h b/Engine/Source/Developer/Windows/LiveCodingServer/Private/External/LC_TimeDate.h
new file mode 100644
--- /dev/null
+++ b/Engine/Source/Developer/Windows/LiveCodingServer/Private/External/LC_TimeDate.h
@@ -0,0 +1,34 @@
+// Copyright 2011-2019 Molecular Matters GmbH, all rights reserved.
+
+#pragma once
+
+
+class TimeDate
+{
+public:
+	// A type representing the time format: extended ISO (ISO 8601) hh:mm:ss,mmmm plus a null terminator.
+	typedef char TimeDescription[14];
+
+	// A type representing the date format: ISO (ISO 8601) YYYY-MM-DD plus a null terminator.
+	typedef char DateDescription[11];
+
+	TimeDate(unsigned short year, unsigned char month, unsigned char day, unsigned char hour, unsigned char minute, unsigned char second, unsigned short milliSecond);
+
+	// Returns the current time and date in local time.
+	static TimeDate GetCurrent(void);
+
+	// Converts the time into the extended ISO 8601 format, and returns a pointer to the description string.
+	const char* ToTimeString(TimeDescription& desc) const;
+
+	// Converts the date into the standard ISO 8601 format, and returns a pointer to the description string.
+	const char* ToDateString(DateDescription& desc) const;
+
+private:
+	unsigned short m_year;
+	unsigned char m_month;
+	unsigned char m_day;
+	unsigned char m_hour;
+	unsigned char m_minute;
+	unsigned char m_second;
+	unsigned short m_milliSecond;
+};
diff --git a/Engine/Source/Developer/Windows/LiveCodingServer/Private/External/LC_TimeStamp.cpp b/Engine/Source/Developer/Windows/LiveCodingServer/Private/External/LC_TimeStamp.cpp
new file mode 100644
--- /dev/null
+++ b/Engine/Source/Developer/Windows/LiveCodingServer/Private/External/LC_TimeStamp.cpp
@@ -0,0 +1,56 @@
+// Copyright 2011-2019 Molecular Matters GmbH, all rights reserved.
+
+#include "LC_TimeStamp.h"
+#include "Windows/WindowsHWrapper.h"
+
+
+namespace timeStamp
+{
+	double g_oneOverFrequency = 0.0;
+	double g_thousandOverFrequency = 0.0;
+	double g_millionOverFrequency = 0.0;
+
+
+	void Startup(void)
+	{
+		::LARGE_INTEGER frequency = {};
+		::QueryPerformanceFrequency(&frequency);
+		g_oneOverFrequency = 1.0 / frequency.QuadPart;
+		g_thousandOverFrequency = 1000.0 / frequency.QuadPart;
+		g_millionOverFrequency = 1000000.0 / frequency.QuadPart;
+	}
+
+
+	void Shutdown(void)
+	{
+		g_oneOverFrequency = 0.0;
+		g_thousandOverFrequency = 0.0;
+		g_millionOverFrequency = 0.0;
+	}
+
+
+	uint64_t Get(void)
+	{
+		::LARGE_INTEGER now = {};
+		::QueryPerformanceCounter(&now);
+		return static_cast<uint64_t>(now.QuadPart);
+	}
+
+
+	double ToSeconds(uint64_t count)
+	{
+		return count * g_oneOverFrequency;
+	}
+
+
+	double ToMilliSeconds(uint64_t count)
+	{
+		return count * g_thousandOverFrequency;
+	}
+
+
+	double ToMicroSeconds(uint64_t count)
+	{
+		return count * g_millionOverFrequency;
+	}
+}
diff --git a/Engine/Source/Developer/Windows/LiveCodingServer/Private/External/LC_TimeStamp.h b/Engine/Source/Developer/Windows/LiveCodingServer/Private/External/LC_TimeStamp.h
new file mode 100644
--- /dev/null
+++ b/Engine/Source/Developer/Windows/LiveCodingServer/Private/External/LC_TimeStamp.h
@@ -0,0 +1,24 @@
+// Copyright 2011-2019 Molecular Matters GmbH, all rights reserved.
+
+#pragma once
+
+
+namespace timeStamp
+{
+	void Startup(void);
+	void Shutdown(void);
+
+
+	// Returns the current value of the high-resolution performance counter.
+	// This is not the same as actual CPU cycles.
+	uint64_t Get(void);
+
+	// Converts counts returned by Get() into seconds.
+	double ToSeconds(uint64_t count);
+
+	// Converts counts returned by Get() into milliseconds.
+	double ToMilliSeconds(uint64_t count);
+
+	// Converts counts returned by Get() into microseconds.
+	double ToMicroSeconds(uint64_t count);
+}
diff --git a/Engine/Source/Developer/Windows/LiveCodingServer/Private/External/LC_VisualStudio.cpp b/Engine/Source/Developer/Windows/LiveCodingServer/Private/External/LC_VisualStudio.cpp
new file mode 100644
--- /dev/null
+++ b/Engine/Source/Developer/Windows/LiveCodingServer/Private/External/LC_VisualStudio.cpp
@@ -0,0 +1,4 @@
+// Copyright 2011-2019 Molecular Matters GmbH, all rights reserved.
+
+#include "LC_VisualStudio.h"
+
diff --git a/Engine/Source/Developer/Windows/LiveCodingServer/Private/External/LC_VisualStudio.h b/Engine/Source/Developer/Windows/LiveCodingServer/Private/External/LC_VisualStudio.h
new file mode 100644
--- /dev/null
+++ b/Engine/Source/Developer/Windows/LiveCodingServer/Private/External/LC_VisualStudio.h
@@ -0,0 +1,34 @@
+// Copyright 2011-2019 Molecular Matters GmbH, all rights reserved.
+
+#pragma once
+
+#include "LC_Types.h"
+#include "Windows/WindowsHWrapper.h"
+
+#define LC_WITH_VISUAL_STUDIO_DTE 1
+
+#include "Windows/AllowWindowsPlatformTypes.h"
+#include "Windows/AllowWindowsPlatformAtomics.h"
+#include <unknwn.h>
+
+#if LC_WITH_VISUAL_STUDIO_DTE
+	#pragma warning(push)
+	#pragma warning(disable: 4278)
+	#pragma warning(disable: 4471)
+	#pragma warning(disable: 4146)
+	#pragma warning(disable: 4191)
+	#pragma warning(disable: 6244)
+
+	// import EnvDTE
+	#define VSACCESSOR_HAS_DTE_OLB 0
+	#if VSACCESSOR_HAS_DTE_OLB
+		#import "NotForLicensees/dte80a.olb" version("8.0") lcid("0") raw_interfaces_only named_guids
+	#else
+		#import "libid:80cc9f66-e7d8-4ddd-85b6-d9e6cd0e93e2" version("8.0") lcid("0") raw_interfaces_only named_guids
+	#endif
+
+	#pragma warning(pop)
+#endif
+
+#include "Windows/HideWindowsPlatformAtomics.h"
+#include "Windows/HideWindowsPlatformTypes.h"
diff --git a/Engine/Source/Developer/Windows/LiveCodingServer/Private/External/LC_VisualStudioAutomation.cpp b/Engine/Source/Developer/Windows/LiveCodingServer/Private/External/LC_VisualStudioAutomation.cpp
new file mode 100644
--- /dev/null
+++ b/Engine/Source/Developer/Windows/LiveCodingServer/Private/External/LC_VisualStudioAutomation.cpp
@@ -0,0 +1,558 @@
+// Copyright 2011-2019 Molecular Matters GmbH, all rights reserved.
+
+#include "LC_VisualStudioAutomation.h"
+#include "LC_COMThread.h"
+#include "LC_StringUtil.h"
+#include "LC_Logging.h"
+#include "Windows/COMPointer.h"
+
+namespace
+{
+	static COMThread* g_comThread = nullptr;
+
+
+	// helper function that waits until the debugger finished executing a command and is in break mode again
+	static void WaitUntilBreakMode(const EnvDTE::DebuggerPtr& debugger)
+	{
+		EnvDTE::dbgDebugMode mode = EnvDTE::dbgDesignMode;
+
+		while (mode != EnvDTE::dbgBreakMode)
+		{
+			debugger->get_CurrentMode(&mode);
+			thread::Sleep(5u);
+		}
+	}
+}
+
+
+void visualStudio::Startup(void)
+{
+	// start the COM thread that does the actual work
+	g_comThread = new COMThread;
+}
+
+
+void visualStudio::Shutdown(void)
+{
+	// shut down the COM thread
+	memory::DeleteAndNull(g_comThread);
+}
+
+
+static EnvDTE::DebuggerPtr FindDebuggerAttachedToProcess(unsigned int processId)
+{
+	// get all objects from the running object table (ROT)
+	TComPtr<IRunningObjectTable> rot = nullptr;
+	HRESULT result = ::GetRunningObjectTable(0u, &rot);
+	if (FAILED(result) || (!rot))
+	{
+		LC_ERROR_DEV("Could not initialize running object table. Error: 0x%X", result);
+		return nullptr;
+	}
+
+	TComPtr<IEnumMoniker> enumMoniker = nullptr;
+	result = rot->EnumRunning(&enumMoniker);
+	if (FAILED(result) || (!enumMoniker))
+	{
+		LC_ERROR_DEV("Could not enumerate running objects. Error: 0x%X", result);
+		return nullptr;
+	}
+
+	enumMoniker->Reset();
+	do
+	{
+		TComPtr<IMoniker> next = nullptr;
+		result = enumMoniker->Next(1u, &next, NULL);
+
+		if (SUCCEEDED(result) && next)
+		{
+			TComPtr<IBindCtx> context = nullptr;
+			result = ::CreateBindCtx(0, &context);
+
+			if (FAILED(result) || (!context))
+			{
+				LC_ERROR_DEV("Could not create COM binding context. Error: 0x%X", result);
+				continue;
+			}
+
+			wchar_t* displayName = nullptr;
+			result = next->GetDisplayName(context, NULL, &displayName);
+			if (FAILED(result) || (!displayName))
+			{
+				LC_ERROR_DEV("Could not retrieve display name. Error: 0x%X", result);
+				continue;
+			}
+
+			// only try objects which are a specific version of Visual Studio
+			if (string::Contains(displayName, L"VisualStudio.DTE."))
+			{
+				// free display name using COM allocator
+				TComPtr<IMalloc> comMalloc = nullptr;
+				result = ::CoGetMalloc(1u, &comMalloc);
+				if (SUCCEEDED(result) && comMalloc)
+				{
+					comMalloc->Free(displayName);
+				}
+
+				TComPtr<IUnknown> unknown = nullptr;
+				result = rot->GetObject(next, &unknown);
+				if (FAILED(result) || (!unknown))
+				{
+					LC_ERROR_DEV("Could not retrieve COM object from running object table. Error: 0x%X", result);
+					continue;
+				}
+
+				EnvDTE::_DTEPtr dte = nullptr;
+				result = unknown->QueryInterface(&dte);
+				if (FAILED(result) || (!dte))
+				{
+					// this COM object doesn't support the DTE interface
+					LC_ERROR_DEV("Could not convert IUnknown to DTE interface. Error: 0x%X", result);
+					continue;
+				}
+
+				EnvDTE::DebuggerPtr debugger = nullptr;
+				result = dte->get_Debugger(&debugger);
+				if (FAILED(result) || (!debugger))
+				{
+					// cannot access debugger, which means that the process is currently not being debugged
+					LC_LOG_DEV("Could not access debugger interface. Error: 0x%X", result);
+					continue;
+				}
+
+				// fetch all processes to which this debugger is attached
+				EnvDTE::ProcessesPtr allProcesses = nullptr;
+				result = debugger->get_DebuggedProcesses(&allProcesses);
+				if (FAILED(result) || (!allProcesses))
+				{
+					LC_ERROR_DEV("Could not retrieve processes from debugger. Error: 0x%X", result);
+					continue;
+				}
+
+				long processCount = 0;
+				result = allProcesses->get_Count(&processCount);
+				if (FAILED(result) || (processCount <= 0))
+				{
+					LC_ERROR_DEV("Could not retrieve process count from debugger. Error: 0x%X", result);
+					continue;
+				}
+
+				// check all processes if any of them is the one we're looking for
+				for (long i = 0u; i < processCount; ++i)
+				{
+					EnvDTE::ProcessPtr singleProcess = nullptr;
+					result = allProcesses->Item(variant_t(i + 1), &singleProcess);
+
+					if (FAILED(result) || (!singleProcess))
+					{
+						LC_ERROR_DEV("Could not retrieve process from debugger. Error: 0x%X", result);
+						continue;
+					}
+
+					long debuggerProcessId = 0;
+					result = singleProcess->get_ProcessID(&debuggerProcessId);
+					if (FAILED(result) || (debuggerProcessId <= 0))
+					{
+						LC_ERROR_DEV("Could not retrieve process ID from debugger. Error: 0x%X", result);
+						continue;
+					}
+
+					// we got a valid processId
+					if (static_cast<unsigned int>(debuggerProcessId) == processId)
+					{
+						// found debugger attached to our process
+						return debugger;
+					}
+				}
+			}
+		}
+	}
+	while (result != S_FALSE);
+
+	return nullptr;
+}
+
+
+EnvDTE::DebuggerPtr visualStudio::FindDebuggerAttachedToProcess(unsigned int processId)
+{
+	return g_comThread->CallInThread(&::FindDebuggerAttachedToProcess, processId);
+}
+
+
+static EnvDTE::DebuggerPtr FindDebuggerForProcess(unsigned int processId)
+{
+	// get all objects from the running object table (ROT)
+	TComPtr<IRunningObjectTable> rot = nullptr;
+	HRESULT result = ::GetRunningObjectTable(0u, &rot);
+	if (FAILED(result) || (!rot))
+	{
+		LC_ERROR_DEV("Could not initialize running object table. Error: 0x%X", result);
+		return nullptr;
+	}
+
+	TComPtr<IEnumMoniker> enumMoniker = nullptr;
+	result = rot->EnumRunning(&enumMoniker);
+	if (FAILED(result) || (!enumMoniker))
+	{
+		LC_ERROR_DEV("Could not enumerate running objects. Error: 0x%X", result);
+		return nullptr;
+	}
+
+	enumMoniker->Reset();
+	do
+	{
+		TComPtr<IMoniker> next = nullptr;
+		result = enumMoniker->Next(1u, &next, NULL);
+
+		if (SUCCEEDED(result) && next)
+		{
+			TComPtr<IBindCtx> context = nullptr;
+			result = ::CreateBindCtx(0, &context);
+
+			if (FAILED(result) || (!context))
+			{
+				LC_ERROR_DEV("Could not create COM binding context. Error: 0x%X", result);
+				continue;
+			}
+
+			wchar_t* displayName = nullptr;
+			result = next->GetDisplayName(context, NULL, &displayName);
+			if (FAILED(result) || (!displayName))
+			{
+				LC_ERROR_DEV("Could not retrieve display name. Error: 0x%X", result);
+				continue;
+			}
+
+			// only try objects which are a specific version of Visual Studio
+			if (string::Contains(displayName, L"VisualStudio.DTE."))
+			{
+				// free display name using COM allocator
+				TComPtr<IMalloc> comMalloc = nullptr;
+				result = ::CoGetMalloc(1u, &comMalloc);
+				if (SUCCEEDED(result) && comMalloc)
+				{
+					comMalloc->Free(displayName);
+				}
+
+				TComPtr<IUnknown> unknown = nullptr;
+				result = rot->GetObject(next, &unknown);
+				if (FAILED(result) || (!unknown))
+				{
+					LC_ERROR_DEV("Could not retrieve COM object from running object table. Error: 0x%X", result);
+					continue;
+				}
+
+				EnvDTE::_DTEPtr dte = nullptr;
+				result = unknown->QueryInterface(&dte);
+				if (FAILED(result) || (!dte))
+				{
+					// this COM object doesn't support the DTE interface
+					LC_ERROR_DEV("Could not convert IUnknown to DTE interface. Error: 0x%X", result);
+					continue;
+				}
+
+				EnvDTE::DebuggerPtr debugger = nullptr;
+				result = dte->get_Debugger(&debugger);
+				if (FAILED(result) || (!debugger))
+				{
+					// cannot access debugger, which means that the process is currently not being debugged
+					LC_LOG_DEV("Could not access debugger interface. Error: 0x%X", result);
+					continue;
+				}
+
+				EnvDTE::ProcessPtr process = nullptr;
+				result = debugger->get_CurrentProcess(&process);
+				if (FAILED(result) || (!process))
+				{
+					// cannot access current process, reason unknown
+					LC_ERROR_DEV("Could not access current process in debugger. Error: 0x%X", result);
+					continue;
+				}
+
+				long debuggerProcessId = 0;
+				result = process->get_ProcessID(&debuggerProcessId);
+				if (FAILED(result) || (debuggerProcessId <= 0))
+				{
+					LC_ERROR_DEV("Could not retrieve process ID from debugger. Error: 0x%X", result);
+					continue;
+				}
+
+				// we got a valid processId
+				if (static_cast<unsigned int>(debuggerProcessId) == processId)
+				{
+					// found debugger debugging our process
+					return debugger;
+				}
+			}
+		}
+	}
+	while (result != S_FALSE);
+
+	return nullptr;
+}
+
+
+EnvDTE::DebuggerPtr visualStudio::FindDebuggerForProcess(unsigned int processId)
+{
+	return g_comThread->CallInThread(&::FindDebuggerForProcess, processId);
+}
+
+
+static bool AttachToProcess(const EnvDTE::DebuggerPtr& debugger, unsigned int processId)
+{
+	// fetch all local processes running on this machine
+	EnvDTE::ProcessesPtr allProcesses = nullptr;
+	HRESULT result = debugger->get_LocalProcesses(&allProcesses);
+	if (FAILED(result) || (!allProcesses))
+	{
+		LC_ERROR_DEV("Could not retrieve local processes from debugger. Error: 0x%X", result);
+		return false;
+	}
+
+	long processCount = 0;
+	result = allProcesses->get_Count(&processCount);
+	if (FAILED(result) || (processCount <= 0))
+	{
+		LC_ERROR_DEV("Could not retrieve local process count from debugger. Error: 0x%X", result);
+		return false;
+	}
+
+	// check all processes if any of them is the one we're looking for
+	for (long i = 0u; i < processCount; ++i)
+	{
+		EnvDTE::ProcessPtr singleProcess = nullptr;
+		result = allProcesses->Item(variant_t(i + 1), &singleProcess);
+
+		if (FAILED(result) || (!singleProcess))
+		{
+			LC_ERROR_DEV("Could not retrieve local process from debugger. Error: 0x%X", result);
+			continue;
+		}
+
+		long localProcessId = 0;
+		result = singleProcess->get_ProcessID(&localProcessId);
+		if (FAILED(result) || (localProcessId <= 0))
+		{
+			LC_ERROR_DEV("Could not retrieve local process ID from debugger. Error: 0x%X", result);
+			continue;
+		}
+
+		// we got a valid processId
+		if (static_cast<unsigned int>(localProcessId) == processId)
+		{
+			// this is the process we want to attach to
+			result = singleProcess->Attach();
+			if (FAILED(result))
+			{
+				LC_ERROR_USER("Could not attach debugger to process. Error: 0x%X", result);
+				return false;
+			}
+
+			return true;
+		}
+	}
+
+	return false;
+}
+
+
+bool visualStudio::AttachToProcess(const EnvDTE::DebuggerPtr& debugger, unsigned int processId)
+{
+	return g_comThread->CallInThread(&::AttachToProcess, debugger, processId);
+}
+
+
+static types::vector<EnvDTE::ThreadPtr> EnumerateThreads(const EnvDTE::DebuggerPtr& debugger)
+{
+	types::vector<EnvDTE::ThreadPtr> threads;
+
+	EnvDTE::ProgramPtr program = nullptr;
+	HRESULT result = debugger->get_CurrentProgram(&program);
+	if (FAILED(result) || (!program))
+	{
+		LC_ERROR_DEV("Could not retrieve current program from debugger. Error: 0x%X", result);
+		return threads;
+	}
+
+	EnvDTE::ThreadsPtr allThreads = nullptr;
+	result = program->get_Threads(&allThreads);
+	if (FAILED(result) || (!allThreads))
+	{
+		LC_ERROR_DEV("Could not retrieve running threads from debugger. Error: 0x%X", result);
+		return threads;
+	}
+
+	long threadCount = 0;
+	result = allThreads->get_Count(&threadCount);
+	if (FAILED(result) || (threadCount <= 0))
+	{
+		LC_ERROR_DEV("Could not retrieve thread count from debugger. Error: 0x%X", result);
+		return threads;
+	}
+
+	threads.reserve(static_cast<size_t>(threadCount));
+
+	for (long i = 0u; i < threadCount; ++i)
+	{
+		EnvDTE::ThreadPtr singleThread = nullptr;
+		result = allThreads->Item(variant_t(i + 1), &singleThread);
+		if (FAILED(result) || (!singleThread))
+		{
+			LC_ERROR_DEV("Could not retrieve thread from debugger. Error: 0x%X", result);
+			continue;
+		}
+
+		threads.push_back(singleThread);
+	}
+
+	return threads;
+}
+
+
+types::vector<EnvDTE::ThreadPtr> visualStudio::EnumerateThreads(const EnvDTE::DebuggerPtr& debugger)
+{
+	return g_comThread->CallInThread(&::EnumerateThreads, debugger);
+}
+
+
+static bool FreezeThreads(const EnvDTE::DebuggerPtr& debugger, const types::vector<EnvDTE::ThreadPtr>& threads)
+{
+	bool success = true;
+
+	const size_t count = threads.size();
+	for (size_t i = 0u; i < count; ++i)
+	{
+		EnvDTE::ThreadPtr singleThread = threads[i];
+		const HRESULT result = singleThread->Freeze();
+		WaitUntilBreakMode(debugger);
+
+		success &= SUCCEEDED(result);
+	}
+
+	return success;
+}
+
+
+bool visualStudio::FreezeThreads(const EnvDTE::DebuggerPtr& debugger, const types::vector<EnvDTE::ThreadPtr>& threads)
+{
+	return g_comThread->CallInThread(&::FreezeThreads, debugger, threads);
+}
+
+
+static bool FreezeThread(const EnvDTE::DebuggerPtr& debugger, const types::vector<EnvDTE::ThreadPtr>& threads, unsigned int threadId)
+{
+	const size_t count = threads.size();
+	for (size_t i = 0u; i < count; ++i)
+	{
+		EnvDTE::ThreadPtr singleThread = threads[i];
+
+		long id = 0;
+		HRESULT result = singleThread->get_ID(&id);
+		if (FAILED(result) || (id <= 0))
+		{
+			continue;
+		}
+
+		if (static_cast<unsigned int>(id) == threadId)
+		{
+			// found the thread we're looking for
+			result = singleThread->Freeze();
+			WaitUntilBreakMode(debugger);
+
+			return SUCCEEDED(result);
+		}
+	}
+
+	return false;
+}
+
+
+bool visualStudio::FreezeThread(const EnvDTE::DebuggerPtr& debugger, const types::vector<EnvDTE::ThreadPtr>& threads, unsigned int threadId)
+{
+	return g_comThread->CallInThread(&::FreezeThread, debugger, threads, threadId);
+}
+
+
+static bool ThawThreads(const EnvDTE::DebuggerPtr& debugger, const types::vector<EnvDTE::ThreadPtr>& threads)
+{
+	bool success = true;
+
+	const size_t count = threads.size();
+	for (size_t i = 0u; i < count; ++i)
+	{
+		EnvDTE::ThreadPtr singleThread = threads[i];
+		const HRESULT result = singleThread->Thaw();
+		WaitUntilBreakMode(debugger);
+
+		success &= SUCCEEDED(result);
+	}
+
+	return success;
+}
+
+
+bool visualStudio::ThawThreads(const EnvDTE::DebuggerPtr& debugger, const types::vector<EnvDTE::ThreadPtr>& threads)
+{
+	return g_comThread->CallInThread(&::ThawThreads, debugger, threads);
+}
+
+
+static bool ThawThread(const EnvDTE::DebuggerPtr& debugger, const types::vector<EnvDTE::ThreadPtr>& threads, unsigned int threadId)
+{
+	const size_t count = threads.size();
+	for (size_t i = 0u; i < count; ++i)
+	{
+		EnvDTE::ThreadPtr singleThread = threads[i];
+
+		long id = 0;
+		HRESULT result = singleThread->get_ID(&id);
+		if (FAILED(result) || (id <= 0))
+		{
+			continue;
+		}
+
+		if (static_cast<unsigned int>(id) == threadId)
+		{
+			// found the thread we're looking for
+			result = singleThread->Thaw();
+			WaitUntilBreakMode(debugger);
+
+			return SUCCEEDED(result);
+		}
+	}
+
+	return false;
+}
+
+
+bool visualStudio::ThawThread(const EnvDTE::DebuggerPtr& debugger, const types::vector<EnvDTE::ThreadPtr>& threads, unsigned int threadId)
+{
+	return g_comThread->CallInThread(&::ThawThread, debugger, threads, threadId);
+}
+
+
+static bool Resume(const EnvDTE::DebuggerPtr& debugger)
+{
+	const HRESULT result = debugger->Go(variant_t(Windows::FALSE));
+	return SUCCEEDED(result);
+}
+
+
+bool visualStudio::Resume(const EnvDTE::DebuggerPtr& debugger)
+{
+	return g_comThread->CallInThread(&::Resume, debugger);
+}
+
+
+static bool Break(const EnvDTE::DebuggerPtr& debugger)
+{
+	// wait until the debugger really enters break-mode
+	const HRESULT result = debugger->Break(variant_t(Windows::TRUE));
+	return SUCCEEDED(result);
+}
+
+
+bool visualStudio::Break(const EnvDTE::DebuggerPtr& debugger)
+{
+	return g_comThread->CallInThread(&::Break, debugger);
+}
diff --git a/Engine/Source/Developer/Windows/LiveCodingServer/Private/External/LC_VisualStudioAutomation.h b/Engine/Source/Developer/Windows/LiveCodingServer/Private/External/LC_VisualStudioAutomation.h
new file mode 100644
--- /dev/null
+++ b/Engine/Source/Developer/Windows/LiveCodingServer/Private/External/LC_VisualStudioAutomation.h
@@ -0,0 +1,43 @@
+// Copyright 2011-2019 Molecular Matters GmbH, all rights reserved.
+
+#pragma once
+
+#include "LC_Types.h"
+#include "LC_VisualStudio.h"
+
+namespace visualStudio
+{
+	void Startup(void);
+	void Shutdown(void);
+
+
+	// Finds a Visual Studio debugger instance currently attached to the process with the given ID
+	EnvDTE::DebuggerPtr FindDebuggerAttachedToProcess(unsigned int processId);
+
+	// Finds a Visual Studio debugger instance currently debugging the process with the given ID
+	EnvDTE::DebuggerPtr FindDebuggerForProcess(unsigned int processId);
+
+	// Attaches a Visual Studio debugger instance to the process with the given ID
+	bool AttachToProcess(const EnvDTE::DebuggerPtr& debugger, unsigned int processId);
+
+	// Enumerates all threads of a debugger instance
+	types::vector<EnvDTE::ThreadPtr> EnumerateThreads(const EnvDTE::DebuggerPtr& debugger);
+
+	// Freezes all given threads
+	bool FreezeThreads(const EnvDTE::DebuggerPtr& debugger, const types::vector<EnvDTE::ThreadPtr>& threads);
+
+	// Freezes a single thread with the given thread ID
+	bool FreezeThread(const EnvDTE::DebuggerPtr& debugger, const types::vector<EnvDTE::ThreadPtr>& threads, unsigned int threadId);
+
+	// Thaws all given threads
+	bool ThawThreads(const EnvDTE::DebuggerPtr& debugger, const types::vector<EnvDTE::ThreadPtr>& threads);
+
+	// Thaws a single thread with the given thread ID
+	bool ThawThread(const EnvDTE::DebuggerPtr& debugger, const types::vector<EnvDTE::ThreadPtr>& threads, unsigned int threadId);
+
+	// Resumes the process in the debugger
+	bool Resume(const EnvDTE::DebuggerPtr& debugger);
+
+	// Breaks the process in the debugger
+	bool Break(const EnvDTE::DebuggerPtr& debugger);
+}
diff --git a/Engine/Source/Developer/Windows/LiveCodingServer/Private/LiveCodingServer.cpp b/Engine/Source/Developer/Windows/LiveCodingServer/Private/LiveCodingServer.cpp
--- a/Engine/Source/Developer/Windows/LiveCodingServer/Private/LiveCodingServer.cpp
+++ b/Engine/Source/Developer/Windows/LiveCodingServer/Private/LiveCodingServer.cpp
@@ -56,8 +56,8 @@ void FLiveCodingServer::SetLinkerPath(const wchar_t* LinkerPath, const TMap<FStr
 
 	if (LinkerEnvironment.Num() > 0)
 	{
-		environment::Block* block = environment::CreateBlockFromMap(LinkerEnvironment);
-		compiler::AddEnvironmentToCache(LinkerPath, block);
+		process::Environment* environment = process::CreateEnvironmentFromMap(LinkerEnvironment);
+		compiler::AddEnvironmentToCache(LinkerPath, environment);
 	}
 }
 
