Fix unity build error.

#rnx

diff --git a/Engine/Source/Developer/Windows/LiveCoding/Private/External/LC_EntryPoint.cpp b/Engine/Source/Developer/Windows/LiveCoding/Private/External/LC_EntryPoint.cpp
--- a/Engine/Source/Developer/Windows/LiveCoding/Private/External/LC_EntryPoint.cpp
+++ b/Engine/Source/Developer/Windows/LiveCoding/Private/External/LC_EntryPoint.cpp
@@ -8,15 +8,15 @@
 namespace
 {
 	// startup thread
-	static ClientStartupThread* g_startupThread = nullptr;
+	static ClientStartupThread* g_mainStartupThread = nullptr;
 }
 
 
 // BEGIN EPIC MOD - Manually trigger startup/shutdown code
 void Startup(HINSTANCE instance)
 {
-	g_startupThread = new ClientStartupThread(instance);
-	api::Startup(g_startupThread);
+	g_mainStartupThread = new ClientStartupThread(instance);
+	api::Startup(g_mainStartupThread);
 }
 
 
@@ -25,8 +25,8 @@ void Shutdown(void)
 	api::Shutdown();
 
 	// wait for the startup thread to finish its work and clean up
-	g_startupThread->Join();
-	delete g_startupThread;
+	g_mainStartupThread->Join();
+	delete g_mainStartupThread;
 }
 
 #if 0
