LiveCoding: Allow building without Visual Studio DTE being present. We don't support the restart feature yet.

diff --git a/Engine/Source/Developer/Windows/LiveCodingServer/Private/External/LC_LiveProcess.cpp b/Engine/Source/Developer/Windows/LiveCodingServer/Private/External/LC_LiveProcess.cpp
--- a/Engine/Source/Developer/Windows/LiveCodingServer/Private/External/LC_LiveProcess.cpp
+++ b/Engine/Source/Developer/Windows/LiveCodingServer/Private/External/LC_LiveProcess.cpp
@@ -6,6 +6,7 @@
 #include "LC_Event.h"
 #include "LC_PrimitiveNames.h"
 #include "LC_VisualStudioAutomation.h"
+#include "LC_Logging.h"
 
 
 LiveProcess::LiveProcess(process::Handle processHandle, unsigned int processId, unsigned int commandThreadId, const void* jumpToSelf, const DuplexPipe* pipe,
@@ -21,8 +22,10 @@ LiveProcess::LiveProcess(process::Handle processHandle, unsigned int processId,
 	, m_environment(environment, environmentSize)
 	, m_imagesTriedToLoad()
 	, m_heartBeatDelta(0ull)
+#if LC_WITH_VISUAL_STUDIO_DTE
 	, m_vsDebugger(nullptr)
 	, m_vsDebuggerThreads()
+#endif
 	, m_codeCave(nullptr)
 	, m_restartState(RestartState::DEFAULT)
 {
@@ -53,6 +56,7 @@ bool LiveProcess::MadeProgress(void) const
 
 void LiveProcess::HandleDebuggingPreCompile(void)
 {
+#if LC_WITH_VISUAL_STUDIO_DTE
 	if (!MadeProgress())
 	{
 		// this process did not make progress.
@@ -81,6 +85,7 @@ void LiveProcess::HandleDebuggingPreCompile(void)
 		LC_LOG_USER("Failed to automate debugger attached to process (PID: %d), using fallback mechanism", m_processId);
 		LC_SUCCESS_USER("Waiting for client process (PID: %d), hit 'Continue' (F5 in Visual Studio) if being held in the debugger", m_processId);
 	}
+#endif
 
 	// this process either made progress and is not held in the debugger, or we failed automating the debugger.
 	// "halt" this process by installing a code cave.
@@ -95,6 +100,7 @@ void LiveProcess::HandleDebuggingPostCompile(void)
 		// we installed a code cave previously, remove it
 		UninstallCodeCave();
 	}
+#if LC_WITH_VISUAL_STUDIO_DTE
 	else if (m_vsDebugger)
 	{
 		// we automated the debugger previously. break into the debugger again and resume all threads.
@@ -108,6 +114,7 @@ void LiveProcess::HandleDebuggingPostCompile(void)
 	}
 
 	m_vsDebugger = nullptr;
+#endif
 }
 
 
diff --git a/Engine/Source/Developer/Windows/LiveCodingServer/Private/External/LC_ServerCommandThread.cpp b/Engine/Source/Developer/Windows/LiveCodingServer/Private/External/LC_ServerCommandThread.cpp
--- a/Engine/Source/Developer/Windows/LiveCodingServer/Private/External/LC_ServerCommandThread.cpp
+++ b/Engine/Source/Developer/Windows/LiveCodingServer/Private/External/LC_ServerCommandThread.cpp
@@ -100,9 +100,13 @@ ServerCommandThread::ServerCommandThread(MainFrame* mainFrame, const wchar_t* co
 	, m_restartCS()
 	, m_restartJob(nullptr)
 	, m_restartedProcessCount(0u)
+#if LC_WITH_VISUAL_STUDIO_DTE
 	, m_restartedProcessIdToDebugger()
+#endif
 {
+#if LC_WITH_VISUAL_STUDIO_DTE
 	visualStudio::Startup();
+#endif
 
 	m_serverThread = thread::Create("Live coding server", 64u * 1024u, &ServerCommandThread::ServerThread, this);
 	m_compileThread = thread::Create("Live coding compilation", 64u * 1024u, &ServerCommandThread::CompileThread, this);
@@ -123,7 +127,9 @@ ServerCommandThread::~ServerCommandThread(void)
 	// some intensive work.
 	delete m_directoryCache;
 
+#if LC_WITH_VISUAL_STUDIO_DTE
 	visualStudio::Shutdown();
+#endif
 }
 
 
@@ -184,6 +190,7 @@ void ServerCommandThread::RestartTargets(void)
 		{
 			++m_restartedProcessCount;
 
+#if LC_WITH_VISUAL_STUDIO_DTE
 			// check if a VS debugger is currently attached to the process about to restart
 			const unsigned int processId = liveProcess->GetProcessId();
 			EnvDTE::DebuggerPtr debugger = visualStudio::FindDebuggerAttachedToProcess(processId);
@@ -191,6 +198,7 @@ void ServerCommandThread::RestartTargets(void)
 			{
 				m_restartedProcessIdToDebugger.emplace(processId, debugger);
 			}
+#endif
 		}
 	}
 
@@ -1548,6 +1556,7 @@ bool ServerCommandThread::actions::RegisterProcess::Execute(const CommandType* c
 				// this process was restarted
 				LC_SUCCESS_USER("Registered restarted process %S (PID: %d, previous PID: %d)", processPath.c_str(), command->processId, command->restartedProcessId);
 
+#if LC_WITH_VISUAL_STUDIO_DTE
 				// reattach the debugger in case the previous process had a debugger attached
 				{
 					auto it = commandThread->m_restartedProcessIdToDebugger.find(command->restartedProcessId);
@@ -1565,6 +1574,7 @@ bool ServerCommandThread::actions::RegisterProcess::Execute(const CommandType* c
 						commandThread->m_restartedProcessIdToDebugger.erase(it);
 					}
 				}
+#endif
 
 				--commandThread->m_restartedProcessCount;
 				if (commandThread->m_restartedProcessCount == 0u)
diff --git a/Engine/Source/Developer/Windows/LiveCodingServer/Private/External/LC_ServerCommandThread.h b/Engine/Source/Developer/Windows/LiveCodingServer/Private/External/LC_ServerCommandThread.h
--- a/Engine/Source/Developer/Windows/LiveCodingServer/Private/External/LC_ServerCommandThread.h
+++ b/Engine/Source/Developer/Windows/LiveCodingServer/Private/External/LC_ServerCommandThread.h
@@ -150,5 +150,7 @@ private:
 	CriticalSection m_restartCS;
 	void* m_restartJob;
 	unsigned int m_restartedProcessCount;
+#if LC_WITH_VISUAL_STUDIO_DTE
 	types::unordered_map<unsigned int, EnvDTE::DebuggerPtr> m_restartedProcessIdToDebugger;
+#endif
 };
diff --git a/Engine/Source/Developer/Windows/LiveCodingServer/Private/External/LC_VisualStudio.h b/Engine/Source/Developer/Windows/LiveCodingServer/Private/External/LC_VisualStudio.h
--- a/Engine/Source/Developer/Windows/LiveCodingServer/Private/External/LC_VisualStudio.h
+++ b/Engine/Source/Developer/Windows/LiveCodingServer/Private/External/LC_VisualStudio.h
@@ -5,7 +5,7 @@
 #include "LC_Types.h"
 #include "Windows/WindowsHWrapper.h"
 
-#define LC_WITH_VISUAL_STUDIO_DTE 1
+#define LC_WITH_VISUAL_STUDIO_DTE 0
 
 #include "Windows/AllowWindowsPlatformTypes.h"
 #include "Windows/AllowWindowsPlatformAtomics.h"
diff --git a/Engine/Source/Developer/Windows/LiveCodingServer/Private/External/LC_VisualStudioAutomation.cpp b/Engine/Source/Developer/Windows/LiveCodingServer/Private/External/LC_VisualStudioAutomation.cpp
--- a/Engine/Source/Developer/Windows/LiveCodingServer/Private/External/LC_VisualStudioAutomation.cpp
+++ b/Engine/Source/Developer/Windows/LiveCodingServer/Private/External/LC_VisualStudioAutomation.cpp
@@ -6,6 +6,8 @@
 #include "LC_Logging.h"
 #include "Windows/COMPointer.h"
 
+#if LC_WITH_VISUAL_STUDIO_DTE
+
 namespace
 {
 	static COMThread* g_comThread = nullptr;
@@ -556,3 +558,5 @@ bool visualStudio::Break(const EnvDTE::DebuggerPtr& debugger)
 {
 	return g_comThread->CallInThread(&::Break, debugger);
 }
+
+#endif
diff --git a/Engine/Source/Developer/Windows/LiveCodingServer/Private/External/LC_VisualStudioAutomation.h b/Engine/Source/Developer/Windows/LiveCodingServer/Private/External/LC_VisualStudioAutomation.h
--- a/Engine/Source/Developer/Windows/LiveCodingServer/Private/External/LC_VisualStudioAutomation.h
+++ b/Engine/Source/Developer/Windows/LiveCodingServer/Private/External/LC_VisualStudioAutomation.h
@@ -5,6 +5,8 @@
 #include "LC_Types.h"
 #include "LC_VisualStudio.h"
 
+#if LC_WITH_VISUAL_STUDIO_DTE
+
 namespace visualStudio
 {
 	void Startup(void);
@@ -41,3 +43,5 @@ namespace visualStudio
 	// Breaks the process in the debugger
 	bool Break(const EnvDTE::DebuggerPtr& debugger);
 }
+
+#endif
