LiveCoding: Add support for the new restart functionality in Live++ 1.5.0. The Live Coding console window now includes a "Quick Restart" button, which restarts the process and adds the existing patches back in.

diff --git a/Engine/Source/Developer/Windows/LiveCoding/Private/External/LC_API.h b/Engine/Source/Developer/Windows/LiveCoding/Private/External/LC_API.h
--- a/Engine/Source/Developer/Windows/LiveCoding/Private/External/LC_API.h
+++ b/Engine/Source/Developer/Windows/LiveCoding/Private/External/LC_API.h
@@ -39,6 +39,14 @@ LPP_DLL_API(void*) LppEnableAllModules(const wchar_t* nameOfExeOrDll);
 LPP_DLL_API(void*) LppDisableModule(const wchar_t* nameOfExeOrDll);
 LPP_DLL_API(void*) LppDisableModules(const wchar_t* namesOfExeOrDll[], unsigned int count);
 LPP_DLL_API(void*) LppDisableAllModules(const wchar_t* nameOfExeOrDll);
+// BEGIN EPIC MOD - Additional functions
+LPP_DLL_API(void) LppShowConsole();
+LPP_DLL_API(void) LppSetVisible(bool visible);
+LPP_DLL_API(void) LppSetActive(bool active);
+LPP_DLL_API(void) LppSetBuildArguments(const wchar_t* arguments);
+LPP_DLL_API(void) LppEnableLazyLoadedModule(const wchar_t* nameOfExeOrDll);
+// END EPIC MOD
 LPP_DLL_API(void) LppApplySettingBool(const char* settingName, int value);
 LPP_DLL_API(void) LppApplySettingInt(const char* settingName, int value);
 LPP_DLL_API(void) LppApplySettingString(const char* settingName, const wchar_t* value);
+
diff --git a/Engine/Source/Developer/Windows/LiveCoding/Private/External/LC_EntryPoint.h b/Engine/Source/Developer/Windows/LiveCoding/Private/External/LC_EntryPoint.h
--- a/Engine/Source/Developer/Windows/LiveCoding/Private/External/LC_EntryPoint.h
+++ b/Engine/Source/Developer/Windows/LiveCoding/Private/External/LC_EntryPoint.h
@@ -2,24 +2,8 @@
 
 #pragma once
 
-void __cdecl LppRegisterProcessGroup(const char* groupName);
-void* __cdecl LppEnableModule(const wchar_t* nameOfExeOrDll);
-void* __cdecl LppEnableModules(const wchar_t* namesOfExeOrDll[], unsigned int count);
-void* __cdecl LppEnableAllModules(const wchar_t* nameOfExeOrDll);
-void* __cdecl LppDisableModule(const wchar_t* nameOfExeOrDll);
-void* __cdecl LppDisableAllModules(const wchar_t* nameOfExeOrDll);
-void __cdecl LppWaitForToken(void* token);
-void __cdecl LppTriggerRecompile(void);
-void __cdecl LppInstallExceptionHandler(void);
-void __cdecl LppUseExternalBuildSystem(void);
-void __cdecl LppShowConsole();
-void __cdecl LppSetVisible(bool visible);
-void __cdecl LppSetActive(bool active);
-void __cdecl LppSetBuildArguments(const wchar_t* arguments);
-void __cdecl LppApplySettingBool(const char* settingName, int value);
-void __cdecl LppApplySettingInt(const char* settingName, int value);
-void __cdecl LppApplySettingString(const char* settingName, const wchar_t* const value);
+#include "CoreTypes.h"
+#include "Windows/MinimalWindowsAPI.h"
 
-// BEGIN EPIC MOD - Support for lazy-loading modules
-void __cdecl LppEnableLazyLoadedModule(const wchar_t* nameOfExeOrDll);
-// END EPIC MODS
\ No newline at end of file
+void Startup(Windows::HINSTANCE hInstance);
+void Shutdown();
diff --git a/Engine/Source/Developer/Windows/LiveCoding/Private/LiveCodingModule.cpp b/Engine/Source/Developer/Windows/LiveCoding/Private/LiveCodingModule.cpp
--- a/Engine/Source/Developer/Windows/LiveCoding/Private/LiveCodingModule.cpp
+++ b/Engine/Source/Developer/Windows/LiveCoding/Private/LiveCodingModule.cpp
@@ -7,6 +7,7 @@
 #include "Misc/CoreDelegates.h"
 #include "LiveCodingLog.h"
 #include "External/LC_EntryPoint.h"
+#include "External/LC_API.h"
 #include "Misc/App.h"
 #include "Misc/Paths.h"
 #include "Misc/ConfigCacheIni.h"
@@ -100,7 +101,6 @@ void FLiveCodingModule::StartupModule()
 		);
 	}
 
-	extern void Startup(Windows::HINSTANCE hInstance);
 	Startup(hInstance);
 
 	if (Settings->bEnabled && !FApp::IsUnattended())
@@ -127,7 +127,6 @@ void FLiveCodingModule::StartupModule()
 
 void FLiveCodingModule::ShutdownModule()
 {
-	extern void Shutdown();
 	Shutdown();
 
 	FCoreDelegates::OnEndFrame.Remove(EndFrameDelegateHandle);
@@ -231,6 +230,11 @@ bool FLiveCodingModule::IsCompiling() const
 
 void FLiveCodingModule::Tick()
 {
+	if (LppWantsRestart())
+	{
+		LppRestart(lpp::LPP_RESTART_BEHAVIOUR_INSTANT_TERMINATION, 0);
+	}
+
 	if (Settings->bEnabled != bEnabledLastTick && Settings->Startup != ELiveCodingStartupMode::Manual)
 	{
 		EnableForSession(Settings->bEnabled);
diff --git a/Engine/Source/Developer/Windows/LiveCodingServer/Private/External/LC_AppSettings.cpp b/Engine/Source/Developer/Windows/LiveCodingServer/Private/External/LC_AppSettings.cpp
--- a/Engine/Source/Developer/Windows/LiveCodingServer/Private/External/LC_AppSettings.cpp
+++ b/Engine/Source/Developer/Windows/LiveCodingServer/Private/External/LC_AppSettings.cpp
@@ -443,7 +443,9 @@ void appSettings::Startup(const wchar_t* group)
 		L"install_compiled_patches_multi_process",
 		L"Install compiled patches",
 		L"Specifies whether compiled patches are installed into launched processes belonging to an existing process group",
-		false
+		// BEGIN EPIC MOD - changing default for restart functionality
+		true
+		// END EPIC MOD
 	);
 
 	g_amalgamationSplitIntoSingleParts = new SettingBool
diff --git a/Engine/Source/Developer/Windows/LiveCodingServer/Private/External/LC_Settings.cpp b/Engine/Source/Developer/Windows/LiveCodingServer/Private/External/LC_Settings.cpp
--- a/Engine/Source/Developer/Windows/LiveCodingServer/Private/External/LC_Settings.cpp
+++ b/Engine/Source/Developer/Windows/LiveCodingServer/Private/External/LC_Settings.cpp
@@ -14,6 +14,10 @@ namespace
 
 	static UINT LoadSetting(const wchar_t* group, const wchar_t* name, int initialValue)
 	{
+		// BEGIN EPIC MOD - Just use default Epic settings
+#if 1
+		return initialValue;
+#else
 		// first try loading the setting from the project settings file
 		const std::wstring& projectSettingsPath = appSettings::GetProjectSettingsPath();
 		const file::Attributes& attributes = file::GetAttributes(projectSettingsPath.c_str());
@@ -39,11 +43,17 @@ namespace
 		const UINT value = ::GetPrivateProfileIntW(group, name, initialValue, userSettingsPath.c_str());
 
 		return value;
+#endif
+		// END EPIC MOD
 	}
 
 
 	static std::wstring LoadSetting(const wchar_t* group, const wchar_t* name, const wchar_t* initialValue)
 	{
+		// BEGIN EPIC MOD - Just use default Epic settings
+#if 1
+		return initialValue;
+#else
 		// first try loading the setting from the project settings file
 		const std::wstring& projectSettingsPath = appSettings::GetProjectSettingsPath();
 		const file::Attributes& attributes = file::GetAttributes(projectSettingsPath.c_str());
@@ -72,6 +82,8 @@ namespace
 		::GetPrivateProfileStringW(group, name, initialValue, value, MAX_PATH, userSettingsPath.c_str());
 
 		return std::wstring(value);
+#endif
+		// END EPIC MOD
 	}
 }
 
diff --git a/Engine/Source/Developer/Windows/LiveCodingServer/Private/LiveCodingServer.cpp b/Engine/Source/Developer/Windows/LiveCodingServer/Private/LiveCodingServer.cpp
--- a/Engine/Source/Developer/Windows/LiveCodingServer/Private/LiveCodingServer.cpp
+++ b/Engine/Source/Developer/Windows/LiveCodingServer/Private/LiveCodingServer.cpp
@@ -49,6 +49,14 @@ void FLiveCodingServer::Stop()
 	ProcessGroupName.clear();
 }
 
+void FLiveCodingServer::RestartTargets()
+{
+	if (CommandThread != nullptr)
+	{
+		CommandThread->RestartTargets();
+	}
+}
+
 void FLiveCodingServer::SetLinkerPath(const wchar_t* LinkerPath, const TMap<FString, FString>& LinkerEnvironment)
 {
 	appSettings::g_linkerPath->SetValueWithoutSaving(LinkerPath);
diff --git a/Engine/Source/Developer/Windows/LiveCodingServer/Private/LiveCodingServer.h b/Engine/Source/Developer/Windows/LiveCodingServer/Private/LiveCodingServer.h
--- a/Engine/Source/Developer/Windows/LiveCodingServer/Private/LiveCodingServer.h
+++ b/Engine/Source/Developer/Windows/LiveCodingServer/Private/LiveCodingServer.h
@@ -18,6 +18,8 @@ public:
 	virtual void Start(const wchar_t* ProcessGroupName) override;
 	virtual void Stop() override;
 
+	virtual void RestartTargets() override;
+
 	virtual void SetLinkerPath(const wchar_t* LinkerPath, const TMap<FString, FString>& LinkerEnvironment) override;
 
 	// ILiveCodingServer implementation
diff --git a/Engine/Source/Developer/Windows/LiveCodingServer/Public/ILiveCodingServer.h b/Engine/Source/Developer/Windows/LiveCodingServer/Public/ILiveCodingServer.h
--- a/Engine/Source/Developer/Windows/LiveCodingServer/Public/ILiveCodingServer.h
+++ b/Engine/Source/Developer/Windows/LiveCodingServer/Public/ILiveCodingServer.h
@@ -29,6 +29,8 @@ public:
 	virtual void Start(const wchar_t* ProcessGroupName) = 0;
 	virtual void Stop() = 0;
 
+	virtual void RestartTargets() = 0;
+
 	virtual void SetLinkerPath(const wchar_t* LinkerPath, const TMap<FString, FString>& LinkerEnvironment) = 0;
 
 	DECLARE_DELEGATE(FBringToFrontDelegate);
diff --git a/Engine/Source/Programs/LiveCodingConsole/Private/LiveCodingConsole.cpp b/Engine/Source/Programs/LiveCodingConsole/Private/LiveCodingConsole.cpp
--- a/Engine/Source/Programs/LiveCodingConsole/Private/LiveCodingConsole.cpp
+++ b/Engine/Source/Programs/LiveCodingConsole/Private/LiveCodingConsole.cpp
@@ -62,7 +62,27 @@ public:
 			.ActivationPolicy(EWindowActivationPolicy::Never)
 			.IsInitiallyMaximized(false)
 			[
-				LogWidget.ToSharedRef()
+				SNew(SBorder)
+				.BorderImage(FCoreStyle::Get().GetBrush("ToolPanel.GroupBorder"))
+				[
+					SNew(SVerticalBox)
+
+					+ SVerticalBox::Slot()
+					.FillHeight(1.0f)
+					[
+						LogWidget.ToSharedRef()
+					]
+					
+					+ SVerticalBox::Slot()
+					.HAlign(HAlign_Center)
+					.AutoHeight()
+					.Padding(0.0f, 6.0f, 0.0f, 4.0f)
+					[
+						SNew(SButton)
+						.Text(LOCTEXT("QuickRestart", "Quick Restart"))
+						.OnClicked(FOnClicked::CreateRaw(this, &FLiveCodingConsoleApp::RestartTargets))
+					]
+				]
 			];
 
 		// Add the window without showing it
@@ -282,6 +302,12 @@ private:
 		return bRequestCancel;
 	}
 
+	FReply RestartTargets()
+	{
+		Server.RestartTargets();
+		return FReply::Handled();
+	}
+
 	void SetVisibleAsync(bool bVisible)
 	{
 		FScopeLock Lock(&CriticalSection);
