DR - Using float3/half3/unit3/int3 as a loose/global parameter will now fail to compile on D3D
#rnx

diff --git a/Engine/Source/Developer/Windows/ShaderFormatD3D/Private/D3D11ShaderCompiler.cpp b/Engine/Source/Developer/Windows/ShaderFormatD3D/Private/D3D11ShaderCompiler.cpp
--- a/Engine/Source/Developer/Windows/ShaderFormatD3D/Private/D3D11ShaderCompiler.cpp
+++ b/Engine/Source/Developer/Windows/ShaderFormatD3D/Private/D3D11ShaderCompiler.cpp
@@ -717,6 +717,16 @@ static bool CompileAndProcessD3DShader(FString& PreprocessedShaderSource, const
 							{
 								bGlobalUniformBufferUsed = true;
 
+								//#todo-rco: We can't have global/loose float3's currently to work around a fundamental issue on hlslcc. FXC already pads them to 16-bytes (unless a single float follows) which is not our usual case.
+								if (VariableDesc.Size == 12)
+								{
+									FString Name = ANSI_TO_TCHAR(VariableDesc.Name);
+									FString Msg = FString::Printf(TEXT("Error: Global parameter '%s' is a 3-component vector(float3/uint3/int3/half3), needs to be converted to a 4-component vector."), *Name);
+									FShaderCompilerError Error(*Input.VirtualSourceFilePath, TEXT("1,1"), *Msg);
+									Output.Errors.Add(Error);
+									Output.bSucceeded = false;
+								}
+
 								Output.ParameterMap.AddParameterAllocation(
 									ANSI_TO_TCHAR(VariableDesc.Name),
 									CBIndex,
diff --git a/Engine/Source/Runtime/RenderCore/Public/ShaderCore.h b/Engine/Source/Runtime/RenderCore/Public/ShaderCore.h
--- a/Engine/Source/Runtime/RenderCore/Public/ShaderCore.h
+++ b/Engine/Source/Runtime/RenderCore/Public/ShaderCore.h
@@ -699,6 +699,12 @@ struct FShaderCompilerError
 	,	StrippedErrorMessage(InStrippedErrorMessage)
 	{}
 
+	FShaderCompilerError(const TCHAR* InVirtualFilePath, const TCHAR* InLineString, const TCHAR* InStrippedErrorMessage)
+		: ErrorVirtualFilePath(InVirtualFilePath)
+		, ErrorLineString(InLineString)
+		, StrippedErrorMessage(InStrippedErrorMessage)
+	{}
+
 	FString ErrorVirtualFilePath;
 	FString ErrorLineString;
 	FString StrippedErrorMessage;
