Fixed several static analysis warnings

- Optimization flags handling in D3DCreateDXCArguments
- Switch over PipelineType in FD3D12StateCacheBase::ApplyState
- LocalCBVs access in SetRayTracingShaderResources
- Possible nullptr dereference in FD3D12DynamicRHI::RHICreateUniformBuffer

diff --git a/Engine/Source/Developer/Windows/ShaderFormatD3D/Private/D3D11ShaderCompiler.cpp b/Engine/Source/Developer/Windows/ShaderFormatD3D/Private/D3D11ShaderCompiler.cpp
--- a/Engine/Source/Developer/Windows/ShaderFormatD3D/Private/D3D11ShaderCompiler.cpp
+++ b/Engine/Source/Developer/Windows/ShaderFormatD3D/Private/D3D11ShaderCompiler.cpp
@@ -554,7 +554,7 @@ static HRESULT D3DCompileToDxil(const char* SourceText, LPCWSTR EntryPoint, LPCW
 	return CompileResultCode;
 }
 
-static void D3DCreateDXCArguments(TArray<const WCHAR*>& OutArgs, const WCHAR* Exports, uint32 CompileFlags, uint32 AutoBindingSpace = ~0u)
+static void D3DCreateDXCArguments(TArray<const WCHAR*>& OutArgs, const WCHAR* Exports, uint32 CompileFlags, FShaderCompilerOutput& Output, uint32 AutoBindingSpace = ~0u)
 {
 	// Static digit strings are used here as they are returned in OutArgs
 	static const WCHAR* DigitStrings[] = 
@@ -619,20 +619,31 @@ static void D3DCreateDXCArguments(TArray<const WCHAR*>& OutArgs, const WCHAR* Ex
 		OutArgs.Add(L"/Gec");
 	}
 
-	if ((CompileFlags & D3D10_SHADER_OPTIMIZATION_LEVEL2) == D3D10_SHADER_OPTIMIZATION_LEVEL2)
+	switch (CompileFlags & SHADER_OPTIMIZATION_LEVEL_MASK)
 	{
+	case D3D10_SHADER_OPTIMIZATION_LEVEL0:
+		CompileFlags &= ~D3D10_SHADER_OPTIMIZATION_LEVEL0;
+		OutArgs.Add(L"/O0");
+		break;
+
+	case D3D10_SHADER_OPTIMIZATION_LEVEL1:
+		CompileFlags &= ~D3D10_SHADER_OPTIMIZATION_LEVEL1;
+		OutArgs.Add(L"/O1");
+		break;
+
+	case D3D10_SHADER_OPTIMIZATION_LEVEL2:
 		CompileFlags &= ~D3D10_SHADER_OPTIMIZATION_LEVEL2;
 		OutArgs.Add(L"/O2");
-	}
-	else if (CompileFlags & D3D10_SHADER_OPTIMIZATION_LEVEL3)
-	{
+		break;
+
+	case D3D10_SHADER_OPTIMIZATION_LEVEL3:
 		CompileFlags &= ~D3D10_SHADER_OPTIMIZATION_LEVEL3;
 		OutArgs.Add(L"/O3");
-	}
-	else if (CompileFlags & D3D10_SHADER_OPTIMIZATION_LEVEL1)
-	{
-		CompileFlags &= ~D3D10_SHADER_OPTIMIZATION_LEVEL1;
-		OutArgs.Add(L"/O1");
+		break;
+
+	default:
+		Output.Errors.Emplace(TEXT("Unknown optimization level flag"));
+		break;
 	}
 
 	checkf(CompileFlags == 0, TEXT("Unhandled shader compiler flag!"));
@@ -998,7 +1009,7 @@ static bool CompileAndProcessD3DShader(FString& PreprocessedShaderSource, const
 		// Ignore backwards compatibility flag (/Gec) as it is deprecated.
 		// #dxr_todo: this flag should not be even passed into this function from the higher level.
 		const uint32 DXCFlags = CompileFlags & (~D3D10_SHADER_ENABLE_BACKWARDS_COMPATIBILITY);
-		D3DCreateDXCArguments(Args, *Exports, DXCFlags, AutoBindingSpace);
+		D3DCreateDXCArguments(Args, *Exports, DXCFlags, Output, AutoBindingSpace);
 
 		TRefCountPtr<IDxcBlobEncoding> DxcErrorBlob;
 
diff --git a/Engine/Source/Runtime/D3D12RHI/Private/D3D12RayTracing.cpp b/Engine/Source/Runtime/D3D12RHI/Private/D3D12RayTracing.cpp
--- a/Engine/Source/Runtime/D3D12RHI/Private/D3D12RayTracing.cpp
+++ b/Engine/Source/Runtime/D3D12RHI/Private/D3D12RayTracing.cpp
@@ -2012,7 +2012,8 @@ static void SetRayTracingShaderResources(
 
 		for (uint32 i = 0; i < Shader->ResourceCounts.NumCBs; ++i)
 		{
-			D3D12_GPU_VIRTUAL_ADDRESS BufferAddress = LocalCBVs[i]->ResourceLocation.GetGPUVirtualAddress();
+			const uint64 SlotMask = (1ull << i);
+			D3D12_GPU_VIRTUAL_ADDRESS BufferAddress = (BoundCBVMask & SlotMask) ? LocalCBVs[i]->ResourceLocation.GetGPUVirtualAddress() : 0;
 			Binder.SetRootCBV(BindSlot, i, BufferAddress);
 		}
 	}
diff --git a/Engine/Source/Runtime/D3D12RHI/Private/D3D12StateCache.cpp b/Engine/Source/Runtime/D3D12RHI/Private/D3D12StateCache.cpp
--- a/Engine/Source/Runtime/D3D12RHI/Private/D3D12StateCache.cpp
+++ b/Engine/Source/Runtime/D3D12RHI/Private/D3D12StateCache.cpp
@@ -354,22 +354,12 @@ void FD3D12StateCacheBase::ApplyState()
 
 	FD3D12CommandListHandle& CommandList = CmdContext->CommandListHandle;
 	const FD3D12RootSignature* pRootSignature = nullptr;
-	switch (PipelineType)
-	{
-	case D3D12PT_Graphics:
-		pRootSignature = GetGraphicsRootSignature();
-		break;
-	case D3D12PT_Compute:
-		pRootSignature = PipelineState.Compute.CurrentPipelineStateObject->ComputeShader->pRootSignature;
-		break;
-	default:
-		check(false);
-		break;
-	}
 
 	// PSO
 	if (PipelineType == D3D12PT_Compute)
 	{
+		pRootSignature = PipelineState.Compute.CurrentPipelineStateObject->ComputeShader->pRootSignature;
+
 		// See if we need to set a compute root signature
 		if (PipelineState.Compute.bNeedSetRootSignature)
 		{
@@ -385,6 +375,8 @@ void FD3D12StateCacheBase::ApplyState()
 	}
 	else if (PipelineType == D3D12PT_Graphics)
 	{
+		pRootSignature = GetGraphicsRootSignature();
+
 		// See if we need to set a graphics root signature
 		if (PipelineState.Graphics.bNeedSetRootSignature)
 		{
diff --git a/Engine/Source/Runtime/D3D12RHI/Private/D3D12UniformBuffer.cpp b/Engine/Source/Runtime/D3D12RHI/Private/D3D12UniformBuffer.cpp
--- a/Engine/Source/Runtime/D3D12RHI/Private/D3D12UniformBuffer.cpp
+++ b/Engine/Source/Runtime/D3D12RHI/Private/D3D12UniformBuffer.cpp
@@ -67,6 +67,8 @@ FUniformBufferRHIRef FD3D12DynamicRHI::RHICreateUniformBuffer(const void* Conten
 		return NewUniformBuffer;
 	});
 
+	check(UniformBufferOut);
+
 	if (Layout.Resources.Num())
 	{
 		const int32 NumResources = Layout.Resources.Num();
@@ -111,7 +113,10 @@ FUniformBufferRHIRef FD3D12DynamicRHI::RHICreateUniformBuffer(const void* Conten
 		}
 	}
 
-	UpdateBufferStats<FD3D12UniformBuffer>(&UniformBufferOut->ResourceLocation, true);
+	if (UniformBufferOut)
+	{
+		UpdateBufferStats<FD3D12UniformBuffer>(&UniformBufferOut->ResourceLocation, true);
+	}
 
 	return UniformBufferOut;
 }
