Implemented ray tracing shader compile batch script generation and shader worker direct compile support

diff --git a/Engine/Source/Developer/ShaderCompilerCommon/Private/ShaderCompilerCommon.cpp b/Engine/Source/Developer/ShaderCompilerCommon/Private/ShaderCompilerCommon.cpp
--- a/Engine/Source/Developer/ShaderCompilerCommon/Private/ShaderCompilerCommon.cpp
+++ b/Engine/Source/Developer/ShaderCompilerCommon/Private/ShaderCompilerCommon.cpp
@@ -777,6 +777,11 @@ FString CreateShaderCompilerWorkerDirectCommandLine(const FShaderCompilerInput&
 	case SF_Geometry:	Text += TEXT(" -gs"); break;
 	case SF_Pixel:		Text += TEXT(" -ps"); break;
 	case SF_Compute:	Text += TEXT(" -cs"); break;
+#if RHI_RAYTRACING
+	case SF_RayGen:			Text += TEXT(" -rgs"); break;
+	case SF_RayMiss:		Text += TEXT(" -rms"); break;
+	case SF_RayHitGroup:	Text += TEXT(" -rhs"); break;
+#endif // RHI_RAYTRACING
 	default: break;
 	}
 	if (Input.bCompilingForShaderPipeline)
diff --git a/Engine/Source/Developer/Windows/ShaderFormatD3D/Private/D3D11ShaderCompiler.cpp b/Engine/Source/Developer/Windows/ShaderFormatD3D/Private/D3D11ShaderCompiler.cpp
--- a/Engine/Source/Developer/Windows/ShaderFormatD3D/Private/D3D11ShaderCompiler.cpp
+++ b/Engine/Source/Developer/Windows/ShaderFormatD3D/Private/D3D11ShaderCompiler.cpp
@@ -649,6 +649,37 @@ static void D3DCreateDXCArguments(TArray<const WCHAR*>& OutArgs, const WCHAR* Ex
 	checkf(CompileFlags == 0, TEXT("Unhandled shader compiler flag!"));
 }
 
+static FString D3DCreateDXCCompileBatchFile(const FString& ShaderPath, const WCHAR* Exports, const TCHAR* ShaderProfile, uint32 CompileFlags, FShaderCompilerOutput& CompilerOutput, uint32 AutoBindingSpace = ~0u)
+{
+	TArray<const WCHAR*> Args;
+	const uint32 DXCFlags = CompileFlags & (~D3D10_SHADER_ENABLE_BACKWARDS_COMPATIBILITY);
+	D3DCreateDXCArguments(Args, Exports, DXCFlags, CompilerOutput, AutoBindingSpace);
+
+	FString BatchFileHeader = TEXT("@ECHO OFF\nSET DXC=\"C:\\Program Files (x86)\\Windows Kits\\10\\bin\\10.0.17763.0\\x64\\dxc.exe\"\n"\
+		"IF EXIST %DXC% (\nREM\n) ELSE (\nECHO Couldn't find Windows 10.0.17763 SDK, falling back to dxc.exe in PATH...\n"\
+		"SET DXC=dxc.exe)\n");
+
+	FString DXCCommandline = FString(TEXT("%DXC%"));
+	for (const WCHAR* Arg : Args)
+	{
+		DXCCommandline += TEXT(" ");
+		DXCCommandline += Arg;
+	}
+
+	DXCCommandline += TEXT(" /T ");
+	DXCCommandline += ShaderProfile;
+
+	if (FPaths::GetExtension(ShaderPath) == TEXT("usf"))
+	{
+		DXCCommandline += FString::Printf(TEXT(" /Fc%sd3dasm"), *ShaderPath.LeftChop(3));
+	}
+
+	DXCCommandline += TEXT(" ");
+	DXCCommandline += ShaderPath;
+
+	return BatchFileHeader + DXCCommandline + TEXT("\npause");
+}
+
 #ifndef DXIL_FOURCC
 #define DXIL_FOURCC(ch0, ch1, ch2, ch3) (                            \
   (uint32_t)(uint8_t)(ch0)        | (uint32_t)(uint8_t)(ch1) << 8  | \
@@ -926,6 +957,38 @@ static bool CompileAndProcessD3DShader(FString& PreprocessedShaderSource, const
 {
 	auto AnsiSourceFile = StringCast<ANSICHAR>(*PreprocessedShaderSource);
 
+	const bool bIsRayTracingShader = IsRayTracingShader(Input.Target);
+
+#if RHI_RAYTRACING
+
+	const uint32 AutoBindingSpace = GetAutoBindingSpace(Input.Target);
+
+	FString RayEntryPoint; // Primary entry point for all ray tracing shaders
+	FString RayAnyHitEntryPoint; // Optional for hit group shaders
+	FString RayIntersectionEntryPoint; // Optional for hit group shaders
+	FString RayTracingExports;
+
+	if (bIsRayTracingShader)
+	{
+		ParseRayTracingEntryPoint(Input.EntryPointName, RayEntryPoint, RayAnyHitEntryPoint, RayIntersectionEntryPoint);
+
+		RayTracingExports = RayEntryPoint;
+
+		if (!RayAnyHitEntryPoint.IsEmpty())
+		{
+			RayTracingExports += TEXT(";");
+			RayTracingExports += RayAnyHitEntryPoint;
+		}
+
+		if (!RayIntersectionEntryPoint.IsEmpty())
+		{
+			RayTracingExports += TEXT(";");
+			RayTracingExports += RayIntersectionEntryPoint;
+		}
+	}
+
+#endif // RHI_RAYTRACING
+
 	// Write out the preprocessed file and a batch file to compile it if requested (DumpDebugInfoPath is valid)
 	if (Input.DumpDebugInfoPath.Len() > 0 && IFileManager::Get().DirectoryExists(*Input.DumpDebugInfoPath))
 	{
@@ -947,14 +1010,26 @@ static bool CompileAndProcessD3DShader(FString& PreprocessedShaderSource, const
 			delete FileWriter;
 		}
 
-		const FString BatchFileContents = D3D11CreateShaderCompileCommandLine(Filename, *EntryPointName, ShaderProfile, CompileFlags, Output);
-		FFileHelper::SaveStringToFile(BatchFileContents, *(Input.DumpDebugInfoPath / TEXT("CompileD3D.bat")));
+		FString BatchFileContents;
 
-		if (GD3DDumpAMDCodeXLFile)
+		if (bIsRayTracingShader)
 		{
-			const FString BatchFileContents2 = CreateAMDCodeXLCommandLine(Filename, *EntryPointName, ShaderProfile, CompileFlags);
-			FFileHelper::SaveStringToFile(BatchFileContents2, *(Input.DumpDebugInfoPath / TEXT("CompileAMD.bat")));
+#if RHI_RAYTRACING
+			BatchFileContents = D3DCreateDXCCompileBatchFile(Filename, *RayTracingExports, ShaderProfile, CompileFlags, Output, AutoBindingSpace);
+#endif // RHI_RAYTRACING
 		}
+		else
+		{
+			BatchFileContents = D3D11CreateShaderCompileCommandLine(Filename, *EntryPointName, ShaderProfile, CompileFlags, Output);
+
+			if (GD3DDumpAMDCodeXLFile)
+			{
+				const FString BatchFileContents2 = CreateAMDCodeXLCommandLine(Filename, *EntryPointName, ShaderProfile, CompileFlags);
+				FFileHelper::SaveStringToFile(BatchFileContents2, *(Input.DumpDebugInfoPath / TEXT("CompileAMD.bat")));
+			}
+		}
+
+		FFileHelper::SaveStringToFile(BatchFileContents, *(Input.DumpDebugInfoPath / TEXT("CompileD3D.bat")));
 
 		if (Input.bGenerateDirectCompileFile)
 		{
@@ -971,45 +1046,16 @@ static bool CompileAndProcessD3DShader(FString& PreprocessedShaderSource, const
 	pD3DStripShader D3DStripShaderFunc = nullptr;
 	bool bCompilerPathFunctionsUsed = false;
 
-	// #dxr_todo: could split this function into DX11, DX12/DXR and shared common parts
-
-	const bool bIsRayTracingShader = IsRayTracingShader(Input.Target);
-	const uint32 AutoBindingSpace = GetAutoBindingSpace(Input.Target);
-
-#if RHI_RAYTRACING
-	FString RayEntryPoint; // Primary entry point for all ray tracing shaders
-	FString RayAnyHitEntryPoint; // Optional for hit group shaders
-	FString RayIntersectionEntryPoint; // Optional for hit group shaders
-	if (bIsRayTracingShader)
-	{
-		ParseRayTracingEntryPoint(Input.EntryPointName, RayEntryPoint, RayAnyHitEntryPoint, RayIntersectionEntryPoint);
-	}
-#endif // RHI_RAYTRACING
-
 	if (bIsRayTracingShader)
 	{
 #if RHI_RAYTRACING
 
 		TArray<const WCHAR*> Args;
-		
-		FString Exports = RayEntryPoint;
-
-		if (!RayAnyHitEntryPoint.IsEmpty())
-		{
-			Exports += TEXT(";");
-			Exports += RayAnyHitEntryPoint;
-		}
-
-		if (!RayIntersectionEntryPoint.IsEmpty())
-		{
-			Exports += TEXT(";");
-			Exports += RayIntersectionEntryPoint;
-		}
 
 		// Ignore backwards compatibility flag (/Gec) as it is deprecated.
 		// #dxr_todo: this flag should not be even passed into this function from the higher level.
 		const uint32 DXCFlags = CompileFlags & (~D3D10_SHADER_ENABLE_BACKWARDS_COMPATIBILITY);
-		D3DCreateDXCArguments(Args, *Exports, DXCFlags, Output, AutoBindingSpace);
+		D3DCreateDXCArguments(Args, *RayTracingExports, DXCFlags, Output, AutoBindingSpace);
 
 		TRefCountPtr<IDxcBlobEncoding> DxcErrorBlob;
 
diff --git a/Engine/Source/Programs/ShaderCompileWorker/Private/ShaderCompileWorker.cpp b/Engine/Source/Programs/ShaderCompileWorker/Private/ShaderCompileWorker.cpp
--- a/Engine/Source/Programs/ShaderCompileWorker/Private/ShaderCompileWorker.cpp
+++ b/Engine/Source/Programs/ShaderCompileWorker/Private/ShaderCompileWorker.cpp
@@ -758,6 +758,20 @@ static void DirectCompile(const TArray<const class IShaderFormat*>& ShaderFormat
 			{
 				Frequency = SF_Compute;
 			}
+#if RHI_RAYTRACING
+			else if (!FCString::Strcmp(*Token, TEXT("rgs")))
+			{
+				Frequency = SF_RayGen;
+			}
+			else if (!FCString::Strcmp(*Token, TEXT("rms")))
+			{
+				Frequency = SF_RayMiss;
+			}
+			else if (!FCString::Strcmp(*Token, TEXT("rhs")))
+			{
+				Frequency = SF_RayHitGroup;
+			}
+#endif // RHI_RAYTRACING
 			else if (!FCString::Strcmp(*Token, TEXT("pipeline")))
 			{
 				bPipeline = true;
