DR - Remove unused cvars from D3D shader compiler
- Write D3D asm file when dumping shaders
#rnx

diff --git a/Engine/Source/Developer/Windows/ShaderFormatD3D/Private/D3D11ShaderCompiler.cpp b/Engine/Source/Developer/Windows/ShaderFormatD3D/Private/D3D11ShaderCompiler.cpp
--- a/Engine/Source/Developer/Windows/ShaderFormatD3D/Private/D3D11ShaderCompiler.cpp
+++ b/Engine/Source/Developer/Windows/ShaderFormatD3D/Private/D3D11ShaderCompiler.cpp
@@ -46,34 +46,7 @@ static int32 GD3DAllowRemoveUnused = 0;
 
 
 static int32 GD3DCheckForDoubles = 1;
-static FAutoConsoleVariableRef CVarD3DCheckForDoubles(
-	TEXT("r.D3DCheckShadersForDouble"),
-	GD3DCheckForDoubles,
-	TEXT("Enables checking the D3D microcode for uses of double. This is not allowed on all D3D11 cards.\n")
-	TEXT(" 0: Do not check for faster compilation\n")
-	TEXT(" 1: Enable checking and error if found (default)"),
-	ECVF_Default
-	);
-
 static int32 GD3DDumpAMDCodeXLFile = 0;
-static FAutoConsoleVariableRef CVarD3DDumpAMDCodeXLFile(
-	TEXT("r.D3DDumpAMDCodeXLFile"),
-	GD3DDumpAMDCodeXLFile,
-	TEXT("When r.DumpShaderDebugInfo is enabled, this will generate a batch file for running CodeXL.\n")
-	TEXT(" 0: Do not generate extra batch file (default)\n")
-	TEXT(" 1: Enable generating extra batch file"),
-	ECVF_Default
-	);
-
-static int32 GD3DDumpD3DAsmFile = 0;
-static FAutoConsoleVariableRef CVarD3DDumpD3DAsmFile(
-	TEXT("r.D3DDumpD3DAsm"),
-	GD3DDumpD3DAsmFile,
-	TEXT("When r.DumpShaderDebugInfo is enabled, this will generate a text file with the fxc assembly.\n")
-	TEXT(" 0: Do not generate extra file (default)\n")
-	TEXT(" 1: Enable generating extra disassembly file"),
-	ECVF_Default
-	);
 
 /**
  * TranslateCompilerFlag - translates the platform-independent compiler flags into D3DX defines
@@ -989,9 +962,11 @@ static bool CompileAndProcessD3DShader(FString& PreprocessedShaderSource, const
 
 #endif // RHI_RAYTRACING
 
+	bool bDumpDebugInfo = false;
 	// Write out the preprocessed file and a batch file to compile it if requested (DumpDebugInfoPath is valid)
 	if (Input.DumpDebugInfoPath.Len() > 0 && IFileManager::Get().DirectoryExists(*Input.DumpDebugInfoPath))
 	{
+		bDumpDebugInfo = true;
 		FString Filename = Input.GetSourceFilename();
 		FArchive* FileWriter = IFileManager::Get().CreateFileWriter(*(Input.DumpDebugInfoPath / Filename));
 		if (FileWriter)
@@ -1127,7 +1102,7 @@ static bool CompileAndProcessD3DShader(FString& PreprocessedShaderSource, const
 		// Fail the compilation if double operations are being used, since those are not supported on all D3D11 cards
 		if (SUCCEEDED(Result))
 		{
-			if (D3DDisassembleFunc && (GD3DCheckForDoubles || GD3DDumpD3DAsmFile))
+			if (D3DDisassembleFunc && (GD3DCheckForDoubles || bDumpDebugInfo))
 			{
 				TRefCountPtr<ID3DBlob> Dissasembly;
 				if (SUCCEEDED(D3DDisassembleFunc(Shader->GetBufferPointer(), Shader->GetBufferSize(), 0, "", Dissasembly.GetInitReference())))
@@ -1138,7 +1113,7 @@ static bool CompileAndProcessD3DShader(FString& PreprocessedShaderSource, const
 					FString DissasemblyStringW(DissasemblyString);
 					delete[] DissasemblyString;
 
-					if (GD3DDumpD3DAsmFile)
+					if (bDumpDebugInfo)
 					{
 						FFileHelper::SaveStringToFile(DissasemblyStringW, *(Input.DumpDebugInfoPath / TEXT("Output.d3dasm")));
 					}
