Renamed D3D11ShaderCompiler.cpp -> D3DShaderCompiler.cpp, as it is now used for both DXBC and DXIL shader compilation.

diff --git a/Engine/Source/Developer/Windows/ShaderFormatD3D/Private/D3D11ShaderCompiler.cpp b/Engine/Source/Developer/Windows/ShaderFormatD3D/Private/D3D11ShaderCompiler.cpp
deleted file mode 100644
diff --git a/Engine/Source/Developer/Windows/ShaderFormatD3D/Private/D3DShaderCompiler.cpp b/Engine/Source/Developer/Windows/ShaderFormatD3D/Private/D3DShaderCompiler.cpp
new file mode 100644
--- /dev/null
+++ b/Engine/Source/Developer/Windows/ShaderFormatD3D/Private/D3DShaderCompiler.cpp
@@ -0,0 +1,1746 @@
+// Copyright 1998-2019 Epic Games, Inc. All Rights Reserved.
+
+#include "ShaderFormatD3D.h"
+#include "ShaderPreprocessor.h"
+#include "ShaderCompilerCommon.h"
+#include "D3D11ShaderResources.h"
+#include "Misc/Paths.h"
+#include "Misc/FileHelper.h"
+#include "HAL/FileManager.h"
+#include "Serialization/MemoryWriter.h"
+#include "RayTracingDefinitions.ush"
+
+DEFINE_LOG_CATEGORY_STATIC(LogD3D11ShaderCompiler, Log, All);
+
+#define DEBUG_SHADERS 0
+
+// D3D headers.
+#define D3D_OVERLOADS 1
+
+// D3D doesn't define a mask for this, so we do so here
+#define SHADER_OPTIMIZATION_LEVEL_MASK (D3D10_SHADER_OPTIMIZATION_LEVEL0 | D3D10_SHADER_OPTIMIZATION_LEVEL1 | D3D10_SHADER_OPTIMIZATION_LEVEL2 | D3D10_SHADER_OPTIMIZATION_LEVEL3)
+
+// Disable macro redefinition warning for compatibility with Windows SDK 8+
+#pragma warning(push)
+#pragma warning(disable : 4005)	// macro redefinition
+
+#include "Windows/AllowWindowsPlatformTypes.h"
+	#include <D3D11.h>
+	#include <D3Dcompiler.h>
+	#include <d3d11Shader.h>
+#include "Windows/HideWindowsPlatformTypes.h"
+#undef DrawText
+
+#pragma warning(pop)
+
+#if RHI_RAYTRACING
+MSVC_PRAGMA(warning(push))
+MSVC_PRAGMA(warning(disable : 4191)) // warning C4191: 'type cast': unsafe conversion from 'FARPROC' to 'DxcCreateInstanceProc'
+#include <dxc/dxcapi.h>
+#include <dxc/Support/dxcapi.use.h>
+#include <d3d12shader.h>
+MSVC_PRAGMA(warning(pop))
+#endif // RHI_RAYTRACING
+
+static int32 GD3DAllowRemoveUnused = 0;
+
+
+static int32 GD3DCheckForDoubles = 1;
+static int32 GD3DDumpAMDCodeXLFile = 0;
+
+/**
+ * TranslateCompilerFlag - translates the platform-independent compiler flags into D3DX defines
+ * @param CompilerFlag - the platform-independent compiler flag to translate
+ * @return uint32 - the value of the appropriate D3DX enum
+ */
+static uint32 TranslateCompilerFlagD3D11(ECompilerFlags CompilerFlag)
+{
+	// @TODO - currently d3d11 uses d3d10 shader compiler flags... update when this changes in DXSDK
+	switch(CompilerFlag)
+	{
+	case CFLAG_PreferFlowControl: return D3D10_SHADER_PREFER_FLOW_CONTROL;
+	case CFLAG_AvoidFlowControl: return D3D10_SHADER_AVOID_FLOW_CONTROL;
+	default: return 0;
+	};
+}
+
+/**
+ * Filters out unwanted shader compile warnings
+ */
+static void D3D11FilterShaderCompileWarnings(const FString& CompileWarnings, TArray<FString>& FilteredWarnings)
+{
+	TArray<FString> WarningArray;
+	FString OutWarningString = TEXT("");
+	CompileWarnings.ParseIntoArray(WarningArray, TEXT("\n"), true);
+	
+	//go through each warning line
+	for (int32 WarningIndex = 0; WarningIndex < WarningArray.Num(); WarningIndex++)
+	{
+		//suppress "warning X3557: Loop only executes for 1 iteration(s), forcing loop to unroll"
+		if (!WarningArray[WarningIndex].Contains(TEXT("X3557"))
+			// "warning X3205: conversion from larger type to smaller, possible loss of data"
+			// Gets spammed when converting from float to half
+			&& !WarningArray[WarningIndex].Contains(TEXT("X3205")))
+		{
+			FilteredWarnings.AddUnique(WarningArray[WarningIndex]);
+		}
+	}
+}
+
+static bool IsRayTracingShader(const FShaderTarget& Target)
+{
+#if RHI_RAYTRACING
+	switch(Target.Frequency)
+	{
+	case SF_RayGen:
+	case SF_RayMiss:
+	case SF_RayHitGroup:
+		return true;
+	default:
+		return false;
+	}
+#else // RHI_RAYTRACING
+	return false;
+#endif // RHI_RAYTRACING
+}
+
+static uint32 GetAutoBindingSpace(const FShaderTarget& Target)
+{
+#if RHI_RAYTRACING
+	switch (Target.Frequency)
+	{
+	case SF_RayGen:
+	case SF_RayMiss:
+		return RAY_TRACING_REGISTER_SPACE_GLOBAL;
+	case SF_RayHitGroup:
+		return RAY_TRACING_REGISTER_SPACE_LOCAL;
+	default:
+		return 0;
+	}
+#else // RHI_RAYTRACING
+	return 0;
+#endif // RHI_RAYTRACING
+}
+
+// @return 0 if not recognized
+static const TCHAR* GetShaderProfileName(FShaderTarget Target)
+{
+	if(Target.Platform == SP_PCD3D_SM5)
+	{
+		//set defines and profiles for the appropriate shader paths
+		switch(Target.Frequency)
+		{
+		default:
+			checkfSlow(false, TEXT("Unexpected shader frequency"));
+			return nullptr;
+		case SF_Pixel:
+			return TEXT("ps_5_0");
+		case SF_Vertex:
+			return TEXT("vs_5_0");
+		case SF_Hull:
+			return TEXT("hs_5_0");
+		case SF_Domain:
+			return TEXT("ds_5_0");
+		case SF_Geometry:
+			return TEXT("gs_5_0");
+		case SF_Compute:
+			return TEXT("cs_5_0");
+#if RHI_RAYTRACING
+		case SF_RayGen:
+		case SF_RayMiss:
+		case SF_RayHitGroup:
+			return TEXT("lib_6_3");
+#endif // RHI_RAYTRACING
+		}
+	}
+	else if(Target.Platform == SP_PCD3D_SM4)
+	{
+		checkSlow(Target.Frequency == SF_Vertex ||
+			Target.Frequency == SF_Pixel ||
+			Target.Frequency == SF_Geometry);
+
+		//set defines and profiles for the appropriate shader paths
+		switch(Target.Frequency)
+		{
+		case SF_Pixel:
+			return TEXT("ps_4_0");
+		case SF_Vertex:
+			return TEXT("vs_4_0");
+		case SF_Geometry:
+			return TEXT("gs_4_0");
+		}
+	}
+	else if (Target.Platform == SP_PCD3D_ES2 || Target.Platform == SP_PCD3D_ES3_1)
+	{
+		checkSlow(Target.Frequency == SF_Vertex ||
+			Target.Frequency == SF_Pixel ||
+			Target.Frequency == SF_Geometry || 
+			Target.Frequency == SF_Compute);
+
+		//set defines and profiles for the appropriate shader paths
+		switch(Target.Frequency)
+		{
+		case SF_Pixel:
+			return TEXT("ps_5_0");
+		case SF_Vertex:
+			return TEXT("vs_5_0");
+		case SF_Geometry:
+			return TEXT("gs_5_0");
+		case SF_Compute:
+			return TEXT("cs_5_0");
+		}
+	}
+
+	return NULL;
+}
+
+/**
+ * D3D11CreateShaderCompileCommandLine - takes shader parameters used to compile with the DX11
+ * compiler and returns an fxc command to compile from the command line
+ */
+static FString D3D11CreateShaderCompileCommandLine(
+	const FString& ShaderPath, 
+	const TCHAR* EntryFunction, 
+	const TCHAR* ShaderProfile, 
+	uint32 CompileFlags,
+	FShaderCompilerOutput& Output
+	)
+{
+	// fxc is our command line compiler
+	FString FXCCommandline = FString(TEXT("%FXC% ")) + ShaderPath;
+
+	// add the entry point reference
+	FXCCommandline += FString(TEXT(" /E ")) + EntryFunction;
+
+	// @TODO - currently d3d11 uses d3d10 shader compiler flags... update when this changes in DXSDK
+	// go through and add other switches
+	if(CompileFlags & D3D10_SHADER_PREFER_FLOW_CONTROL)
+	{
+		CompileFlags &= ~D3D10_SHADER_PREFER_FLOW_CONTROL;
+		FXCCommandline += FString(TEXT(" /Gfp"));
+	}
+
+	if(CompileFlags & D3D10_SHADER_DEBUG)
+	{
+		CompileFlags &= ~D3D10_SHADER_DEBUG;
+		FXCCommandline += FString(TEXT(" /Zi"));
+	}
+
+	if(CompileFlags & D3D10_SHADER_SKIP_OPTIMIZATION)
+	{
+		CompileFlags &= ~D3D10_SHADER_SKIP_OPTIMIZATION;
+		FXCCommandline += FString(TEXT(" /Od"));
+	}
+
+	if (CompileFlags & D3D10_SHADER_SKIP_VALIDATION)
+	{
+		CompileFlags &= ~D3D10_SHADER_SKIP_VALIDATION;
+		FXCCommandline += FString(TEXT(" /Vd"));
+	}
+
+	if(CompileFlags & D3D10_SHADER_AVOID_FLOW_CONTROL)
+	{
+		CompileFlags &= ~D3D10_SHADER_AVOID_FLOW_CONTROL;
+		FXCCommandline += FString(TEXT(" /Gfa"));
+	}
+
+	if(CompileFlags & D3D10_SHADER_PACK_MATRIX_ROW_MAJOR)
+	{
+		CompileFlags &= ~D3D10_SHADER_PACK_MATRIX_ROW_MAJOR;
+		FXCCommandline += FString(TEXT(" /Zpr"));
+	}
+
+	if(CompileFlags & D3D10_SHADER_ENABLE_BACKWARDS_COMPATIBILITY)
+	{
+		CompileFlags &= ~D3D10_SHADER_ENABLE_BACKWARDS_COMPATIBILITY;
+		FXCCommandline += FString(TEXT(" /Gec"));
+	}
+
+	switch (CompileFlags & SHADER_OPTIMIZATION_LEVEL_MASK)
+	{
+		case D3D10_SHADER_OPTIMIZATION_LEVEL2:
+		CompileFlags &= ~D3D10_SHADER_OPTIMIZATION_LEVEL2;
+		FXCCommandline += FString(TEXT(" /O2"));
+			break;
+
+		case D3D10_SHADER_OPTIMIZATION_LEVEL3:
+		CompileFlags &= ~D3D10_SHADER_OPTIMIZATION_LEVEL3;
+		FXCCommandline += FString(TEXT(" /O3"));
+			break;
+
+		case D3D10_SHADER_OPTIMIZATION_LEVEL1:
+		CompileFlags &= ~D3D10_SHADER_OPTIMIZATION_LEVEL1;
+		FXCCommandline += FString(TEXT(" /O1"));
+			break;
+
+		case D3D10_SHADER_OPTIMIZATION_LEVEL0:
+			CompileFlags &= ~D3D10_SHADER_OPTIMIZATION_LEVEL0;
+			break;
+
+		default:
+			Output.Errors.Emplace(TEXT("Unknown D3D10 optimization level"));
+			break;
+	}
+
+	checkf(CompileFlags == 0, TEXT("Unhandled d3d11 shader compiler flag!"));
+
+	// add the target instruction set
+	FXCCommandline += FString(TEXT(" /T ")) + ShaderProfile;
+
+	// Assembly instruction numbering
+	FXCCommandline += TEXT(" /Ni");
+
+	// Output to ShaderPath.d3dasm
+	if (FPaths::GetExtension(ShaderPath) == TEXT("usf"))
+	{
+		FXCCommandline += FString::Printf(TEXT(" /Fc%sd3dasm"), *ShaderPath.LeftChop(3));
+	}
+
+	// add a pause on a newline
+	FXCCommandline += FString(TEXT(" \r\n pause"));
+
+	// Batch file header:
+	/*
+	@ECHO OFF
+		SET FXC="C:\Program Files (x86)\Windows Kits\10\bin\x64\fxc.exe"
+		IF EXIST %FXC% (
+			REM
+			) ELSE (
+				ECHO Couldn't find Windows 10 SDK, falling back to DXSDK...
+				SET FXC="%DXSDK_DIR%\Utilities\bin\x86\fxc.exe"
+				IF EXIST %FXC% (
+					REM
+					) ELSE (
+						ECHO Couldn't find DXSDK! Exiting...
+						GOTO END
+					)
+			)
+	*/
+	FString BatchFileHeader = TEXT("@ECHO OFF\nSET FXC=\"C:\\Program Files (x86)\\Windows Kits\\10\\bin\\x64\\fxc.exe\"\n"\
+		"IF EXIST %FXC% (\nREM\n) ELSE (\nECHO Couldn't find Windows 10 SDK, falling back to DXSDK...\n"\
+		"SET FXC=\"%DXSDK_DIR%\\Utilities\\bin\\x86\\fxc.exe\"\nIF EXIST %FXC% (\nREM\n) ELSE (\nECHO Couldn't find DXSDK! Exiting...\n"\
+		"GOTO END\n)\n)\n");
+	return BatchFileHeader + FXCCommandline + TEXT("\n:END\nREM\n");
+}
+
+/** Creates a batch file string to call the AMD shader analyzer. */
+static FString CreateAMDCodeXLCommandLine(
+	const FString& ShaderPath, 
+	const TCHAR* EntryFunction, 
+	const TCHAR* ShaderProfile,
+	uint32 DXFlags
+	)
+{
+	// Hardcoded to the default install path since there's no Env variable or addition to PATH
+	FString Commandline = FString(TEXT("\"C:\\Program Files (x86)\\AMD\\CodeXL\\CodeXLAnalyzer.exe\" -c Pitcairn")) 
+		+ TEXT(" -f ") + EntryFunction
+		+ TEXT(" -s HLSL")
+		+ TEXT(" -p ") + ShaderProfile
+		+ TEXT(" -a AnalyzerStats.csv")
+		+ TEXT(" --isa ISA.txt")
+		+ *FString::Printf(TEXT(" --DXFlags %u "), DXFlags)
+		+ ShaderPath;
+
+	// add a pause on a newline
+	Commandline += FString(TEXT(" \r\n pause"));
+	return Commandline;
+}
+
+// D3Dcompiler.h has function pointer typedefs for some functions, but not all
+typedef HRESULT(WINAPI *pD3DReflect)
+	(__in_bcount(SrcDataSize) LPCVOID pSrcData,
+	 __in SIZE_T  SrcDataSize,
+	 __in  REFIID pInterface,
+	 __out void** ppReflector);
+
+typedef HRESULT(WINAPI *pD3DStripShader)
+	(__in_bcount(BytecodeLength) LPCVOID pShaderBytecode,
+	 __in SIZE_T     BytecodeLength,
+	 __in UINT       uStripFlags,
+	__out ID3DBlob** ppStrippedBlob);
+
+#define DEFINE_GUID_FOR_CURRENT_COMPILER(name, l, w1, w2, b1, b2, b3, b4, b5, b6, b7, b8) \
+	static const GUID name = { l, w1, w2, { b1, b2,  b3,  b4,  b5,  b6,  b7,  b8 } }
+
+// ShaderReflection IIDs may change between SDK versions if the reflection API changes.
+// Define a GUID below that matches the desired IID for the DLL in CompilerPath. For example,
+// look for IID_ID3D11ShaderReflection in d3d11shader.h for the SDK matching the compiler DLL.
+DEFINE_GUID_FOR_CURRENT_COMPILER(IID_ID3D11ShaderReflectionForCurrentCompiler, 0x8d536ca1, 0x0cca, 0x4956, 0xa8, 0x37, 0x78, 0x69, 0x63, 0x75, 0x55, 0x84);
+
+/**
+ * GetD3DCompilerFuncs - gets function pointers from the dll at NewCompilerPath
+ * @param OutD3DCompile - function pointer for D3DCompile (0 if not found)
+ * @param OutD3DReflect - function pointer for D3DReflect (0 if not found)
+ * @param OutD3DDisassemble - function pointer for D3DDisassemble (0 if not found)
+ * @param OutD3DStripShader - function pointer for D3DStripShader (0 if not found)
+ * @return bool - true if functions were retrieved from NewCompilerPath
+ */
+static bool GetD3DCompilerFuncs(const FString& NewCompilerPath, pD3DCompile* OutD3DCompile,
+	pD3DReflect* OutD3DReflect, pD3DDisassemble* OutD3DDisassemble, pD3DStripShader* OutD3DStripShader)
+{
+	static FString CurrentCompiler;
+	static HMODULE CompilerDLL = 0;
+
+	if(CurrentCompiler != *NewCompilerPath)
+	{
+		CurrentCompiler = *NewCompilerPath;
+
+		if(CompilerDLL)
+		{
+			FreeLibrary(CompilerDLL);
+			CompilerDLL = 0;
+		}
+
+		if(CurrentCompiler.Len())
+		{
+			CompilerDLL = LoadLibrary(*CurrentCompiler);
+		}
+
+		if(!CompilerDLL && NewCompilerPath.Len())
+		{
+			// Couldn't find HLSL compiler in specified path. We fail the first compile.
+			*OutD3DCompile = 0;
+			*OutD3DReflect = 0;
+			*OutD3DDisassemble = 0;
+			*OutD3DStripShader = 0;
+			return false;
+		}
+	}
+
+	if(CompilerDLL)
+	{
+		// from custom folder e.g. "C:/DXWin8/D3DCompiler_44.dll"
+		*OutD3DCompile = (pD3DCompile)(void*)GetProcAddress(CompilerDLL, "D3DCompile");
+		*OutD3DReflect = (pD3DReflect)(void*)GetProcAddress(CompilerDLL, "D3DReflect");
+		*OutD3DDisassemble = (pD3DDisassemble)(void*)GetProcAddress(CompilerDLL, "D3DDisassemble");
+		*OutD3DStripShader = (pD3DStripShader)(void*)GetProcAddress(CompilerDLL, "D3DStripShader");
+		return true;
+	}
+
+	// D3D SDK we compiled with (usually D3DCompiler_43.dll from windows folder)
+	*OutD3DCompile = &D3DCompile;
+	*OutD3DReflect = &D3DReflect;
+	*OutD3DDisassemble = &D3DDisassemble;
+	*OutD3DStripShader = &D3DStripShader;
+	return false;
+}
+
+static HRESULT D3DCompileWrapper(
+	pD3DCompile				D3DCompileFunc,
+	bool&					bException,
+	LPCVOID					pSrcData,
+	SIZE_T					SrcDataSize,
+	LPCSTR					pFileName,
+	CONST D3D_SHADER_MACRO*	pDefines,
+	ID3DInclude*			pInclude,
+	LPCSTR					pEntrypoint,
+	LPCSTR					pTarget,
+	uint32					Flags1,
+	uint32					Flags2,
+	ID3DBlob**				ppCode,
+	ID3DBlob**				ppErrorMsgs
+	)
+{
+#if !PLATFORM_SEH_EXCEPTIONS_DISABLED
+	__try
+#endif
+	{
+		return D3DCompileFunc(
+			pSrcData,
+			SrcDataSize,
+			pFileName,
+			pDefines,
+			pInclude,
+			pEntrypoint,
+			pTarget,
+			Flags1,
+			Flags2,
+			ppCode,
+			ppErrorMsgs
+		);
+	}
+#if !PLATFORM_SEH_EXCEPTIONS_DISABLED
+	__except( EXCEPTION_EXECUTE_HANDLER )
+	{
+		bException = true;
+		return E_FAIL;
+	}
+#endif
+}
+
+// Utility variable so we can place a breakpoint while debugging
+static int32 GBreakpoint = 0;
+
+#if RHI_RAYTRACING
+
+#define VERIFYHRESULT(expr) { HRESULT HR##__LINE__ = expr; if (FAILED(HR##__LINE__)) { UE_LOG(LogD3D11ShaderCompiler, Fatal, TEXT(#expr " failed: Result=%08x"), HR##__LINE__); } }
+
+static dxc::DxcDllSupport& GetDxcDllHelper()
+{
+	static dxc::DxcDllSupport DxcDllSupport;
+	static bool DxcDllInitialized = false;
+	if (!DxcDllInitialized)
+	{
+		VERIFYHRESULT(DxcDllSupport.Initialize());
+		DxcDllInitialized = true;
+	}
+	return DxcDllSupport;
+}
+
+static HRESULT D3DCompileToDxil(const char* SourceText, LPCWSTR EntryPoint, LPCWSTR TargetProfile, LPCWSTR* Arguments, uint32 NumArguments,
+	TRefCountPtr<ID3DBlob>& OutDxilBlob, TRefCountPtr<IDxcBlobEncoding>& OutErrorBlob)
+{
+	dxc::DxcDllSupport& DxcDllHelper = GetDxcDllHelper();
+
+	TRefCountPtr<IDxcCompiler> Compiler;
+	VERIFYHRESULT(DxcDllHelper.CreateInstance(CLSID_DxcCompiler, Compiler.GetInitReference()));
+
+	TRefCountPtr<IDxcLibrary> Library;
+	VERIFYHRESULT(DxcDllHelper.CreateInstance(CLSID_DxcLibrary, Library.GetInitReference()));
+
+	TRefCountPtr<IDxcBlobEncoding> TextBlob;
+	VERIFYHRESULT(Library->CreateBlobWithEncodingFromPinned((LPBYTE)SourceText, FCStringAnsi::Strlen(SourceText), CP_UTF8, TextBlob.GetInitReference()));
+
+	TRefCountPtr<IDxcOperationResult> CompileResult;
+
+	VERIFYHRESULT(Compiler->Compile(
+		TextBlob,							// source text to compile
+		nullptr,							// optional file name for pSource. Used in errors and include handlers.
+		EntryPoint,							// entry point name
+		TargetProfile,						// shader profile to compile
+		Arguments,							// array of pointers to arguments
+		NumArguments,						// number of arguments
+		nullptr,							// array of defines
+		0,									// number of defines
+		nullptr,							// user-provided interface to handle #include directives (optional)
+		CompileResult.GetInitReference()	// compiler output status, buffer, and errors
+	));
+
+	HRESULT CompileResultCode;
+	CompileResult->GetStatus(&CompileResultCode);
+
+	if (SUCCEEDED(CompileResultCode))
+	{
+		// NOTE: IDxcBlob is an alias of ID3D10Blob and ID3DBlob.
+		VERIFYHRESULT(CompileResult->GetResult((IDxcBlob**)OutDxilBlob.GetInitReference()));
+	}
+
+	CompileResult->GetErrorBuffer(OutErrorBlob.GetInitReference());
+
+	return CompileResultCode;
+}
+
+static void D3DCreateDXCArguments(TArray<const WCHAR*>& OutArgs, const WCHAR* Exports, uint32 CompileFlags, FShaderCompilerOutput& Output, uint32 AutoBindingSpace = ~0u)
+{
+	// Static digit strings are used here as they are returned in OutArgs
+	static const WCHAR* DigitStrings[] = 
+	{
+		L"0", L"1", L"2", L"3", L"4", L"5", L"6", L"7", L"8", L"9"
+	};
+
+	if (AutoBindingSpace < ARRAY_COUNT(DigitStrings))
+	{
+		OutArgs.Add(L"/auto-binding-space");
+		OutArgs.Add(DigitStrings[AutoBindingSpace]);
+	}
+	else if (AutoBindingSpace != ~0u)
+	{
+		UE_LOG(LogD3D11ShaderCompiler, Fatal, TEXT("Unsupported register binding space %d"), AutoBindingSpace);
+	}
+
+	// Ensure that only the requested functions exists in the output DXIL.
+	// All other functions and their used resources must be eliminated.
+	OutArgs.Add(L"/exports");
+	OutArgs.Add(Exports);
+
+	if (CompileFlags & D3D10_SHADER_PREFER_FLOW_CONTROL)
+	{
+		CompileFlags &= ~D3D10_SHADER_PREFER_FLOW_CONTROL;
+		OutArgs.Add(L"/Gfp");
+	}
+
+	if (CompileFlags & D3D10_SHADER_DEBUG)
+	{
+		CompileFlags &= ~D3D10_SHADER_DEBUG;
+		OutArgs.Add(L"/Zi");
+	}
+
+	if (CompileFlags & D3D10_SHADER_SKIP_OPTIMIZATION)
+	{
+		CompileFlags &= ~D3D10_SHADER_SKIP_OPTIMIZATION;
+		OutArgs.Add(L"/Od");
+	}
+
+	if (CompileFlags & D3D10_SHADER_SKIP_VALIDATION)
+	{
+		CompileFlags &= ~D3D10_SHADER_SKIP_VALIDATION;
+		OutArgs.Add(L"/Vd");
+	}
+
+	if (CompileFlags & D3D10_SHADER_AVOID_FLOW_CONTROL)
+	{
+		CompileFlags &= ~D3D10_SHADER_AVOID_FLOW_CONTROL;
+		OutArgs.Add(L"/Gfa");
+	}
+
+	if (CompileFlags & D3D10_SHADER_PACK_MATRIX_ROW_MAJOR)
+	{
+		CompileFlags &= ~D3D10_SHADER_PACK_MATRIX_ROW_MAJOR;
+		OutArgs.Add(L"/Zpr");
+	}
+
+	if (CompileFlags & D3D10_SHADER_ENABLE_BACKWARDS_COMPATIBILITY)
+	{
+		CompileFlags &= ~D3D10_SHADER_ENABLE_BACKWARDS_COMPATIBILITY;
+		OutArgs.Add(L"/Gec");
+	}
+
+	switch (CompileFlags & SHADER_OPTIMIZATION_LEVEL_MASK)
+	{
+	case D3D10_SHADER_OPTIMIZATION_LEVEL0:
+		CompileFlags &= ~D3D10_SHADER_OPTIMIZATION_LEVEL0;
+		OutArgs.Add(L"/O0");
+		break;
+
+	case D3D10_SHADER_OPTIMIZATION_LEVEL1:
+		CompileFlags &= ~D3D10_SHADER_OPTIMIZATION_LEVEL1;
+		OutArgs.Add(L"/O1");
+		break;
+
+	case D3D10_SHADER_OPTIMIZATION_LEVEL2:
+		CompileFlags &= ~D3D10_SHADER_OPTIMIZATION_LEVEL2;
+		OutArgs.Add(L"/O2");
+		break;
+
+	case D3D10_SHADER_OPTIMIZATION_LEVEL3:
+		CompileFlags &= ~D3D10_SHADER_OPTIMIZATION_LEVEL3;
+		OutArgs.Add(L"/O3");
+		break;
+
+	default:
+		Output.Errors.Emplace(TEXT("Unknown optimization level flag"));
+		break;
+	}
+
+	checkf(CompileFlags == 0, TEXT("Unhandled shader compiler flag!"));
+}
+
+static FString D3DCreateDXCCompileBatchFile(const FString& ShaderPath, const WCHAR* Exports, const TCHAR* ShaderProfile, uint32 CompileFlags, FShaderCompilerOutput& CompilerOutput, uint32 AutoBindingSpace = ~0u)
+{
+	TArray<const WCHAR*> Args;
+	const uint32 DXCFlags = CompileFlags & (~D3D10_SHADER_ENABLE_BACKWARDS_COMPATIBILITY);
+	D3DCreateDXCArguments(Args, Exports, DXCFlags, CompilerOutput, AutoBindingSpace);
+
+	FString BatchFileHeader = TEXT("@ECHO OFF\nSET DXC=\"C:\\Program Files (x86)\\Windows Kits\\10\\bin\\10.0.17763.0\\x64\\dxc.exe\"\n"\
+		"IF EXIST %DXC% (\nREM\n) ELSE (\nECHO Couldn't find Windows 10.0.17763 SDK, falling back to dxc.exe in PATH...\n"\
+		"SET DXC=dxc.exe)\n");
+
+	FString DXCCommandline = FString(TEXT("%DXC%"));
+	for (const WCHAR* Arg : Args)
+	{
+		DXCCommandline += TEXT(" ");
+		DXCCommandline += Arg;
+	}
+
+	DXCCommandline += TEXT(" /T ");
+	DXCCommandline += ShaderProfile;
+
+	if (FPaths::GetExtension(ShaderPath) == TEXT("usf"))
+	{
+		DXCCommandline += FString::Printf(TEXT(" /Fc%sd3dasm"), *ShaderPath.LeftChop(3));
+	}
+
+	DXCCommandline += TEXT(" ");
+	DXCCommandline += ShaderPath;
+
+	return BatchFileHeader + DXCCommandline + TEXT("\npause");
+}
+
+#ifndef DXIL_FOURCC
+#define DXIL_FOURCC(ch0, ch1, ch2, ch3) (                            \
+  (uint32_t)(uint8_t)(ch0)        | (uint32_t)(uint8_t)(ch1) << 8  | \
+  (uint32_t)(uint8_t)(ch2) << 16  | (uint32_t)(uint8_t)(ch3) << 24   \
+  )
+#endif
+
+static HRESULT D3DCreateLibraryReflectionFromBlob(ID3DBlob* DxilBlob, TRefCountPtr<ID3D12LibraryReflection>& OutReflection)
+{
+	dxc::DxcDllSupport& DxcDllHelper = GetDxcDllHelper();
+
+	TRefCountPtr<IDxcContainerReflection> ContainerReflection;
+	VERIFYHRESULT(DxcDllHelper.CreateInstance(CLSID_DxcContainerReflection, ContainerReflection.GetInitReference()));
+	VERIFYHRESULT(ContainerReflection->Load((IDxcBlob *)DxilBlob));
+
+	const uint32 DxilPartKind = DXIL_FOURCC('D', 'X', 'I', 'L');
+	uint32 DxilPartIndex = ~0u;
+	VERIFYHRESULT(ContainerReflection->FindFirstPartKind(DxilPartKind, &DxilPartIndex));
+
+	HRESULT Result = ContainerReflection->GetPartReflection(DxilPartIndex, IID_PPV_ARGS(OutReflection.GetInitReference()));
+
+	return Result;
+}
+
+#undef VERIFYHRESULT
+
+inline bool IsCompatibleBinding(const D3D12_SHADER_INPUT_BIND_DESC& BindDesc, uint32 BindingSpace)
+{
+	return BindDesc.Space == BindingSpace;
+}
+
+#endif // RHI_RAYTRACING
+
+inline bool IsCompatibleBinding(const D3D11_SHADER_INPUT_BIND_DESC& BindDesc, uint32 BindingSpace)
+{
+	return true;
+}
+
+template <typename ID3D1xShaderReflection, typename D3D1x_SHADER_DESC, typename D3D1x_SHADER_INPUT_BIND_DESC,
+	typename ID3D1xShaderReflectionConstantBuffer, typename D3D1x_SHADER_BUFFER_DESC,
+	typename ID3D1xShaderReflectionVariable, typename D3D1x_SHADER_VARIABLE_DESC>
+static void ExtractParameterMapFromD3DShader(
+	uint32 TargetPlatform, uint32 BindingSpace, const FString& VirtualSourceFilePath, ID3D1xShaderReflection* Reflector, const D3D1x_SHADER_DESC& ShaderDesc,
+	bool& bGlobalUniformBufferUsed, uint32& NumSamplers, uint32& NumSRVs, uint32& NumCBs, uint32& NumUAVs,
+	FShaderCompilerOutput& Output, TArray<FString>& UniformBufferNames, TBitArray<>& UsedUniformBufferSlots
+)
+{
+	// Add parameters for shader resources (constant buffers, textures, samplers, etc. */
+	for (uint32 ResourceIndex = 0; ResourceIndex < ShaderDesc.BoundResources; ResourceIndex++)
+	{
+		D3D1x_SHADER_INPUT_BIND_DESC BindDesc;
+		Reflector->GetResourceBindingDesc(ResourceIndex, &BindDesc);
+
+		if (!IsCompatibleBinding(BindDesc, BindingSpace))
+		{
+			continue;
+		}
+
+		if (BindDesc.Type == D3D10_SIT_CBUFFER || BindDesc.Type == D3D10_SIT_TBUFFER)
+		{
+			const uint32 CBIndex = BindDesc.BindPoint;
+			ID3D1xShaderReflectionConstantBuffer* ConstantBuffer = Reflector->GetConstantBufferByName(BindDesc.Name);
+			D3D1x_SHADER_BUFFER_DESC CBDesc;
+			ConstantBuffer->GetDesc(&CBDesc);
+			bool bGlobalCB = (FCStringAnsi::Strcmp(CBDesc.Name, "$Globals") == 0);
+
+			if (bGlobalCB)
+			{
+				// Track all of the variables in this constant buffer.
+				for (uint32 ConstantIndex = 0; ConstantIndex < CBDesc.Variables; ConstantIndex++)
+				{
+					ID3D1xShaderReflectionVariable* Variable = ConstantBuffer->GetVariableByIndex(ConstantIndex);
+					D3D1x_SHADER_VARIABLE_DESC VariableDesc;
+					Variable->GetDesc(&VariableDesc);
+					if (VariableDesc.uFlags & D3D10_SVF_USED)
+					{
+						bGlobalUniformBufferUsed = true;
+
+						Output.ParameterMap.AddParameterAllocation(
+							ANSI_TO_TCHAR(VariableDesc.Name),
+							CBIndex,
+							VariableDesc.StartOffset,
+							VariableDesc.Size,
+							EShaderParameterType::LooseData
+						);
+						UsedUniformBufferSlots[CBIndex] = true;
+					}
+				}
+			}
+			else
+			{
+				// Track just the constant buffer itself.
+				Output.ParameterMap.AddParameterAllocation(
+					ANSI_TO_TCHAR(CBDesc.Name),
+					CBIndex,
+					0,
+					0,
+					EShaderParameterType::UniformBuffer
+				);
+				UsedUniformBufferSlots[CBIndex] = true;
+
+				if (UniformBufferNames.Num() <= (int32)CBIndex)
+				{
+					UniformBufferNames.AddDefaulted(CBIndex - UniformBufferNames.Num() + 1);
+				}
+				UniformBufferNames[CBIndex] = CBDesc.Name;
+			}
+
+			NumCBs = FMath::Max(NumCBs, BindDesc.BindPoint + BindDesc.BindCount);
+		}
+		else if (BindDesc.Type == D3D10_SIT_TEXTURE || BindDesc.Type == D3D10_SIT_SAMPLER)
+		{
+			TCHAR OfficialName[1024];
+			uint32 BindCount = BindDesc.BindCount;
+			FCString::Strcpy(OfficialName, ANSI_TO_TCHAR(BindDesc.Name));
+
+			if (TargetPlatform == SP_PCD3D_SM5)
+			{
+				// Assign the name and optionally strip any "[#]" suffixes
+				TCHAR *BracketLocation = FCString::Strchr(OfficialName, TEXT('['));
+				if (BracketLocation)
+				{
+					BindCount = 1;
+
+					// This needs to include the first [ character otherwise it will include non array textures with matching starting characters.
+					// e.g. "LightMapTexturesTest" which is not part of "LightMapTextures[#]" would be included as the last index of "LightMapTextures"
+					const int32 NumCharactersToCompare = BracketLocation - OfficialName + 1;
+
+					// In SM5, for some reason, array suffixes are included in Name, i.e. "LightMapTextures[0]", rather than "LightMapTextures"
+					// Additionally elements in an array are listed as SEPERATE bound resources.
+					// However, they are always contiguous in resource index, so iterate over the samplers and textures of the initial association
+					// and count them, identifying the bindpoint and bindcounts
+
+					while (ResourceIndex + 1 < ShaderDesc.BoundResources)
+					{
+						D3D1x_SHADER_INPUT_BIND_DESC BindDesc2;
+						Reflector->GetResourceBindingDesc(ResourceIndex + 1, &BindDesc2);
+
+						if (BindDesc2.Type == BindDesc.Type && FCStringAnsi::Strncmp(BindDesc2.Name, BindDesc.Name, NumCharactersToCompare) == 0)
+						{
+							BindCount++;
+							// Skip over this resource since it is part of an array
+							ResourceIndex++;
+						}
+						else
+						{
+							break;
+						}
+					}
+
+					// Array loop is complete, now the array suffix can be removed.
+					*BracketLocation = 0;
+				}
+			}
+
+			EShaderParameterType ParameterType = EShaderParameterType::Num;
+			if (BindDesc.Type == D3D10_SIT_SAMPLER)
+			{
+				ParameterType = EShaderParameterType::Sampler;
+				NumSamplers = FMath::Max(NumSamplers, BindDesc.BindPoint + BindCount);
+			}
+			else if (BindDesc.Type == D3D10_SIT_TEXTURE)
+			{
+				ParameterType = EShaderParameterType::SRV;
+				NumSRVs = FMath::Max(NumSRVs, BindDesc.BindPoint + BindCount);
+			}
+
+			// Add a parameter for the texture only, the sampler index will be invalid
+			Output.ParameterMap.AddParameterAllocation(
+				OfficialName,
+				0,
+				BindDesc.BindPoint,
+				BindCount,
+				ParameterType
+			);
+		}
+		else if (BindDesc.Type == D3D11_SIT_UAV_RWTYPED || BindDesc.Type == D3D11_SIT_UAV_RWSTRUCTURED ||
+			BindDesc.Type == D3D11_SIT_UAV_RWBYTEADDRESS || BindDesc.Type == D3D11_SIT_UAV_RWSTRUCTURED_WITH_COUNTER ||
+			BindDesc.Type == D3D11_SIT_UAV_APPEND_STRUCTURED)
+		{
+			TCHAR OfficialName[1024];
+			FCString::Strcpy(OfficialName, ANSI_TO_TCHAR(BindDesc.Name));
+
+			Output.ParameterMap.AddParameterAllocation(
+				OfficialName,
+				0,
+				BindDesc.BindPoint,
+				1,
+				EShaderParameterType::UAV
+			);
+
+			NumUAVs = FMath::Max(NumUAVs, BindDesc.BindPoint + BindDesc.BindCount);
+		}
+		else if (BindDesc.Type == D3D11_SIT_STRUCTURED || BindDesc.Type == D3D11_SIT_BYTEADDRESS)
+		{
+			TCHAR OfficialName[1024];
+			FCString::Strcpy(OfficialName, ANSI_TO_TCHAR(BindDesc.Name));
+
+			Output.ParameterMap.AddParameterAllocation(
+				OfficialName,
+				0,
+				BindDesc.BindPoint,
+				1,
+				EShaderParameterType::SRV
+			);
+
+			NumSRVs = FMath::Max(NumSRVs, BindDesc.BindPoint + BindDesc.BindCount);
+		}
+#if RHI_RAYTRACING
+		// #dxr_todo: D3D_SIT_RTACCELERATIONSTRUCTURE is declared in latest version of dxcapi.h. Update this code after upgrading DXC.
+		else if (BindDesc.Type == 12 /*D3D_SIT_RTACCELERATIONSTRUCTURE*/)
+		{
+			// Acceleration structure resources are treated as SRVs.
+
+			TCHAR OfficialName[1024];
+			FCString::Strcpy(OfficialName, ANSI_TO_TCHAR(BindDesc.Name));
+
+			Output.ParameterMap.AddParameterAllocation(
+				OfficialName,
+				0,
+				BindDesc.BindPoint,
+				1,
+				EShaderParameterType::SRV
+			);
+
+			NumSRVs = FMath::Max(NumSRVs, BindDesc.BindPoint + BindDesc.BindCount);
+		}
+#endif
+	}
+}
+
+#if RHI_RAYTRACING
+// Parses ray tracing shader entry point specification string in one of the following formats:
+// 1) Verbatim single entry point name, e.g. "MainRGS"
+// 2) Complex entry point for ray tracing hit group shaders:
+//      a) "closesthit=MainCHS"
+//      b) "closesthit=MainCHS anyhit=MainAHS"
+//      c) "closesthit=MainCHS anyhit=MainAHS intersection=MainIS"
+//      d) "closesthit=MainCHS intersection=MainIS"
+//    NOTE: closesthit attribute must always be provided for complex hit group entry points
+static void ParseRayTracingEntryPoint(const FString& Input, FString& OutMain, FString& OutAnyHit, FString& OutIntersection)
+{
+	auto ParseEntry = [&Input](const TCHAR* Marker)
+	{
+		FString Result;
+		int32 BeginIndex = Input.Find(Marker, ESearchCase::IgnoreCase, ESearchDir::FromStart);
+		if (BeginIndex != INDEX_NONE)
+		{
+			int32 EndIndex = Input.Find(TEXT(" "), ESearchCase::IgnoreCase, ESearchDir::FromStart, BeginIndex);
+			if (EndIndex == INDEX_NONE) EndIndex = Input.Len() + 1;
+			int32 MarkerLen = FCString::Strlen(Marker);
+			int32 Count = EndIndex - BeginIndex;
+			Result = Input.Mid(BeginIndex + MarkerLen, Count - MarkerLen);
+		}
+		return Result;
+	};
+
+	OutMain = ParseEntry(TEXT("closesthit="));
+	OutAnyHit = ParseEntry(TEXT("anyhit="));
+	OutIntersection = ParseEntry(TEXT("intersection="));
+
+	// If complex hit group entry is not specified, assume a single verbatim entry point
+	if (OutMain.IsEmpty() && OutAnyHit.IsEmpty() && OutIntersection.IsEmpty())
+	{
+		OutMain = Input;
+	}
+}
+#endif // RHI_RAYTRACING
+
+// Generate the dumped usf file; call the D3D compiler, gather reflection information and generate the output data
+static bool CompileAndProcessD3DShader(FString& PreprocessedShaderSource, const FString& CompilerPath,
+	uint32 CompileFlags, const FShaderCompilerInput& Input, FString& EntryPointName,
+	const TCHAR* ShaderProfile, bool bProcessingSecondTime,
+	TArray<FString>& FilteredErrors, FShaderCompilerOutput& Output)
+{
+	auto AnsiSourceFile = StringCast<ANSICHAR>(*PreprocessedShaderSource);
+
+	const bool bIsRayTracingShader = IsRayTracingShader(Input.Target);
+
+#if RHI_RAYTRACING
+
+	const uint32 AutoBindingSpace = GetAutoBindingSpace(Input.Target);
+
+	FString RayEntryPoint; // Primary entry point for all ray tracing shaders
+	FString RayAnyHitEntryPoint; // Optional for hit group shaders
+	FString RayIntersectionEntryPoint; // Optional for hit group shaders
+	FString RayTracingExports;
+
+	if (bIsRayTracingShader)
+	{
+		ParseRayTracingEntryPoint(Input.EntryPointName, RayEntryPoint, RayAnyHitEntryPoint, RayIntersectionEntryPoint);
+
+		RayTracingExports = RayEntryPoint;
+
+		if (!RayAnyHitEntryPoint.IsEmpty())
+		{
+			RayTracingExports += TEXT(";");
+			RayTracingExports += RayAnyHitEntryPoint;
+		}
+
+		if (!RayIntersectionEntryPoint.IsEmpty())
+		{
+			RayTracingExports += TEXT(";");
+			RayTracingExports += RayIntersectionEntryPoint;
+		}
+	}
+
+#endif // RHI_RAYTRACING
+
+	bool bDumpDebugInfo = false;
+	// Write out the preprocessed file and a batch file to compile it if requested (DumpDebugInfoPath is valid)
+	if (Input.DumpDebugInfoPath.Len() > 0 && IFileManager::Get().DirectoryExists(*Input.DumpDebugInfoPath))
+	{
+		bDumpDebugInfo = true;
+		FString Filename = Input.GetSourceFilename();
+		FArchive* FileWriter = IFileManager::Get().CreateFileWriter(*(Input.DumpDebugInfoPath / Filename));
+		if (FileWriter)
+		{
+			FileWriter->Serialize((ANSICHAR*)AnsiSourceFile.Get(), AnsiSourceFile.Length());
+			{
+				FString Line = CrossCompiler::CreateResourceTableFromEnvironment(Input.Environment);
+
+				Line += TEXT("#if 0 /*DIRECT COMPILE*/\n");
+				Line += CreateShaderCompilerWorkerDirectCommandLine(Input);
+				Line += TEXT("\n#endif /*DIRECT COMPILE*/\n");
+
+				FileWriter->Serialize(TCHAR_TO_ANSI(*Line), Line.Len());
+			}
+			FileWriter->Close();
+			delete FileWriter;
+		}
+
+		FString BatchFileContents;
+
+		if (bIsRayTracingShader)
+		{
+#if RHI_RAYTRACING
+			BatchFileContents = D3DCreateDXCCompileBatchFile(Filename, *RayTracingExports, ShaderProfile, CompileFlags, Output, AutoBindingSpace);
+#endif // RHI_RAYTRACING
+		}
+		else
+		{
+			BatchFileContents = D3D11CreateShaderCompileCommandLine(Filename, *EntryPointName, ShaderProfile, CompileFlags, Output);
+
+			if (GD3DDumpAMDCodeXLFile)
+			{
+				const FString BatchFileContents2 = CreateAMDCodeXLCommandLine(Filename, *EntryPointName, ShaderProfile, CompileFlags);
+				FFileHelper::SaveStringToFile(BatchFileContents2, *(Input.DumpDebugInfoPath / TEXT("CompileAMD.bat")));
+			}
+		}
+
+		FFileHelper::SaveStringToFile(BatchFileContents, *(Input.DumpDebugInfoPath / TEXT("CompileD3D.bat")));
+
+		if (Input.bGenerateDirectCompileFile)
+		{
+			FFileHelper::SaveStringToFile(CreateShaderCompilerWorkerDirectCommandLine(Input), *(Input.DumpDebugInfoPath / TEXT("DirectCompile.txt")));
+		}
+	}
+
+	TRefCountPtr<ID3DBlob> Shader;
+
+	HRESULT Result = S_OK;
+	pD3DCompile D3DCompileFunc = nullptr;
+	pD3DReflect D3DReflectFunc = nullptr;
+	pD3DDisassemble D3DDisassembleFunc = nullptr;
+	pD3DStripShader D3DStripShaderFunc = nullptr;
+	bool bCompilerPathFunctionsUsed = false;
+
+	if (bIsRayTracingShader)
+	{
+#if RHI_RAYTRACING
+
+		TArray<const WCHAR*> Args;
+
+		// Ignore backwards compatibility flag (/Gec) as it is deprecated.
+		// #dxr_todo: this flag should not be even passed into this function from the higher level.
+		const uint32 DXCFlags = CompileFlags & (~D3D10_SHADER_ENABLE_BACKWARDS_COMPATIBILITY);
+		D3DCreateDXCArguments(Args, *RayTracingExports, DXCFlags, Output, AutoBindingSpace);
+
+		TRefCountPtr<IDxcBlobEncoding> DxcErrorBlob;
+
+		Result = D3DCompileToDxil(AnsiSourceFile.Get(),
+			L"", // dummy entry point for ray tracing shaders
+			TCHAR_TO_WCHAR(ShaderProfile),
+			Args.GetData(), Args.Num(),
+			Shader, DxcErrorBlob);
+
+		if (DxcErrorBlob && DxcErrorBlob->GetBufferSize())
+		{
+			void* ErrorBuffer = DxcErrorBlob->GetBufferPointer();
+			D3D11FilterShaderCompileWarnings(ANSI_TO_TCHAR(ErrorBuffer), FilteredErrors);
+		}
+
+		if (!SUCCEEDED(Result))
+		{
+			FilteredErrors.Add(TEXT("D3DCompileToDxil failed"));
+		}
+#else
+		FilteredErrors.Add(TEXT("Attempted to compile a ray tracing shader, but ray tracing support is not enabled in this binary."));
+		Result = E_FAIL;
+		checkNoEntry();
+#endif
+	}
+	else
+	{
+		bCompilerPathFunctionsUsed = GetD3DCompilerFuncs(CompilerPath, &D3DCompileFunc, &D3DReflectFunc, &D3DDisassembleFunc, &D3DStripShaderFunc);
+		TRefCountPtr<ID3DBlob> Errors;
+
+		if (D3DCompileFunc)
+		{
+			bool bException = false;
+
+			Result = D3DCompileWrapper(
+				D3DCompileFunc,
+				bException,
+				AnsiSourceFile.Get(),
+				AnsiSourceFile.Length(),
+				TCHAR_TO_ANSI(*Input.VirtualSourceFilePath),
+				/*pDefines=*/ NULL,
+				/*pInclude=*/ NULL,
+				TCHAR_TO_ANSI(*EntryPointName),
+				TCHAR_TO_ANSI(ShaderProfile),
+				CompileFlags,
+				0,
+				Shader.GetInitReference(),
+				Errors.GetInitReference()
+			);
+
+			if (bException)
+			{
+				FilteredErrors.Add(TEXT("D3DCompile exception"));
+			}
+		}
+		else
+		{
+			FilteredErrors.Add(FString::Printf(TEXT("Couldn't find shader compiler: %s"), *CompilerPath));
+			Result = E_FAIL;
+		}
+
+		// Filter any errors.
+		void* ErrorBuffer = Errors ? Errors->GetBufferPointer() : NULL;
+		if (ErrorBuffer)
+		{
+			D3D11FilterShaderCompileWarnings(ANSI_TO_TCHAR(ErrorBuffer), FilteredErrors);
+		}
+
+		// Fail the compilation if double operations are being used, since those are not supported on all D3D11 cards
+		if (SUCCEEDED(Result))
+		{
+			if (D3DDisassembleFunc && (GD3DCheckForDoubles || bDumpDebugInfo))
+			{
+				TRefCountPtr<ID3DBlob> Dissasembly;
+				if (SUCCEEDED(D3DDisassembleFunc(Shader->GetBufferPointer(), Shader->GetBufferSize(), 0, "", Dissasembly.GetInitReference())))
+				{
+					ANSICHAR* DissasemblyString = new ANSICHAR[Dissasembly->GetBufferSize() + 1];
+					FMemory::Memcpy(DissasemblyString, Dissasembly->GetBufferPointer(), Dissasembly->GetBufferSize());
+					DissasemblyString[Dissasembly->GetBufferSize()] = 0;
+					FString DissasemblyStringW(DissasemblyString);
+					delete[] DissasemblyString;
+
+					if (bDumpDebugInfo)
+					{
+						FFileHelper::SaveStringToFile(DissasemblyStringW, *(Input.DumpDebugInfoPath / TEXT("Output.d3dasm")));
+					}
+					else if (GD3DCheckForDoubles)
+					{
+						// dcl_globalFlags will contain enableDoublePrecisionFloatOps when the shader uses doubles, even though the docs on dcl_globalFlags don't say anything about this
+						if (DissasemblyStringW.Contains(TEXT("enableDoublePrecisionFloatOps")))
+						{
+							FilteredErrors.Add(TEXT("Shader uses double precision floats, which are not supported on all D3D11 hardware!"));
+							return false;
+						}
+					}
+				}
+			}
+		}
+	}
+
+	// Gather reflection information
+	int32 NumInterpolants = 0;
+	TIndirectArray<FString> InterpolantNames;
+	TArray<FString> ShaderInputs;
+
+	if (SUCCEEDED(Result))
+	{
+		bool bGlobalUniformBufferUsed = false;
+		uint32 NumInstructions = 0;
+		uint32 NumSamplers = 0;
+		uint32 NumSRVs = 0;
+		uint32 NumCBs = 0;
+		uint32 NumUAVs = 0;
+		TArray<FString> UniformBufferNames;
+		TArray<FString> ShaderOutputs;
+
+		TBitArray<> UsedUniformBufferSlots;
+		UsedUniformBufferSlots.Init(false, 32);
+
+#if RHI_RAYTRACING
+		if (bIsRayTracingShader)
+		{
+			TRefCountPtr<ID3D12LibraryReflection> LibraryReflection;
+
+			Result = D3DCreateLibraryReflectionFromBlob(Shader, LibraryReflection);
+
+			if (FAILED(Result))
+			{
+				UE_LOG(LogD3D11ShaderCompiler, Fatal, TEXT("D3DReflectDxil failed: Result=%08x"), Result);
+			}
+
+			D3D12_LIBRARY_DESC LibraryDesc = {};
+			LibraryReflection->GetDesc(&LibraryDesc);
+
+			ID3D12FunctionReflection* FunctionReflection = nullptr;
+			D3D12_FUNCTION_DESC FunctionDesc = {};
+
+			// MangledEntryPoints contains partial mangled entry point signatures in a the following form:
+			// ?QualifiedName@ (as described here: https://en.wikipedia.org/wiki/Name_mangling)
+			// Entry point parameters are currently not included in the partial mangling.
+			TArray<FString, TInlineAllocator<3>> MangledEntryPoints;
+			if (!RayEntryPoint.IsEmpty())
+			{
+				MangledEntryPoints.Add(FString::Printf(TEXT("?%s@"), *RayEntryPoint));
+			}
+			if (!RayAnyHitEntryPoint.IsEmpty())
+			{
+				MangledEntryPoints.Add(FString::Printf(TEXT("?%s@"), *RayAnyHitEntryPoint));
+			}
+			if (!RayIntersectionEntryPoint.IsEmpty())
+			{
+				MangledEntryPoints.Add(FString::Printf(TEXT("?%s@"), *RayIntersectionEntryPoint));
+			}
+
+			uint32 NumFoundEntryPoints = 0;
+
+			for (uint32 FunctionIndex = 0; FunctionIndex < LibraryDesc.FunctionCount; ++FunctionIndex)
+			{
+				FunctionReflection = LibraryReflection->GetFunctionByIndex(FunctionIndex);
+				FunctionReflection->GetDesc(&FunctionDesc);
+
+				for (const FString& MangledEntryPoint : MangledEntryPoints)
+				{
+					// Entry point parameters are currently not included in the partial mangling, therefore partial substring match is used here.
+					if (FCStringAnsi::Strstr(FunctionDesc.Name, TCHAR_TO_ANSI(*MangledEntryPoint)))
+					{
+						// Note: calling ExtractParameterMapFromD3DShader multiple times merges the reflection data for multiple functions
+						ExtractParameterMapFromD3DShader<ID3D12FunctionReflection, D3D12_FUNCTION_DESC, D3D12_SHADER_INPUT_BIND_DESC,
+							ID3D12ShaderReflectionConstantBuffer, D3D12_SHADER_BUFFER_DESC,
+							ID3D12ShaderReflectionVariable, D3D12_SHADER_VARIABLE_DESC>(
+								Input.Target.Platform, AutoBindingSpace, Input.VirtualSourceFilePath, FunctionReflection, FunctionDesc, bGlobalUniformBufferUsed, NumSamplers, NumSRVs, NumCBs, NumUAVs,
+								Output, UniformBufferNames, UsedUniformBufferSlots);
+
+						NumFoundEntryPoints++;
+					}
+				}
+			}
+
+			if (NumFoundEntryPoints == MangledEntryPoints.Num())
+			{
+				Output.bSucceeded = true;
+
+				if (bGlobalUniformBufferUsed)
+				{
+					FString ErrorString = TEXT("Global constant buffer cannot be used in a ray tracing shader.");
+
+					uint32 NumLooseParameters = 0;
+					for (const auto& It : Output.ParameterMap.ParameterMap)
+					{
+						if (It.Value.Type == EShaderParameterType::LooseData)
+						{
+							NumLooseParameters++;
+						}
+					}
+
+					if (NumLooseParameters)
+					{
+						ErrorString += TEXT(" Global parameters: ");
+						uint32 ParameterIndex = 0;
+						for (const auto& It : Output.ParameterMap.ParameterMap)
+						{
+							if (It.Value.Type == EShaderParameterType::LooseData)
+							{
+								--NumLooseParameters;
+								ErrorString += FString::Printf(TEXT("%s%s"), *It.Key, NumLooseParameters ? TEXT(", ") : TEXT("."));
+							}
+						}
+					}
+
+					FilteredErrors.Add(ErrorString);
+					Result = E_FAIL;
+					Output.bSucceeded = false;
+				}
+			}
+			else
+			{
+				UE_LOG(LogD3D11ShaderCompiler, Fatal, TEXT("Failed to find required points in the shader library."));
+				Output.bSucceeded = false;
+			}
+		}
+		else
+#endif // RHI_RAYTRACING
+		if (D3DReflectFunc)
+		{
+			Output.bSucceeded = true;
+			ID3D11ShaderReflection* Reflector = NULL;
+
+			// IID_ID3D11ShaderReflectionForCurrentCompiler is defined in this file and needs to match the IID from the dll in CompilerPath
+			// if the function pointers from that dll are being used
+			const IID ShaderReflectionInterfaceID = bCompilerPathFunctionsUsed ? IID_ID3D11ShaderReflectionForCurrentCompiler : IID_ID3D11ShaderReflection;
+			Result = D3DReflectFunc(Shader->GetBufferPointer(), Shader->GetBufferSize(), ShaderReflectionInterfaceID, (void**)&Reflector);
+			if (FAILED(Result))
+			{
+				UE_LOG(LogD3D11ShaderCompiler, Fatal, TEXT("D3DReflect failed: Result=%08x"), Result);
+			}
+
+			// Read the constant table description.
+			D3D11_SHADER_DESC ShaderDesc;
+			Reflector->GetDesc(&ShaderDesc);
+
+			if (Input.Target.Frequency == SF_Vertex)
+			{
+				for (uint32 Index = 0; Index < ShaderDesc.OutputParameters; ++Index)
+				{
+					// VC++ horrible hack: Runtime ESP checks get confused and fail for some reason calling Reflector->GetOutputParameterDesc() (because it comes from another DLL?)
+					// so "guard it" using the middle of an array; it's been confirmed NO corruption is really happening.
+					D3D11_SIGNATURE_PARAMETER_DESC ParamDescs[3];
+					D3D11_SIGNATURE_PARAMETER_DESC& ParamDesc = ParamDescs[1];
+					Reflector->GetOutputParameterDesc(Index, &ParamDesc);
+					if (ParamDesc.SystemValueType == D3D_NAME_UNDEFINED && ParamDesc.Mask != 0)
+					{
+						++NumInterpolants;
+						new(InterpolantNames) FString(FString::Printf(TEXT("%s%d"), ANSI_TO_TCHAR(ParamDesc.SemanticName), ParamDesc.SemanticIndex));
+						ShaderOutputs.Add(*InterpolantNames.Last());
+					}
+				}
+			}
+			else if (Input.Target.Frequency == SF_Pixel)
+			{
+				if (GD3DAllowRemoveUnused != 0 && Input.bCompilingForShaderPipeline)
+				{
+					// Handy place for a breakpoint for debugging...
+					++GBreakpoint;
+				}
+
+				bool bFoundUnused = false;
+				for (uint32 Index = 0; Index < ShaderDesc.InputParameters; ++Index)
+				{
+					// VC++ horrible hack: Runtime ESP checks get confused and fail for some reason calling Reflector->GetInputParameterDesc() (because it comes from another DLL?)
+					// so "guard it" using the middle of an array; it's been confirmed NO corruption is really happening.
+					D3D11_SIGNATURE_PARAMETER_DESC ParamDescs[3];
+					D3D11_SIGNATURE_PARAMETER_DESC& ParamDesc = ParamDescs[1];
+					Reflector->GetInputParameterDesc(Index, &ParamDesc);
+					if (ParamDesc.SystemValueType == D3D_NAME_UNDEFINED)
+					{
+						if (ParamDesc.ReadWriteMask != 0)
+						{
+							FString SemanticName = ANSI_TO_TCHAR(ParamDesc.SemanticName);
+
+							ShaderInputs.AddUnique(SemanticName);
+
+							// Add the number (for the case of TEXCOORD)
+							FString SemanticIndexName = FString::Printf(TEXT("%s%d"), *SemanticName, ParamDesc.SemanticIndex);
+							ShaderInputs.AddUnique(SemanticIndexName);
+
+							// Add _centroid
+							ShaderInputs.AddUnique(SemanticName + TEXT("_centroid"));
+							ShaderInputs.AddUnique(SemanticIndexName + TEXT("_centroid"));
+						}
+						else
+						{
+							bFoundUnused = true;
+						}
+					}
+					else
+					{
+						//if (ParamDesc.ReadWriteMask != 0)
+						{
+							// Keep system values
+							ShaderInputs.AddUnique(FString(ANSI_TO_TCHAR(ParamDesc.SemanticName)));
+						}
+					}
+				}
+
+				if (GD3DAllowRemoveUnused && Input.bCompilingForShaderPipeline && bFoundUnused && !bProcessingSecondTime)
+				{
+					// Rewrite the source removing the unused inputs so the bindings will match
+					TArray<FString> RemoveErrors;
+					if (RemoveUnusedInputs(PreprocessedShaderSource, ShaderInputs, EntryPointName, RemoveErrors))
+					{
+						return CompileAndProcessD3DShader(PreprocessedShaderSource, CompilerPath, CompileFlags, Input, EntryPointName, ShaderProfile, true, FilteredErrors, Output);
+					}
+					else
+					{
+						UE_LOG(LogD3D11ShaderCompiler, Warning, TEXT("Failed to Remove unused inputs [%s]!"), *Input.DumpDebugInfoPath);
+						for (int32 Index = 0; Index < RemoveErrors.Num(); ++Index)
+						{
+							FShaderCompilerError NewError;
+							NewError.StrippedErrorMessage = RemoveErrors[Index];
+							Output.Errors.Add(NewError);
+						}
+						Output.bFailedRemovingUnused = true;
+					}
+				}
+			}
+
+			const uint32 BindingSpace = 0; // Default binding space for D3D11 shaders
+			ExtractParameterMapFromD3DShader<
+				ID3D11ShaderReflection, D3D11_SHADER_DESC, D3D11_SHADER_INPUT_BIND_DESC,
+				ID3D11ShaderReflectionConstantBuffer, D3D11_SHADER_BUFFER_DESC,
+				ID3D11ShaderReflectionVariable, D3D11_SHADER_VARIABLE_DESC>
+				(Input.Target.Platform, BindingSpace, Input.VirtualSourceFilePath, Reflector, ShaderDesc,
+					bGlobalUniformBufferUsed, NumSamplers, NumSRVs, NumCBs, NumUAVs,
+					Output, UniformBufferNames, UsedUniformBufferSlots);
+
+			NumInstructions = ShaderDesc.InstructionCount;
+
+			// Reflector is a com interface, so it needs to be released.
+			Reflector->Release();
+		}
+		else
+		{
+			FilteredErrors.Add(FString::Printf(TEXT("Couldn't find shader reflection function in %s"), *CompilerPath));
+			Result = E_FAIL;
+			Output.bSucceeded = false;
+		}
+
+		// Save results if compilation and reflection succeeded
+
+
+		if (Output.bSucceeded)
+		{
+			TRefCountPtr<ID3DBlob> CompressedData;
+
+			if (Input.Environment.CompilerFlags.Contains(CFLAG_KeepDebugInfo))
+			{
+				CompressedData = Shader;
+			}
+			else if (bIsRayTracingShader)
+			{
+				// Handy place for a breakpoint for debugging...
+				++GBreakpoint;
+
+				// #dxr_todo: strip DXIL debug and reflection data
+				CompressedData = Shader;
+			}
+			else if (D3DStripShaderFunc)
+			{
+				// Strip shader reflection and debug info
+				D3D_SHADER_DATA ShaderData;
+				ShaderData.pBytecode = Shader->GetBufferPointer();
+				ShaderData.BytecodeLength = Shader->GetBufferSize();
+				Result = D3DStripShaderFunc(Shader->GetBufferPointer(),
+					Shader->GetBufferSize(),
+					D3DCOMPILER_STRIP_REFLECTION_DATA | D3DCOMPILER_STRIP_DEBUG_INFO | D3DCOMPILER_STRIP_TEST_BLOBS,
+					CompressedData.GetInitReference());
+
+				if (FAILED(Result))
+				{
+					UE_LOG(LogD3D11ShaderCompiler, Fatal, TEXT("D3DStripShader failed: Result=%08x"), Result);
+				}
+			}
+			else
+			{
+				// D3DStripShader is not guaranteed to exist
+				// e.g. the open-source DXIL shader compiler does not currently implement it
+				CompressedData = Shader;
+			}
+
+			// Build the SRT for this shader.
+			FD3D11ShaderResourceTable SRT;
+
+			TArray<uint8> UniformBufferNameBytes;
+
+			{
+				// Build the generic SRT for this shader.
+				FShaderCompilerResourceTable GenericSRT;
+				BuildResourceTableMapping(Input.Environment.ResourceTableMap, Input.Environment.ResourceTableLayoutHashes, UsedUniformBufferSlots, Output.ParameterMap, GenericSRT);
+
+				if (UniformBufferNames.Num() < GenericSRT.ResourceTableLayoutHashes.Num())
+				{
+					UniformBufferNames.AddDefaulted(GenericSRT.ResourceTableLayoutHashes.Num() - UniformBufferNames.Num() + 1);
+				}
+
+				for (int32 Index = 0; Index < GenericSRT.ResourceTableLayoutHashes.Num(); ++Index)
+				{
+					if (GenericSRT.ResourceTableLayoutHashes[Index] != 0 && UniformBufferNames[Index].Len() == 0)
+					{
+						auto* Name = Input.Environment.ResourceTableLayoutHashes.FindKey(GenericSRT.ResourceTableLayoutHashes[Index]);
+						check(Name);
+						UniformBufferNames[Index] = *Name;
+					}
+				}
+
+				FMemoryWriter UniformBufferNameWriter(UniformBufferNameBytes);
+				UniformBufferNameWriter << UniformBufferNames;
+
+				// Copy over the bits indicating which resource tables are active.
+				SRT.ResourceTableBits = GenericSRT.ResourceTableBits;
+
+				SRT.ResourceTableLayoutHashes = GenericSRT.ResourceTableLayoutHashes;
+
+				// Now build our token streams.
+				BuildResourceTableTokenStream(GenericSRT.TextureMap, GenericSRT.MaxBoundResourceTable, SRT.TextureMap);
+				BuildResourceTableTokenStream(GenericSRT.ShaderResourceViewMap, GenericSRT.MaxBoundResourceTable, SRT.ShaderResourceViewMap);
+				BuildResourceTableTokenStream(GenericSRT.SamplerMap, GenericSRT.MaxBoundResourceTable, SRT.SamplerMap);
+				BuildResourceTableTokenStream(GenericSRT.UnorderedAccessViewMap, GenericSRT.MaxBoundResourceTable, SRT.UnorderedAccessViewMap);
+			}
+
+			if (GD3DAllowRemoveUnused != 0 && Input.Target.Frequency == SF_Pixel && Input.bCompilingForShaderPipeline && bProcessingSecondTime)
+			{
+				Output.bSupportsQueryingUsedAttributes = true;
+				if (GD3DAllowRemoveUnused == 1)
+				{
+					Output.UsedAttributes = ShaderInputs;
+				}
+			}
+
+			// Generate the final Output
+			FMemoryWriter Ar(Output.ShaderCode.GetWriteAccess(), true);
+			Ar << SRT;
+
+#if RHI_RAYTRACING
+			if (bIsRayTracingShader)
+			{
+				Ar << RayEntryPoint;
+				Ar << RayAnyHitEntryPoint;
+				Ar << RayIntersectionEntryPoint;
+			}
+#endif // RHI_RAYTRACING
+
+			Ar.Serialize(CompressedData->GetBufferPointer(), CompressedData->GetBufferSize());
+
+			// append data that is generate from the shader code and assist the usage, mostly needed for DX12 
+			{
+				FShaderCodePackedResourceCounts PackedResourceCounts = { bGlobalUniformBufferUsed, static_cast<uint8>(NumSamplers), static_cast<uint8>(NumSRVs), static_cast<uint8>(NumCBs), static_cast<uint8>(NumUAVs) };
+
+				Output.ShaderCode.AddOptionalData(PackedResourceCounts);
+				Output.ShaderCode.AddOptionalData('u', UniformBufferNameBytes.GetData(), UniformBufferNameBytes.Num());
+			}
+
+			// store data we can pickup later with ShaderCode.FindOptionalData('n'), could be removed for shipping
+			// Daniel L: This GenerateShaderName does not generate a deterministic output among shaders as the shader code can be shared. 
+			//			uncommenting this will cause the project to have non deterministic materials and will hurt patch sizes
+			//Output.ShaderCode.AddOptionalData('n', TCHAR_TO_UTF8(*Input.GenerateShaderName()));
+
+			// Set the number of instructions.
+			Output.NumInstructions = NumInstructions;
+
+			Output.NumTextureSamplers = NumSamplers;
+
+			// Pass the target through to the output.
+			Output.Target = Input.Target;
+		}
+	}
+
+	if (SUCCEEDED(Result))
+	{
+		if (Input.Target.Platform == SP_PCD3D_ES2)
+		{
+			if (Output.NumTextureSamplers > 8)
+			{
+				FilteredErrors.Add(FString::Printf(TEXT("Shader uses more than 8 texture samplers which is not supported by ES2!  Used: %u"), Output.NumTextureSamplers));
+				Result = E_FAIL;
+				Output.bSucceeded = false;
+			}
+			// Disabled for now while we work out some issues with it. A compiler bug is causing 
+			// Landscape to require a 9th interpolant even though the pixel shader never reads from
+			// it. Search for LANDSCAPE_BUG_WORKAROUND.
+			else if (false && NumInterpolants > 8)
+			{
+				FString InterpolantsStr;
+				for (int32 i = 0; i < InterpolantNames.Num(); ++i)
+				{
+					InterpolantsStr += FString::Printf(TEXT("\n\t%s"), *InterpolantNames[i]);
+				}
+				FilteredErrors.Add(FString::Printf(TEXT("Shader uses more than 8 interpolants which is not supported by ES2!  Used: %u%s"), NumInterpolants, *InterpolantsStr));
+				Result = E_FAIL;
+				Output.bSucceeded = false;
+			}
+		}
+	}
+
+	if (FAILED(Result))
+	{
+		++GBreakpoint;
+	}
+
+	return SUCCEEDED(Result);
+}
+
+void CompileD3D11Shader(const FShaderCompilerInput& Input,FShaderCompilerOutput& Output, FShaderCompilerDefinitions& AdditionalDefines, const FString& WorkingDirectory)
+{
+	FString PreprocessedShaderSource;
+	FString CompilerPath;
+	const TCHAR* ShaderProfile = GetShaderProfileName(Input.Target);
+
+	if(!ShaderProfile)
+	{
+		Output.Errors.Add(FShaderCompilerError(TEXT("Unrecognized shader frequency")));
+		return;
+	}
+
+	// Set additional defines.
+	AdditionalDefines.SetDefine(TEXT("COMPILER_HLSL"), 1);
+
+	if (Input.bSkipPreprocessedCache)
+	{
+		if (!FFileHelper::LoadFileToString(PreprocessedShaderSource, *Input.VirtualSourceFilePath))
+		{
+			return;
+		}
+
+		// Remove const as we are on debug-only mode
+		CrossCompiler::CreateEnvironmentFromResourceTable(PreprocessedShaderSource, (FShaderCompilerEnvironment&)Input.Environment);
+	}
+	else
+	{
+		if (!PreprocessShader(PreprocessedShaderSource, Output, Input, AdditionalDefines))
+		{
+			// The preprocessing stage will add any relevant errors.
+			return;
+		}
+	}
+
+	GD3DAllowRemoveUnused = Input.Environment.CompilerFlags.Contains(CFLAG_ForceRemoveUnusedInterpolators) ? 1 : 0;
+
+	FString EntryPointName = Input.EntryPointName;
+
+	Output.bFailedRemovingUnused = false;
+	if (GD3DAllowRemoveUnused == 1 && Input.Target.Frequency == SF_Vertex && Input.bCompilingForShaderPipeline)
+	{
+		// Always add SV_Position
+		TArray<FString> UsedOutputs = Input.UsedOutputs;
+		UsedOutputs.AddUnique(TEXT("SV_POSITION"));
+
+		// We can't remove any of the output-only system semantics
+		//@todo - there are a bunch of tessellation ones as well
+		TArray<FString> Exceptions;
+		Exceptions.AddUnique(TEXT("SV_ClipDistance"));
+		Exceptions.AddUnique(TEXT("SV_ClipDistance0"));
+		Exceptions.AddUnique(TEXT("SV_ClipDistance1"));
+		Exceptions.AddUnique(TEXT("SV_ClipDistance2"));
+		Exceptions.AddUnique(TEXT("SV_ClipDistance3"));
+		Exceptions.AddUnique(TEXT("SV_ClipDistance4"));
+		Exceptions.AddUnique(TEXT("SV_ClipDistance5"));
+		Exceptions.AddUnique(TEXT("SV_ClipDistance6"));
+		Exceptions.AddUnique(TEXT("SV_ClipDistance7"));
+
+		Exceptions.AddUnique(TEXT("SV_CullDistance"));
+		Exceptions.AddUnique(TEXT("SV_CullDistance0"));
+		Exceptions.AddUnique(TEXT("SV_CullDistance1"));
+		Exceptions.AddUnique(TEXT("SV_CullDistance2"));
+		Exceptions.AddUnique(TEXT("SV_CullDistance3"));
+		Exceptions.AddUnique(TEXT("SV_CullDistance4"));
+		Exceptions.AddUnique(TEXT("SV_CullDistance5"));
+		Exceptions.AddUnique(TEXT("SV_CullDistance6"));
+		Exceptions.AddUnique(TEXT("SV_CullDistance7"));
+		
+		TArray<FString> Errors;
+		if (!RemoveUnusedOutputs(PreprocessedShaderSource, UsedOutputs, Exceptions, EntryPointName, Errors))
+		{
+			UE_LOG(LogD3D11ShaderCompiler, Warning, TEXT("Failed to Remove unused outputs [%s]!"), *Input.DumpDebugInfoPath);
+			for (int32 Index = 0; Index < Errors.Num(); ++Index)
+			{
+				FShaderCompilerError NewError;
+				NewError.StrippedErrorMessage = Errors[Index];
+				Output.Errors.Add(NewError);
+			}
+			Output.bFailedRemovingUnused = true;
+		}
+	}
+
+	if (Input.RootParameterBindings.Num())
+	{
+		MoveShaderParametersToRootConstantBuffer(Input, PreprocessedShaderSource);
+	}
+	RemoveUniformBuffersFromSource(Input.Environment, PreprocessedShaderSource);
+
+	// Override default compiler path to newer dll
+	CompilerPath = FPaths::EngineDir();
+#if !PLATFORM_64BITS
+	CompilerPath.Append(TEXT("Binaries/ThirdParty/Windows/DirectX/x86/d3dcompiler_47.dll"));
+#else
+	CompilerPath.Append(TEXT("Binaries/ThirdParty/Windows/DirectX/x64/d3dcompiler_47.dll"));
+#endif
+
+	// @TODO - currently d3d11 uses d3d10 shader compiler flags... update when this changes in DXSDK
+	// @TODO - implement different material path to allow us to remove backwards compat flag on sm5 shaders
+	uint32 CompileFlags = D3D10_SHADER_ENABLE_BACKWARDS_COMPATIBILITY
+		// Unpack uniform matrices as row-major to match the CPU layout.
+		| D3D10_SHADER_PACK_MATRIX_ROW_MAJOR;
+
+	if (DEBUG_SHADERS || Input.Environment.CompilerFlags.Contains(CFLAG_Debug)) 
+	{
+		//add the debug flags
+		CompileFlags |= D3D10_SHADER_DEBUG | D3D10_SHADER_SKIP_OPTIMIZATION;
+	}
+	else
+	{
+		if (Input.Environment.CompilerFlags.Contains(CFLAG_StandardOptimization))
+		{
+			CompileFlags |= D3D10_SHADER_OPTIMIZATION_LEVEL1;
+		}
+		else
+		{
+			CompileFlags |= D3D10_SHADER_OPTIMIZATION_LEVEL3;
+		}
+	}
+
+	for (int32 FlagIndex = 0; FlagIndex < Input.Environment.CompilerFlags.Num(); FlagIndex++)
+	{
+		//accumulate flags set by the shader
+		CompileFlags |= TranslateCompilerFlagD3D11((ECompilerFlags)Input.Environment.CompilerFlags[FlagIndex]);
+	}
+
+	TArray<FString> FilteredErrors;
+	if (!CompileAndProcessD3DShader(PreprocessedShaderSource, CompilerPath, CompileFlags, Input, EntryPointName, ShaderProfile, false, FilteredErrors, Output))
+	{
+		if (!FilteredErrors.Num())
+		{
+			FilteredErrors.Add(TEXT("Compile Failed without errors!"));
+		}
+	}
+
+	// Process errors
+	for (int32 ErrorIndex = 0; ErrorIndex < FilteredErrors.Num(); ErrorIndex++)
+	{
+		const FString& CurrentError = FilteredErrors[ErrorIndex];
+		FShaderCompilerError NewError;
+		// Extract the filename and line number from the shader compiler error message for PC whose format is:
+		// "d:\UE4\Binaries\BasePassPixelShader(30,7): error X3000: invalid target or usage string"
+		int32 FirstParenIndex = CurrentError.Find(TEXT("("));
+		int32 LastParenIndex = CurrentError.Find(TEXT("):"));
+		if (FirstParenIndex != INDEX_NONE 
+			&& LastParenIndex != INDEX_NONE
+			&& LastParenIndex > FirstParenIndex)
+		{
+			NewError.ErrorVirtualFilePath = CurrentError.Left(FirstParenIndex);
+			NewError.ErrorLineString = CurrentError.Mid(FirstParenIndex + 1, LastParenIndex - FirstParenIndex - FCString::Strlen(TEXT("(")));
+			NewError.StrippedErrorMessage = CurrentError.Right(CurrentError.Len() - LastParenIndex - FCString::Strlen(TEXT("):")));
+		}
+		else
+		{
+			NewError.StrippedErrorMessage = CurrentError;
+		}
+		Output.Errors.Add(NewError);
+	}
+
+	if (Input.ExtraSettings.bExtractShaderSource)
+	{
+		Output.OptionalFinalShaderSource = PreprocessedShaderSource;
+	}
+}
+
+void CompileShader_Windows_SM5(const FShaderCompilerInput& Input,FShaderCompilerOutput& Output,const FString& WorkingDirectory)
+{
+	check(Input.Target.Platform == SP_PCD3D_SM5);
+
+	FShaderCompilerDefinitions AdditionalDefines;
+	AdditionalDefines.SetDefine(TEXT("SM5_PROFILE"), 1);
+	CompileD3D11Shader(Input, Output, AdditionalDefines, WorkingDirectory);
+}
+
+void CompileShader_Windows_SM4(const FShaderCompilerInput& Input,FShaderCompilerOutput& Output,const FString& WorkingDirectory)
+{
+	check(Input.Target.Platform == SP_PCD3D_SM4);
+
+	FShaderCompilerDefinitions AdditionalDefines;
+	AdditionalDefines.SetDefine(TEXT("SM4_PROFILE"), 1);
+	CompileD3D11Shader(Input, Output, AdditionalDefines, WorkingDirectory);
+}
+
+void CompileShader_Windows_ES2(const FShaderCompilerInput& Input,FShaderCompilerOutput& Output,const FString& WorkingDirectory)
+{
+	check(Input.Target.Platform == SP_PCD3D_ES2);
+
+	FShaderCompilerDefinitions AdditionalDefines;
+	AdditionalDefines.SetDefine(TEXT("ES2_PROFILE"), 1);
+	CompileD3D11Shader(Input, Output, AdditionalDefines, WorkingDirectory);
+}
+
+void CompileShader_Windows_ES3_1(const FShaderCompilerInput& Input, FShaderCompilerOutput& Output, const FString& WorkingDirectory)
+{
+	check(Input.Target.Platform == SP_PCD3D_ES3_1);
+
+	FShaderCompilerDefinitions AdditionalDefines;
+	AdditionalDefines.SetDefine(TEXT("ES3_1_PROFILE"), 1);
+	CompileD3D11Shader(Input, Output, AdditionalDefines, WorkingDirectory);
+}
