Implemented support for wave operation intrinsics (Shader Model 6.0) as an extension of SP_PCD3D_SM5, available when running on compatible hardware and D3D12.

Wave intrinsics can be enabled for a shader by providing CFLAG_WaveOperations to the compiler flags in ModifyCompilationEnvironment().
Run-time availability must be queried by checking GRHISupportsWaveOperations before trying to use a shader that uses wave intrinsics.

Removed all RHI_RAYTRACING define checks from D3DShaderCompiler.cpp.

diff --git a/Engine/Source/Developer/Windows/ShaderFormatD3D/Private/D3DShaderCompiler.cpp b/Engine/Source/Developer/Windows/ShaderFormatD3D/Private/D3DShaderCompiler.cpp
--- a/Engine/Source/Developer/Windows/ShaderFormatD3D/Private/D3DShaderCompiler.cpp
+++ b/Engine/Source/Developer/Windows/ShaderFormatD3D/Private/D3DShaderCompiler.cpp
@@ -33,14 +33,12 @@ DEFINE_LOG_CATEGORY_STATIC(LogD3D11ShaderCompiler, Log, All);
 
 #pragma warning(pop)
 
-#if RHI_RAYTRACING
 MSVC_PRAGMA(warning(push))
 MSVC_PRAGMA(warning(disable : 4191)) // warning C4191: 'type cast': unsafe conversion from 'FARPROC' to 'DxcCreateInstanceProc'
 #include <dxc/dxcapi.h>
 #include <dxc/Support/dxcapi.use.h>
 #include <d3d12shader.h>
 MSVC_PRAGMA(warning(pop))
-#endif // RHI_RAYTRACING
 
 static int32 GD3DAllowRemoveUnused = 0;
 
@@ -89,7 +87,6 @@ static void D3D11FilterShaderCompileWarnings(const FString& CompileWarnings, TAr
 
 static bool IsRayTracingShader(const FShaderTarget& Target)
 {
-#if RHI_RAYTRACING
 	switch(Target.Frequency)
 	{
 	case SF_RayGen:
@@ -99,14 +96,10 @@ static bool IsRayTracingShader(const FShaderTarget& Target)
 	default:
 		return false;
 	}
-#else // RHI_RAYTRACING
-	return false;
-#endif // RHI_RAYTRACING
 }
 
 static uint32 GetAutoBindingSpace(const FShaderTarget& Target)
 {
-#if RHI_RAYTRACING
 	switch (Target.Frequency)
 	{
 	case SF_RayGen:
@@ -117,13 +110,10 @@ static uint32 GetAutoBindingSpace(const FShaderTarget& Target)
 	default:
 		return 0;
 	}
-#else // RHI_RAYTRACING
-	return 0;
-#endif // RHI_RAYTRACING
 }
 
 // @return 0 if not recognized
-static const TCHAR* GetShaderProfileName(FShaderTarget Target)
+static const TCHAR* GetShaderProfileName(FShaderTarget Target, bool bUseWaveOperations)
 {
 	if(Target.Platform == SP_PCD3D_SM5)
 	{
@@ -134,23 +124,21 @@ static const TCHAR* GetShaderProfileName(FShaderTarget Target)
 			checkfSlow(false, TEXT("Unexpected shader frequency"));
 			return nullptr;
 		case SF_Pixel:
-			return TEXT("ps_5_0");
+			return bUseWaveOperations ? TEXT("ps_6_0") : TEXT("ps_5_0");
 		case SF_Vertex:
-			return TEXT("vs_5_0");
+			return bUseWaveOperations ? TEXT("vs_6_0") : TEXT("vs_5_0");
 		case SF_Hull:
-			return TEXT("hs_5_0");
+			return bUseWaveOperations ? TEXT("hs_6_0") : TEXT("hs_5_0");
 		case SF_Domain:
-			return TEXT("ds_5_0");
+			return bUseWaveOperations ? TEXT("ds_6_0") : TEXT("ds_5_0");
 		case SF_Geometry:
-			return TEXT("gs_5_0");
+			return bUseWaveOperations ? TEXT("gs_6_0") : TEXT("gs_5_0");
 		case SF_Compute:
-			return TEXT("cs_5_0");
-#if RHI_RAYTRACING
+			return bUseWaveOperations ? TEXT("cs_6_0") : TEXT("cs_5_0");
 		case SF_RayGen:
 		case SF_RayMiss:
 		case SF_RayHitGroup:
 			return TEXT("lib_6_3");
-#endif // RHI_RAYTRACING
 		}
 	}
 	else if(Target.Platform == SP_PCD3D_SM4)
@@ -471,8 +459,6 @@ static HRESULT D3DCompileWrapper(
 // Utility variable so we can place a breakpoint while debugging
 static int32 GBreakpoint = 0;
 
-#if RHI_RAYTRACING
-
 #define VERIFYHRESULT(expr) { HRESULT HR##__LINE__ = expr; if (FAILED(HR##__LINE__)) { UE_LOG(LogD3D11ShaderCompiler, Fatal, TEXT(#expr " failed: Result=%08x"), HR##__LINE__); } }
 
 static dxc::DxcDllSupport& GetDxcDllHelper()
@@ -548,10 +534,13 @@ static void D3DCreateDXCArguments(TArray<const WCHAR*>& OutArgs, const WCHAR* Ex
 		UE_LOG(LogD3D11ShaderCompiler, Fatal, TEXT("Unsupported register binding space %d"), AutoBindingSpace);
 	}
 
-	// Ensure that only the requested functions exists in the output DXIL.
-	// All other functions and their used resources must be eliminated.
-	OutArgs.Add(L"/exports");
-	OutArgs.Add(Exports);
+	if (Exports && *Exports)
+	{
+		// Ensure that only the requested functions exists in the output DXIL.
+		// All other functions and their used resources must be eliminated.
+		OutArgs.Add(L"/exports");
+		OutArgs.Add(Exports);
+	}
 
 	if (CompileFlags & D3D10_SHADER_PREFER_FLOW_CONTROL)
 	{
@@ -625,7 +614,7 @@ static void D3DCreateDXCArguments(TArray<const WCHAR*>& OutArgs, const WCHAR* Ex
 	checkf(CompileFlags == 0, TEXT("Unhandled shader compiler flag!"));
 }
 
-static FString D3DCreateDXCCompileBatchFile(const FString& ShaderPath, const WCHAR* Exports, const TCHAR* ShaderProfile, uint32 CompileFlags, FShaderCompilerOutput& CompilerOutput, uint32 AutoBindingSpace = ~0u)
+static FString D3DCreateDXCCompileBatchFile(const FString& ShaderPath, const WCHAR* EntryName, const WCHAR* Exports, const TCHAR* ShaderProfile, uint32 CompileFlags, FShaderCompilerOutput& CompilerOutput, uint32 AutoBindingSpace = ~0u)
 {
 	TArray<const WCHAR*> Args;
 	const uint32 DXCFlags = CompileFlags & (~D3D10_SHADER_ENABLE_BACKWARDS_COMPATIBILITY);
@@ -645,6 +634,14 @@ static FString D3DCreateDXCCompileBatchFile(const FString& ShaderPath, const WCH
 	DXCCommandline += TEXT(" /T ");
 	DXCCommandline += ShaderProfile;
 
+	// Append entry point name if export symbol list is not provided.
+	// Explicit export symbol list is used for lib_6_x targets, such as ray tracing shaders.
+	if (!(*Exports))
+	{
+		DXCCommandline += TEXT(" /E ");
+		DXCCommandline += EntryName;
+	}
+
 	if (FPaths::GetExtension(ShaderPath) == TEXT("usf"))
 	{
 		DXCCommandline += FString::Printf(TEXT(" /Fc%sd3dasm"), *ShaderPath.LeftChop(3));
@@ -663,7 +660,8 @@ static FString D3DCreateDXCCompileBatchFile(const FString& ShaderPath, const WCH
   )
 #endif
 
-static HRESULT D3DCreateLibraryReflectionFromBlob(ID3DBlob* DxilBlob, TRefCountPtr<ID3D12LibraryReflection>& OutReflection)
+template <typename T>
+static HRESULT D3DCreateReflectionFromBlob(ID3DBlob* DxilBlob, TRefCountPtr<T>& OutReflection)
 {
 	dxc::DxcDllSupport& DxcDllHelper = GetDxcDllHelper();
 
@@ -687,8 +685,6 @@ inline bool IsCompatibleBinding(const D3D12_SHADER_INPUT_BIND_DESC& BindDesc, ui
 	return BindDesc.Space == BindingSpace;
 }
 
-#endif // RHI_RAYTRACING
-
 inline bool IsCompatibleBinding(const D3D11_SHADER_INPUT_BIND_DESC& BindDesc, uint32 BindingSpace)
 {
 	return true;
@@ -864,7 +860,6 @@ static void ExtractParameterMapFromD3DShader(
 
 			NumSRVs = FMath::Max(NumSRVs, BindDesc.BindPoint + BindDesc.BindCount);
 		}
-#if RHI_RAYTRACING
 		// #dxr_todo: D3D_SIT_RTACCELERATIONSTRUCTURE is declared in latest version of dxcapi.h. Update this code after upgrading DXC.
 		else if (BindDesc.Type == 12 /*D3D_SIT_RTACCELERATIONSTRUCTURE*/)
 		{
@@ -883,11 +878,9 @@ static void ExtractParameterMapFromD3DShader(
 
 			NumSRVs = FMath::Max(NumSRVs, BindDesc.BindPoint + BindDesc.BindCount);
 		}
-#endif
 	}
 }
 
-#if RHI_RAYTRACING
 // Parses ray tracing shader entry point specification string in one of the following formats:
 // 1) Verbatim single entry point name, e.g. "MainRGS"
 // 2) Complex entry point for ray tracing hit group shaders:
@@ -923,7 +916,6 @@ static void ParseRayTracingEntryPoint(const FString& Input, FString& OutMain, FS
 		OutMain = Input;
 	}
 }
-#endif // RHI_RAYTRACING
 
 // Generate the dumped usf file; call the D3D compiler, gather reflection information and generate the output data
 static bool CompileAndProcessD3DShader(FString& PreprocessedShaderSource, const FString& CompilerPath,
@@ -934,8 +926,7 @@ static bool CompileAndProcessD3DShader(FString& PreprocessedShaderSource, const
 	auto AnsiSourceFile = StringCast<ANSICHAR>(*PreprocessedShaderSource);
 
 	const bool bIsRayTracingShader = IsRayTracingShader(Input.Target);
-
-#if RHI_RAYTRACING
+	const bool bUseDXC = bIsRayTracingShader || Input.Environment.CompilerFlags.Contains(CFLAG_WaveOperations);
 
 	const uint32 AutoBindingSpace = GetAutoBindingSpace(Input.Target);
 
@@ -963,8 +954,6 @@ static bool CompileAndProcessD3DShader(FString& PreprocessedShaderSource, const
 		}
 	}
 
-#endif // RHI_RAYTRACING
-
 	bool bDumpDebugInfo = false;
 	// Write out the preprocessed file and a batch file to compile it if requested (DumpDebugInfoPath is valid)
 	if (Input.DumpDebugInfoPath.Len() > 0 && IFileManager::Get().DirectoryExists(*Input.DumpDebugInfoPath))
@@ -990,11 +979,9 @@ static bool CompileAndProcessD3DShader(FString& PreprocessedShaderSource, const
 
 		FString BatchFileContents;
 
-		if (bIsRayTracingShader)
+		if (bUseDXC)
 		{
-#if RHI_RAYTRACING
-			BatchFileContents = D3DCreateDXCCompileBatchFile(Filename, *RayTracingExports, ShaderProfile, CompileFlags, Output, AutoBindingSpace);
-#endif // RHI_RAYTRACING
+			BatchFileContents = D3DCreateDXCCompileBatchFile(Filename, *EntryPointName, *RayTracingExports, ShaderProfile, CompileFlags, Output, AutoBindingSpace);
 		}
 		else
 		{
@@ -1024,10 +1011,8 @@ static bool CompileAndProcessD3DShader(FString& PreprocessedShaderSource, const
 	pD3DStripShader D3DStripShaderFunc = nullptr;
 	bool bCompilerPathFunctionsUsed = false;
 
-	if (bIsRayTracingShader)
+	if (bUseDXC)
 	{
-#if RHI_RAYTRACING
-
 		TArray<const WCHAR*> Args;
 
 		// Ignore backwards compatibility flag (/Gec) as it is deprecated.
@@ -1038,7 +1023,7 @@ static bool CompileAndProcessD3DShader(FString& PreprocessedShaderSource, const
 		TRefCountPtr<IDxcBlobEncoding> DxcErrorBlob;
 
 		Result = D3DCompileToDxil(AnsiSourceFile.Get(),
-			L"", // dummy entry point for ray tracing shaders
+			bIsRayTracingShader ? L"" : *EntryPointName, // dummy entry point for ray tracing shaders
 			TCHAR_TO_WCHAR(ShaderProfile),
 			Args.GetData(), Args.Num(),
 			Shader, DxcErrorBlob);
@@ -1053,11 +1038,6 @@ static bool CompileAndProcessD3DShader(FString& PreprocessedShaderSource, const
 		{
 			FilteredErrors.Add(TEXT("D3DCompileToDxil failed"));
 		}
-#else
-		FilteredErrors.Add(TEXT("Attempted to compile a ray tracing shader, but ray tracing support is not enabled in this binary."));
-		Result = E_FAIL;
-		checkNoEntry();
-#endif
 	}
 	else
 	{
@@ -1153,110 +1133,134 @@ static bool CompileAndProcessD3DShader(FString& PreprocessedShaderSource, const
 		TBitArray<> UsedUniformBufferSlots;
 		UsedUniformBufferSlots.Init(false, 32);
 
-#if RHI_RAYTRACING
-		if (bIsRayTracingShader)
+		if (bUseDXC)
 		{
-			TRefCountPtr<ID3D12LibraryReflection> LibraryReflection;
-
-			Result = D3DCreateLibraryReflectionFromBlob(Shader, LibraryReflection);
-
-			if (FAILED(Result))
+			if (bIsRayTracingShader)
 			{
-				UE_LOG(LogD3D11ShaderCompiler, Fatal, TEXT("D3DReflectDxil failed: Result=%08x"), Result);
-			}
+				TRefCountPtr<ID3D12LibraryReflection> LibraryReflection;
 
-			D3D12_LIBRARY_DESC LibraryDesc = {};
-			LibraryReflection->GetDesc(&LibraryDesc);
+				Result = D3DCreateReflectionFromBlob(Shader, LibraryReflection);
 
-			ID3D12FunctionReflection* FunctionReflection = nullptr;
-			D3D12_FUNCTION_DESC FunctionDesc = {};
+				if (FAILED(Result))
+				{
+					UE_LOG(LogD3D11ShaderCompiler, Fatal, TEXT("D3DReflectDxil failed: Result=%08x"), Result);
+				}
 
-			// MangledEntryPoints contains partial mangled entry point signatures in a the following form:
-			// ?QualifiedName@ (as described here: https://en.wikipedia.org/wiki/Name_mangling)
-			// Entry point parameters are currently not included in the partial mangling.
-			TArray<FString, TInlineAllocator<3>> MangledEntryPoints;
-			if (!RayEntryPoint.IsEmpty())
-			{
-				MangledEntryPoints.Add(FString::Printf(TEXT("?%s@"), *RayEntryPoint));
-			}
-			if (!RayAnyHitEntryPoint.IsEmpty())
-			{
-				MangledEntryPoints.Add(FString::Printf(TEXT("?%s@"), *RayAnyHitEntryPoint));
-			}
-			if (!RayIntersectionEntryPoint.IsEmpty())
-			{
-				MangledEntryPoints.Add(FString::Printf(TEXT("?%s@"), *RayIntersectionEntryPoint));
-			}
+				D3D12_LIBRARY_DESC LibraryDesc = {};
+				LibraryReflection->GetDesc(&LibraryDesc);
 
-			uint32 NumFoundEntryPoints = 0;
+				ID3D12FunctionReflection* FunctionReflection = nullptr;
+				D3D12_FUNCTION_DESC FunctionDesc = {};
 
-			for (uint32 FunctionIndex = 0; FunctionIndex < LibraryDesc.FunctionCount; ++FunctionIndex)
-			{
-				FunctionReflection = LibraryReflection->GetFunctionByIndex(FunctionIndex);
-				FunctionReflection->GetDesc(&FunctionDesc);
+				// MangledEntryPoints contains partial mangled entry point signatures in a the following form:
+				// ?QualifiedName@ (as described here: https://en.wikipedia.org/wiki/Name_mangling)
+				// Entry point parameters are currently not included in the partial mangling.
+				TArray<FString, TInlineAllocator<3>> MangledEntryPoints;
 
-				for (const FString& MangledEntryPoint : MangledEntryPoints)
+				if (!RayEntryPoint.IsEmpty())
 				{
-					// Entry point parameters are currently not included in the partial mangling, therefore partial substring match is used here.
-					if (FCStringAnsi::Strstr(FunctionDesc.Name, TCHAR_TO_ANSI(*MangledEntryPoint)))
-					{
-						// Note: calling ExtractParameterMapFromD3DShader multiple times merges the reflection data for multiple functions
-						ExtractParameterMapFromD3DShader<ID3D12FunctionReflection, D3D12_FUNCTION_DESC, D3D12_SHADER_INPUT_BIND_DESC,
-							ID3D12ShaderReflectionConstantBuffer, D3D12_SHADER_BUFFER_DESC,
-							ID3D12ShaderReflectionVariable, D3D12_SHADER_VARIABLE_DESC>(
-								Input.Target.Platform, AutoBindingSpace, Input.VirtualSourceFilePath, FunctionReflection, FunctionDesc, bGlobalUniformBufferUsed, NumSamplers, NumSRVs, NumCBs, NumUAVs,
-								Output, UniformBufferNames, UsedUniformBufferSlots);
-
-						NumFoundEntryPoints++;
-					}
+					MangledEntryPoints.Add(FString::Printf(TEXT("?%s@"), *RayEntryPoint));
+				}
+				if (!RayAnyHitEntryPoint.IsEmpty())
+				{
+					MangledEntryPoints.Add(FString::Printf(TEXT("?%s@"), *RayAnyHitEntryPoint));
+				}
+				if (!RayIntersectionEntryPoint.IsEmpty())
+				{
+					MangledEntryPoints.Add(FString::Printf(TEXT("?%s@"), *RayIntersectionEntryPoint));
 				}
-			}
 
-			if (NumFoundEntryPoints == MangledEntryPoints.Num())
-			{
-				Output.bSucceeded = true;
+				uint32 NumFoundEntryPoints = 0;
 
-				if (bGlobalUniformBufferUsed)
+				for (uint32 FunctionIndex = 0; FunctionIndex < LibraryDesc.FunctionCount; ++FunctionIndex)
 				{
-					FString ErrorString = TEXT("Global constant buffer cannot be used in a ray tracing shader.");
+					FunctionReflection = LibraryReflection->GetFunctionByIndex(FunctionIndex);
+					FunctionReflection->GetDesc(&FunctionDesc);
 
-					uint32 NumLooseParameters = 0;
-					for (const auto& It : Output.ParameterMap.ParameterMap)
+					for (const FString& MangledEntryPoint : MangledEntryPoints)
 					{
-						if (It.Value.Type == EShaderParameterType::LooseData)
+						// Entry point parameters are currently not included in the partial mangling, therefore partial substring match is used here.
+						if (FCStringAnsi::Strstr(FunctionDesc.Name, TCHAR_TO_ANSI(*MangledEntryPoint)))
 						{
-							NumLooseParameters++;
+							// Note: calling ExtractParameterMapFromD3DShader multiple times merges the reflection data for multiple functions
+							ExtractParameterMapFromD3DShader<ID3D12FunctionReflection, D3D12_FUNCTION_DESC, D3D12_SHADER_INPUT_BIND_DESC,
+								ID3D12ShaderReflectionConstantBuffer, D3D12_SHADER_BUFFER_DESC,
+								ID3D12ShaderReflectionVariable, D3D12_SHADER_VARIABLE_DESC>(
+									Input.Target.Platform, AutoBindingSpace, Input.VirtualSourceFilePath, FunctionReflection, FunctionDesc, bGlobalUniformBufferUsed, NumSamplers, NumSRVs, NumCBs, NumUAVs,
+									Output, UniformBufferNames, UsedUniformBufferSlots);
+
+							NumFoundEntryPoints++;
 						}
 					}
+				}
+
+				if (NumFoundEntryPoints == MangledEntryPoints.Num())
+				{
+					Output.bSucceeded = true;
 
-					if (NumLooseParameters)
+					if (bGlobalUniformBufferUsed && bIsRayTracingShader)
 					{
-						ErrorString += TEXT(" Global parameters: ");
-						uint32 ParameterIndex = 0;
+						FString ErrorString = TEXT("Global constant buffer cannot be used in a ray tracing shader.");
+
+						uint32 NumLooseParameters = 0;
 						for (const auto& It : Output.ParameterMap.ParameterMap)
 						{
 							if (It.Value.Type == EShaderParameterType::LooseData)
 							{
-								--NumLooseParameters;
-								ErrorString += FString::Printf(TEXT("%s%s"), *It.Key, NumLooseParameters ? TEXT(", ") : TEXT("."));
+								NumLooseParameters++;
+							}
+						}
+
+						if (NumLooseParameters)
+						{
+							ErrorString += TEXT(" Global parameters: ");
+							uint32 ParameterIndex = 0;
+							for (const auto& It : Output.ParameterMap.ParameterMap)
+							{
+								if (It.Value.Type == EShaderParameterType::LooseData)
+								{
+									--NumLooseParameters;
+									ErrorString += FString::Printf(TEXT("%s%s"), *It.Key, NumLooseParameters ? TEXT(", ") : TEXT("."));
+								}
 							}
 						}
-					}
 
-					FilteredErrors.Add(ErrorString);
-					Result = E_FAIL;
+						FilteredErrors.Add(ErrorString);
+						Result = E_FAIL;
+						Output.bSucceeded = false;
+					}
+				}
+				else
+				{
+					UE_LOG(LogD3D11ShaderCompiler, Fatal, TEXT("Failed to find required points in the shader library."));
 					Output.bSucceeded = false;
 				}
 			}
 			else
 			{
-				UE_LOG(LogD3D11ShaderCompiler, Fatal, TEXT("Failed to find required points in the shader library."));
-				Output.bSucceeded = false;
+				TRefCountPtr<ID3D12ShaderReflection> ShaderReflection;
+
+				Result = D3DCreateReflectionFromBlob(Shader, ShaderReflection);
+
+				if (FAILED(Result))
+				{
+					UE_LOG(LogD3D11ShaderCompiler, Fatal, TEXT("D3DReflectDxil failed: Result=%08x"), Result);
+				}
+
+				D3D12_SHADER_DESC ShaderDesc = {};
+				ShaderReflection->GetDesc(&ShaderDesc);
+
+				ExtractParameterMapFromD3DShader<ID3D12ShaderReflection, D3D12_SHADER_DESC, D3D12_SHADER_INPUT_BIND_DESC,
+					ID3D12ShaderReflectionConstantBuffer, D3D12_SHADER_BUFFER_DESC,
+					ID3D12ShaderReflectionVariable, D3D12_SHADER_VARIABLE_DESC>(
+						Input.Target.Platform, AutoBindingSpace, Input.VirtualSourceFilePath, ShaderReflection, ShaderDesc, bGlobalUniformBufferUsed, NumSamplers, NumSRVs, NumCBs, NumUAVs,
+						Output, UniformBufferNames, UsedUniformBufferSlots);
+
+
+				Output.bSucceeded = true;
 			}
 		}
-		else
-#endif // RHI_RAYTRACING
-		if (D3DReflectFunc)
+		else if (D3DReflectFunc)
 		{
 			Output.bSucceeded = true;
 			ID3D11ShaderReflection* Reflector = NULL;
@@ -1476,14 +1480,12 @@ static bool CompileAndProcessD3DShader(FString& PreprocessedShaderSource, const
 			FMemoryWriter Ar(Output.ShaderCode.GetWriteAccess(), true);
 			Ar << SRT;
 
-#if RHI_RAYTRACING
 			if (bIsRayTracingShader)
 			{
 				Ar << RayEntryPoint;
 				Ar << RayAnyHitEntryPoint;
 				Ar << RayIntersectionEntryPoint;
 			}
-#endif // RHI_RAYTRACING
 
 			Ar.Serialize(CompressedData->GetBufferPointer(), CompressedData->GetBufferSize());
 
@@ -1545,11 +1547,12 @@ static bool CompileAndProcessD3DShader(FString& PreprocessedShaderSource, const
 	return SUCCEEDED(Result);
 }
 
-void CompileD3D11Shader(const FShaderCompilerInput& Input,FShaderCompilerOutput& Output, FShaderCompilerDefinitions& AdditionalDefines, const FString& WorkingDirectory)
+void CompileD3DShader(const FShaderCompilerInput& Input,FShaderCompilerOutput& Output, FShaderCompilerDefinitions& AdditionalDefines, const FString& WorkingDirectory)
 {
 	FString PreprocessedShaderSource;
 	FString CompilerPath;
-	const TCHAR* ShaderProfile = GetShaderProfileName(Input.Target);
+	const bool bUseWaveOperations = Input.Environment.CompilerFlags.Contains(CFLAG_WaveOperations); // Forces shader model 6.0 for this shader
+	const TCHAR* ShaderProfile = GetShaderProfileName(Input.Target, bUseWaveOperations);
 
 	if(!ShaderProfile)
 	{
@@ -1715,7 +1718,7 @@ void CompileShader_Windows_SM5(const FShaderCompilerInput& Input,FShaderCompiler
 
 	FShaderCompilerDefinitions AdditionalDefines;
 	AdditionalDefines.SetDefine(TEXT("SM5_PROFILE"), 1);
-	CompileD3D11Shader(Input, Output, AdditionalDefines, WorkingDirectory);
+	CompileD3DShader(Input, Output, AdditionalDefines, WorkingDirectory);
 }
 
 void CompileShader_Windows_SM4(const FShaderCompilerInput& Input,FShaderCompilerOutput& Output,const FString& WorkingDirectory)
@@ -1724,7 +1727,7 @@ void CompileShader_Windows_SM4(const FShaderCompilerInput& Input,FShaderCompiler
 
 	FShaderCompilerDefinitions AdditionalDefines;
 	AdditionalDefines.SetDefine(TEXT("SM4_PROFILE"), 1);
-	CompileD3D11Shader(Input, Output, AdditionalDefines, WorkingDirectory);
+	CompileD3DShader(Input, Output, AdditionalDefines, WorkingDirectory);
 }
 
 void CompileShader_Windows_ES2(const FShaderCompilerInput& Input,FShaderCompilerOutput& Output,const FString& WorkingDirectory)
@@ -1733,7 +1736,7 @@ void CompileShader_Windows_ES2(const FShaderCompilerInput& Input,FShaderCompiler
 
 	FShaderCompilerDefinitions AdditionalDefines;
 	AdditionalDefines.SetDefine(TEXT("ES2_PROFILE"), 1);
-	CompileD3D11Shader(Input, Output, AdditionalDefines, WorkingDirectory);
+	CompileD3DShader(Input, Output, AdditionalDefines, WorkingDirectory);
 }
 
 void CompileShader_Windows_ES3_1(const FShaderCompilerInput& Input, FShaderCompilerOutput& Output, const FString& WorkingDirectory)
@@ -1742,5 +1745,5 @@ void CompileShader_Windows_ES3_1(const FShaderCompilerInput& Input, FShaderCompi
 
 	FShaderCompilerDefinitions AdditionalDefines;
 	AdditionalDefines.SetDefine(TEXT("ES3_1_PROFILE"), 1);
-	CompileD3D11Shader(Input, Output, AdditionalDefines, WorkingDirectory);
+	CompileD3DShader(Input, Output, AdditionalDefines, WorkingDirectory);
 }
diff --git a/Engine/Source/Runtime/D3D12RHI/Private/D3D12Adapter.cpp b/Engine/Source/Runtime/D3D12RHI/Private/D3D12Adapter.cpp
--- a/Engine/Source/Runtime/D3D12RHI/Private/D3D12Adapter.cpp
+++ b/Engine/Source/Runtime/D3D12RHI/Private/D3D12Adapter.cpp
@@ -139,6 +139,13 @@ void FD3D12Adapter::CreateRootDevice(bool bWithDebug)
 		IID_PPV_ARGS(RootDevice.GetInitReference())
 	));
 
+	// Detect availability of shader model 6.0 wave operations
+	{
+		D3D12_FEATURE_DATA_D3D12_OPTIONS1 Features = {};
+		RootDevice->CheckFeatureSupport(D3D12_FEATURE_D3D12_OPTIONS1, &Features, sizeof(Features));
+		GRHISupportsWaveOperations = Features.WaveOps;
+	}
+
 #if D3D12_RHI_RAYTRACING
 	bool bRayTracingSupported = false;
 
diff --git a/Engine/Source/Runtime/RHI/Private/RHI.cpp b/Engine/Source/Runtime/RHI/Private/RHI.cpp
--- a/Engine/Source/Runtime/RHI/Private/RHI.cpp
+++ b/Engine/Source/Runtime/RHI/Private/RHI.cpp
@@ -498,6 +498,7 @@ TRHIGlobal<bool> GRHISupportsInstancing(true);
 bool GRHISupportsFirstInstance = false;
 bool GRHISupportsDynamicResolution = false;
 bool GRHISupportsRayTracing = false;
+bool GRHISupportsWaveOperations = false;
 bool GRHISupportsRHIThread = false;
 bool GRHISupportsRHIOnTaskThread = false;
 bool GRHISupportsParallelRHIExecute = false;
diff --git a/Engine/Source/Runtime/RHI/Public/RHI.h b/Engine/Source/Runtime/RHI/Public/RHI.h
--- a/Engine/Source/Runtime/RHI/Public/RHI.h
+++ b/Engine/Source/Runtime/RHI/Public/RHI.h
@@ -194,6 +194,16 @@ inline RHI_API bool RHISupportsRayTracingShaders(EShaderPlatform Platform)
 	return Platform == SP_PCD3D_SM5;
 }
 
+/** Can this platform compile shaders that use shader model 6.0 wave intrinsics.
+ *  To use such shaders at runtime, also check GRHISupportsWaveOperations.
+ **/
+inline RHI_API bool RHISupportsWaveOperations(EShaderPlatform Platform)
+{
+	// Currently SM6 shaders are treated as an extension of SM5.
+	return Platform == SP_PCD3D_SM5;
+}
+
+
 // Wrapper for GRHI## global variables, allows values to be overridden for mobile preview modes.
 template <typename TValueType>
 class TRHIGlobal
@@ -461,6 +471,9 @@ extern RHI_API bool GRHISupportsDynamicResolution;
 /** Whether or not the RHI supports ray tracing on current hardware (acceleration structure building and new ray tracing-specific shader types). */
 extern RHI_API bool GRHISupportsRayTracing;
 
+/** Whether or not the RHI supports shader wave operations (shader model 6.0). */
+extern RHI_API bool GRHISupportsWaveOperations;
+
 /** Whether or not the RHI supports an RHI thread.
 Requirements for RHI thread
 * Microresources (those in RHIStaticStates.h) need to be able to be created by any thread at any time and be able to work with a radically simplified rhi resource lifecycle. CreateSamplerState, CreateRasterizerState, CreateDepthStencilState, CreateBlendState
diff --git a/Engine/Source/Runtime/RenderCore/Public/ShaderCore.h b/Engine/Source/Runtime/RenderCore/Public/ShaderCore.h
--- a/Engine/Source/Runtime/RenderCore/Public/ShaderCore.h
+++ b/Engine/Source/Runtime/RenderCore/Public/ShaderCore.h
@@ -174,7 +174,11 @@ enum ECompilerFlags
 	// Shaders uses external texture so may need special runtime handling
 	CFLAG_UsesExternalTexture,
 	// Use emulated uniform buffers on supported platforms
-	CFLAG_UseEmulatedUB
+	CFLAG_UseEmulatedUB,
+	// Enable wave operation intrinsics (requires DX12 and DXC/DXIL on PC).
+	// Check GRHISupportsWaveOperations before using shaders compiled with this flag at runtime.
+	// [REDACTED_PROFILE_URL]/DirectXShaderCompiler/wiki/Wave-Intrinsics
+	CFLAG_WaveOperations,
 };
 
 enum class EShaderParameterType : uint8
