Add DXC runtime requirements to fix binary builds

diff --git a/Engine/Source/Developer/Windows/ShaderFormatD3D/ShaderFormatD3D.Build.cs b/Engine/Source/Developer/Windows/ShaderFormatD3D/ShaderFormatD3D.Build.cs
--- a/Engine/Source/Developer/Windows/ShaderFormatD3D/ShaderFormatD3D.Build.cs
+++ b/Engine/Source/Developer/Windows/ShaderFormatD3D/ShaderFormatD3D.Build.cs
@@ -20,6 +20,12 @@ public class ShaderFormatD3D : ModuleRules
 			}
 			);
 
+		//DXC
+        if (Target.Platform == UnrealTargetPlatform.Win64)
+        {
+            RuntimeDependencies.Add("$(EngineDir)/Binaries/Win64/dxil.dll");
+            RuntimeDependencies.Add("$(EngineDir)/Binaries/Win64/dxcompiler.dll");
+        }
 		AddEngineThirdPartyPrivateStaticDependencies(Target, "DX11");
 	}
 }
