Live coding is only available for Win32/Win64, dont try to include diasdk dir for other platforms

diff --git a/Engine/Source/Developer/Windows/LiveCodingServer/LiveCodingServer.Build.cs b/Engine/Source/Developer/Windows/LiveCodingServer/LiveCodingServer.Build.cs
--- a/Engine/Source/Developer/Windows/LiveCodingServer/LiveCodingServer.Build.cs
+++ b/Engine/Source/Developer/Windows/LiveCodingServer/LiveCodingServer.Build.cs
@@ -10,19 +10,22 @@ public class LiveCodingServer : ModuleRules
 		PrivateDependencyModuleNames.Add("Core");
 		PrivateDependencyModuleNames.Add("Distorm");
 		PrivateDependencyModuleNames.Add("LiveCoding");
-		
-		PrivateIncludePaths.Add("Developer/Windows/LiveCoding/Private");
-		PrivateIncludePaths.Add("Developer/Windows/LiveCoding/Private/External");
 
-		string DiaSdkDir = Target.WindowsPlatform.DiaSdkDir;
-		if(DiaSdkDir == null)
+		if (Target.Platform == UnrealTargetPlatform.Win32 || Target.Platform == UnrealTargetPlatform.Win64)
 		{
-			throw new System.Exception("Unable to find DIA SDK directory");
-		}
+			PrivateIncludePaths.Add("Developer/Windows/LiveCoding/Private");
+			PrivateIncludePaths.Add("Developer/Windows/LiveCoding/Private/External");
+
+			string DiaSdkDir = Target.WindowsPlatform.DiaSdkDir;
+			if(DiaSdkDir == null)
+			{
+				throw new System.Exception("Unable to find DIA SDK directory");
+			}
 
-		PrivateIncludePaths.Add(Path.Combine(DiaSdkDir, "include"));
-		PublicAdditionalLibraries.Add(Path.Combine(DiaSdkDir, "lib", "amd64", "diaguids.lib"));
-		RuntimeDependencies.Add("$(TargetOutputDir)/msdia140.dll", Path.Combine(DiaSdkDir, "bin", "amd64", "msdia140.dll"));
+			PrivateIncludePaths.Add(Path.Combine(DiaSdkDir, "include"));
+			PublicAdditionalLibraries.Add(Path.Combine(DiaSdkDir, "lib", "amd64", "diaguids.lib"));
+			RuntimeDependencies.Add("$(TargetOutputDir)/msdia140.dll", Path.Combine(DiaSdkDir, "bin", "amd64", "msdia140.dll"));
+		}
 
 		PrecompileForTargets = PrecompileTargetsType.None;
 	}
