Better handling of window focus when compiling for live code.

* Console window is made visible and brought to the front when a compile is triggered.
* Patched application window is brought to the front once patching is complete.

diff --git a/Engine/Source/Developer/Windows/LiveCodingServer/Private/External/LC_ServerCommandThread.cpp b/Engine/Source/Developer/Windows/LiveCodingServer/Private/External/LC_ServerCommandThread.cpp
--- a/Engine/Source/Developer/Windows/LiveCodingServer/Private/External/LC_ServerCommandThread.cpp
+++ b/Engine/Source/Developer/Windows/LiveCodingServer/Private/External/LC_ServerCommandThread.cpp
@@ -566,6 +566,23 @@ unsigned int __stdcall ServerCommandThread::CompileThreadProxy(void* context)
 	return instance->CompileThread();
 }
 
+// BEGIN EPIC MOD - Focus application windows on patch complete
+BOOL CALLBACK FocusApplicationWindows(HWND WindowHandle, LPARAM Lparam)
+{
+	DWORD WindowProcessId;
+    GetWindowThreadProcessId(WindowHandle, &WindowProcessId);
+
+	const types::vector<LiveProcess*>& Processes = *(const types::vector<LiveProcess*>*)Lparam;
+	for (LiveProcess* Process : Processes)
+	{
+		if (Process->GetProcessId() == WindowProcessId && IsWindowVisible(WindowHandle))
+		{
+			SetForegroundWindow(WindowHandle);
+		}
+	}
+    return Windows::TRUE;
+}
+// END EPIC MOD
 
 void ServerCommandThread::CompileChanges(bool didAllProcessesMakeProgress)
 {
@@ -765,6 +782,7 @@ void ServerCommandThread::CompileChanges(bool didAllProcessesMakeProgress)
 
 		case LiveModule::ErrorType::SUCCESS:
 			GLiveCodingServer->GetCompileFinishedDelegate().ExecuteIfBound(ELiveCodingResult::Success, L"Patch creation successful.");
+			EnumWindows(FocusApplicationWindows, (LPARAM)&m_liveProcesses);
 			break;
 
 		default:
diff --git a/Engine/Source/Programs/LiveCodingConsole/Private/LiveCodingConsole.cpp b/Engine/Source/Programs/LiveCodingConsole/Private/LiveCodingConsole.cpp
--- a/Engine/Source/Programs/LiveCodingConsole/Private/LiveCodingConsole.cpp
+++ b/Engine/Source/Programs/LiveCodingConsole/Private/LiveCodingConsole.cpp
@@ -303,6 +303,8 @@ private:
 	{
 		if (!CompileNotification.IsValid())
 		{
+			ShowConsole();
+
 			FNotificationInfo Info(FText::FromString(TEXT("Starting...")));
 			Info.bFireAndForget = false;
 			Info.FadeOutDuration = 0.0f;
