LiveCoding: Expose a Tick() function so that systems that don't use fire EndFrame() callbacks can still tick the live coding module. Also add a -LiveCoding command line argument to force it on.

diff --git a/Engine/Source/Developer/Windows/LiveCoding/Private/LiveCodingModule.cpp b/Engine/Source/Developer/Windows/LiveCoding/Private/LiveCodingModule.cpp
--- a/Engine/Source/Developer/Windows/LiveCoding/Private/LiveCodingModule.cpp
+++ b/Engine/Source/Developer/Windows/LiveCoding/Private/LiveCodingModule.cpp
@@ -27,12 +27,17 @@ FLiveCodingModule::FLiveCodingModule()
 	, bStarted(false)
 {
 #if WITH_EDITOR
-	GConfig->GetBool(TEXT("LiveCoding"), TEXT("Enabled"), bEnabled, GEditorPerProjectIni);
+//	GConfig->GetBool(TEXT("LiveCoding"), TEXT("Enabled"), bEnabled, GEditorPerProjectIni);
 #endif
 }
 
 void FLiveCodingModule::StartupModule()
 {
+	if(FParse::Param(FCommandLine::Get(), TEXT("LiveCoding")))
+	{
+		bEnabled = true;
+	}
+
 	IConsoleManager& ConsoleManager = IConsoleManager::Get();
 
 	EnableCommand = ConsoleManager.RegisterConsoleCommand(
@@ -49,7 +54,7 @@ void FLiveCodingModule::StartupModule()
 		ECVF_Cheat
 	);
 
-	EndFrameDelegateHandle = FCoreDelegates::OnEndFrame.AddRaw(this, &FLiveCodingModule::OnEndFrame);
+	EndFrameDelegateHandle = FCoreDelegates::OnEndFrame.AddRaw(this, &FLiveCodingModule::Tick);
 
 	ISettingsModule* SettingsModule = FModuleManager::GetModulePtr<ISettingsModule>("Settings");
 	if (SettingsModule != nullptr)
@@ -130,10 +135,18 @@ void FLiveCodingModule::ShowConsole()
 
 void FLiveCodingModule::TriggerRecompile()
 {
-	LppTriggerRecompile();
+	if (!bStarted)
+	{
+		bShouldStart = true;
+		Tick();
+	}
+	if(bStarted)
+	{
+		LppTriggerRecompile();
+	}
 }
 
-void FLiveCodingModule::OnEndFrame()
+void FLiveCodingModule::Tick()
 {
 	if (bShouldStart && !bStarted)
 	{
diff --git a/Engine/Source/Developer/Windows/LiveCoding/Private/LiveCodingModule.h b/Engine/Source/Developer/Windows/LiveCoding/Private/LiveCodingModule.h
--- a/Engine/Source/Developer/Windows/LiveCoding/Private/LiveCodingModule.h
+++ b/Engine/Source/Developer/Windows/LiveCoding/Private/LiveCodingModule.h
@@ -23,6 +23,7 @@ public:
 	virtual bool IsEnabled() const override;
 	virtual void ShowConsole() override;
 	virtual void TriggerRecompile() override;
+	virtual void Tick() override;
 
 private:
 	bool bEnabled;
diff --git a/Engine/Source/Developer/Windows/LiveCoding/Public/ILiveCodingModule.h b/Engine/Source/Developer/Windows/LiveCoding/Public/ILiveCodingModule.h
--- a/Engine/Source/Developer/Windows/LiveCoding/Public/ILiveCodingModule.h
+++ b/Engine/Source/Developer/Windows/LiveCoding/Public/ILiveCodingModule.h
@@ -14,5 +14,6 @@ public:
 	virtual bool IsEnabled() const = 0;
 	virtual void ShowConsole() = 0;
 	virtual void TriggerRecompile() = 0;
+	virtual void Tick() = 0;
 };
 
