Restore code that was accidentally commented out while debugging.

#rnx

diff --git a/Engine/Source/Developer/Windows/LiveCoding/Private/LiveCodingModule.cpp b/Engine/Source/Developer/Windows/LiveCoding/Private/LiveCodingModule.cpp
--- a/Engine/Source/Developer/Windows/LiveCoding/Private/LiveCodingModule.cpp
+++ b/Engine/Source/Developer/Windows/LiveCoding/Private/LiveCodingModule.cpp
@@ -27,7 +27,7 @@ FLiveCodingModule::FLiveCodingModule()
 	, bStarted(false)
 {
 #if WITH_EDITOR
-//	GConfig->GetBool(TEXT("LiveCoding"), TEXT("Enabled"), bEnabled, GEditorPerProjectIni);
+	GConfig->GetBool(TEXT("LiveCoding"), TEXT("Enabled"), bEnabled, GEditorPerProjectIni);
 #endif
 }
 
