LiveCoding: Disable the editor compile button while a compile is in progress.

diff --git a/Engine/Source/Developer/Windows/LiveCoding/Private/External/LC_ClientCommandActions.cpp b/Engine/Source/Developer/Windows/LiveCoding/Private/External/LC_ClientCommandActions.cpp
--- a/Engine/Source/Developer/Windows/LiveCoding/Private/External/LC_ClientCommandActions.cpp
+++ b/Engine/Source/Developer/Windows/LiveCoding/Private/External/LC_ClientCommandActions.cpp
@@ -190,11 +190,18 @@ bool actions::LogOutput::Execute(CommandType* command, const DuplexPipe* pipe, v
 	return true;
 }
 
+// BEGIN EPIC MOD - Notification that compilation has finished
+extern bool GIsCompileActive;
+// END EPIC MOD
 
 bool actions::CompilationFinished::Execute(CommandType*, const DuplexPipe* pipe, void*)
 {
 	pipe->SendAck();
 
+	// BEGIN EPIC MOD - Notification that compilation has finished
+	GIsCompileActive = false;
+	// END EPIC MOD
+
 	// don't continue execution
 	return false;
 }
diff --git a/Engine/Source/Developer/Windows/LiveCoding/Private/LiveCodingModule.cpp b/Engine/Source/Developer/Windows/LiveCoding/Private/LiveCodingModule.cpp
--- a/Engine/Source/Developer/Windows/LiveCoding/Private/LiveCodingModule.cpp
+++ b/Engine/Source/Developer/Windows/LiveCoding/Private/LiveCodingModule.cpp
@@ -18,6 +18,7 @@ IMPLEMENT_MODULE(FLiveCodingModule, LiveCoding)
 
 #define LOCTEXT_NAMESPACE "LiveCodingModule"
 
+bool GIsCompileActive = false;
 FString GLiveCodingConsolePath;
 FString GLiveCodingConsoleArguments;
 
@@ -135,17 +136,26 @@ void FLiveCodingModule::ShowConsole()
 
 void FLiveCodingModule::TriggerRecompile()
 {
-	if (!bStarted)
+	if(!GIsCompileActive)
 	{
-		bShouldStart = true;
-		Tick();
-	}
-	if(bStarted)
-	{
-		LppTriggerRecompile();
+		if (!bStarted)
+		{
+			bShouldStart = true;
+			Tick();
+		}
+		if(bStarted)
+		{
+			LppTriggerRecompile();
+			GIsCompileActive = true;
+		}
 	}
 }
 
+bool FLiveCodingModule::IsCompiling() const
+{
+	return GIsCompileActive;
+}
+
 void FLiveCodingModule::Tick()
 {
 	if (bShouldStart && !bStarted)
diff --git a/Engine/Source/Developer/Windows/LiveCoding/Private/LiveCodingModule.h b/Engine/Source/Developer/Windows/LiveCoding/Private/LiveCodingModule.h
--- a/Engine/Source/Developer/Windows/LiveCoding/Private/LiveCodingModule.h
+++ b/Engine/Source/Developer/Windows/LiveCoding/Private/LiveCodingModule.h
@@ -23,6 +23,7 @@ public:
 	virtual bool IsEnabled() const override;
 	virtual void ShowConsole() override;
 	virtual void TriggerRecompile() override;
+	virtual bool IsCompiling() const override;
 	virtual void Tick() override;
 
 private:
diff --git a/Engine/Source/Developer/Windows/LiveCoding/Public/ILiveCodingModule.h b/Engine/Source/Developer/Windows/LiveCoding/Public/ILiveCodingModule.h
--- a/Engine/Source/Developer/Windows/LiveCoding/Public/ILiveCodingModule.h
+++ b/Engine/Source/Developer/Windows/LiveCoding/Public/ILiveCodingModule.h
@@ -14,6 +14,7 @@ public:
 	virtual bool IsEnabled() const = 0;
 	virtual void ShowConsole() = 0;
 	virtual void TriggerRecompile() = 0;
+	virtual bool IsCompiling() const = 0;
 	virtual void Tick() = 0;
 };
 
diff --git a/Engine/Source/Editor/LevelEditor/Private/LevelEditorActions.cpp b/Engine/Source/Editor/LevelEditor/Private/LevelEditorActions.cpp
--- a/Engine/Source/Editor/LevelEditor/Private/LevelEditorActions.cpp
+++ b/Engine/Source/Editor/LevelEditor/Private/LevelEditorActions.cpp
@@ -1050,6 +1050,14 @@ void FLevelEditorActionCallbacks::RecompileGameCode_Clicked()
 
 bool FLevelEditorActionCallbacks::Recompile_CanExecute()
 {
+#if WITH_LIVE_CODING
+	ILiveCodingModule* LiveCoding = FModuleManager::GetModulePtr<ILiveCodingModule>(LIVE_CODING_MODULE_NAME);
+	if (LiveCoding != nullptr && LiveCoding->IsEnabled())
+	{
+		return !LiveCoding->IsCompiling();
+	}
+#endif
+
 	// We can't recompile while in PIE
 	if (GEditor->bIsPlayWorldQueued || GEditor->PlayWorld)
 	{
