Move live coding settings back under the editor preferences window, since it's a user-specific thing.

diff --git a/Engine/Source/Developer/Windows/LiveCoding/Private/LiveCodingModule.cpp b/Engine/Source/Developer/Windows/LiveCoding/Private/LiveCodingModule.cpp
--- a/Engine/Source/Developer/Windows/LiveCoding/Private/LiveCodingModule.cpp
+++ b/Engine/Source/Developer/Windows/LiveCoding/Private/LiveCodingModule.cpp
@@ -60,7 +60,7 @@ void FLiveCodingModule::StartupModule()
 	ISettingsModule* SettingsModule = FModuleManager::GetModulePtr<ISettingsModule>("Settings");
 	if (SettingsModule != nullptr)
 	{
-		SettingsModule->RegisterSettings("Project", "Project", "Live Coding",
+		SettingsModule->RegisterSettings("Editor", "General", "Live Coding",
 			LOCTEXT("LiveCodingSettingsName", "Live Coding"),
 			LOCTEXT("LiveCodintSettingsDescription", "Settings for recompiling C++ code while the engine is running."),
 			GetMutableDefault<ULiveCodingSettings>()
