LiveCoding: Suppress dev channel log messages from console. Add a new server log channel that they are written to instead.

diff --git a/Engine/Source/Developer/Windows/LiveCodingServer/Private/LiveCodingServerModule.cpp b/Engine/Source/Developer/Windows/LiveCodingServer/Private/LiveCodingServerModule.cpp
--- a/Engine/Source/Developer/Windows/LiveCodingServer/Private/LiveCodingServerModule.cpp
+++ b/Engine/Source/Developer/Windows/LiveCodingServer/Private/LiveCodingServerModule.cpp
@@ -8,27 +8,45 @@
 
 IMPLEMENT_MODULE(FLiveCodingServerModule, LiveCodingServer)
 
+DEFINE_LOG_CATEGORY_STATIC(LogLiveCodingServer, Display, All);
+
 static void ServerOutputHandler(logging::Channel::Enum Channel, logging::Type::Enum Type, const wchar_t* const Text)
 {
-	ELiveCodingLogVerbosity Verbosity;
+	FString TrimText = FString(Text).TrimEnd();
 	switch (Type)
 	{
-	case logging::Type::LOG_SUCCESS:
-		Verbosity = ELiveCodingLogVerbosity::Success;
-		break;
 	case logging::Type::LOG_ERROR:
-		Verbosity = ELiveCodingLogVerbosity::Failure;
+		UE_LOG(LogLiveCodingServer, Error, TEXT("%s"), *TrimText);
 		break;
 	case logging::Type::LOG_WARNING:
-		Verbosity = ELiveCodingLogVerbosity::Warning;
+		UE_LOG(LogLiveCodingServer, Warning, TEXT("%s"), *TrimText);
 		break;
 	default:
-		Verbosity = ELiveCodingLogVerbosity::Info;
+		UE_LOG(LogLiveCodingServer, Display, TEXT("%s"), *TrimText);
 		break;
 	}
-	GLiveCodingServer->GetLogOutputDelegate().ExecuteIfBound(Verbosity, Text);
-}
 
+	if (Channel == logging::Channel::USER)
+	{
+		ELiveCodingLogVerbosity Verbosity;
+		switch (Type)
+		{
+		case logging::Type::LOG_SUCCESS:
+			Verbosity = ELiveCodingLogVerbosity::Success;
+			break;
+		case logging::Type::LOG_ERROR:
+			Verbosity = ELiveCodingLogVerbosity::Failure;
+			break;
+		case logging::Type::LOG_WARNING:
+			Verbosity = ELiveCodingLogVerbosity::Warning;
+			break;
+		default:
+			Verbosity = ELiveCodingLogVerbosity::Info;
+			break;
+		}
+		GLiveCodingServer->GetLogOutputDelegate().ExecuteIfBound(Verbosity, Text);
+	}
+}
 
 void FLiveCodingServerModule::StartupModule()
 {
