LiveCoding: Tidy up console output.

diff --git a/Engine/Source/Developer/Windows/LiveCodingServer/Private/External/LC_ServerCommandThread.cpp b/Engine/Source/Developer/Windows/LiveCodingServer/Private/External/LC_ServerCommandThread.cpp
--- a/Engine/Source/Developer/Windows/LiveCodingServer/Private/External/LC_ServerCommandThread.cpp
+++ b/Engine/Source/Developer/Windows/LiveCodingServer/Private/External/LC_ServerCommandThread.cpp
@@ -590,14 +590,18 @@ void ServerCommandThread::CompileChanges(bool didAllProcessesMakeProgress)
 	// EPIC REMOVED: g_theApp.GetMainFrame()->SetBusy(true);
 	// EPIC REMOVED: g_theApp.GetMainFrame()->ChangeStatusBarText(L"Creating patch...");
 
+	telemetry::Scope scope("Creating patch");
+
+	// EPIC REMOVED: g_theApp.GetMainFrame()->OnCompilationStart();
+
+	LC_LOG_USER("---------- Creating patch ----------");
+
 	// BEGIN EPIC MOD - Hook for the compiler
 	GLiveCodingServer->GetCompileStartedDelegate().ExecuteIfBound();
 
 	const ILiveCodingServer::FCompileDelegate& CompileDelegate = GLiveCodingServer->GetCompileDelegate();
 	if (CompileDelegate.IsBound())
 	{
-		LC_LOG_USER("---------- Starting build ----------");
-
 		TArray<FString> Targets;
 		for (LiveProcess* liveProcess : m_liveProcesses)
 		{
@@ -645,12 +649,6 @@ void ServerCommandThread::CompileChanges(bool didAllProcessesMakeProgress)
 	GLiveCodingServer->GetStatusChangeDelegate().ExecuteIfBound(L"Creating patch...");
 	// END EPIC MOD
 
-	telemetry::Scope scope("Creating patch");
-
-	// EPIC REMOVED: g_theApp.GetMainFrame()->OnCompilationStart();
-
-	LC_LOG_USER("---------- Creating patch ----------");
-
 	// recompile files, if any
 	const size_t count = m_liveModules.size();
 	if (count == 0u)
