LiveCoding: Tweaks to live coding settings.

* Compile dropdown now just contains a checkbox for enabling live coding.
* Compile dropdown has a link to open the editor preferences window.
* Live coding can be enabled from the editor settings window.
* Option to show the console on startup is now folded into the startup mode.
* Options in the editor settings window are greyed out unless live coding is enabled.

diff --git a/Engine/Source/Developer/Windows/LiveCoding/Private/LiveCodingModule.cpp b/Engine/Source/Developer/Windows/LiveCoding/Private/LiveCodingModule.cpp
--- a/Engine/Source/Developer/Windows/LiveCoding/Private/LiveCodingModule.cpp
+++ b/Engine/Source/Developer/Windows/LiveCoding/Private/LiveCodingModule.cpp
@@ -23,28 +23,21 @@ FString GLiveCodingConsolePath;
 FString GLiveCodingConsoleArguments;
 
 FLiveCodingModule::FLiveCodingModule()
-	: bEnabled(false)
-	, bShouldStart(false)
+	: bEnabledLastTick(false)
 	, bStarted(false)
 {
-#if WITH_EDITOR
-	GConfig->GetBool(TEXT("LiveCoding"), TEXT("Enabled"), bEnabled, GEditorPerProjectIni);
-#endif
 }
 
 void FLiveCodingModule::StartupModule()
 {
-	if(FParse::Param(FCommandLine::Get(), TEXT("LiveCoding")))
-	{
-		bEnabled = true;
-	}
+	Settings = GetMutableDefault<ULiveCodingSettings>();
 
 	IConsoleManager& ConsoleManager = IConsoleManager::Get();
 
 	EnableCommand = ConsoleManager.RegisterConsoleCommand(
 		TEXT("LiveCoding"),
 		TEXT("Enables live coding support"),
-		FConsoleCommandDelegate::CreateRaw(this, &FLiveCodingModule::Enable, true),
+		FConsoleCommandDelegate::CreateRaw(this, &FLiveCodingModule::EnableForSession, true),
 		ECVF_Cheat
 	);
 
@@ -67,18 +60,29 @@ void FLiveCodingModule::StartupModule()
 		);
 	}
 
-	const ULiveCodingSettings* Settings = GetDefault<ULiveCodingSettings>();
-	if (bEnabled && Settings->StartupMode == ELiveCodingStartupMode::Automatic)
+	extern void Startup(Windows::HINSTANCE hInstance);
+	Startup(hInstance);
+
+	if (Settings->bEnabled)
 	{
-		bShouldStart = true;
-		if(!Settings->bShowConsole)
+		if(Settings->Startup == ELiveCodingStartupMode::Automatic)
+		{
+			StartLiveCoding();
+			ShowConsole();
+		}
+		else if(Settings->Startup == ELiveCodingStartupMode::AutomaticButHidden)
 		{
 			GLiveCodingConsoleArguments = L"-Hidden";
+			StartLiveCoding();
 		}
 	}
 
-	extern void Startup(Windows::HINSTANCE hInstance);
-	Startup(hInstance);
+	if(FParse::Param(FCommandLine::Get(), TEXT("LiveCoding")))
+	{
+		StartLiveCoding();
+	}
+
+	bEnabledLastTick = Settings->bEnabled;
 }
 
 void FLiveCodingModule::ShutdownModule()
@@ -93,20 +97,30 @@ void FLiveCodingModule::ShutdownModule()
 	ConsoleManager.UnregisterConsoleObject(EnableCommand);
 }
 
-void FLiveCodingModule::Enable(bool bInEnabled)
+void FLiveCodingModule::EnableByDefault(bool bEnable)
 {
-	if (bEnabled != bInEnabled)
-	{
-		bEnabled = bInEnabled;
-#if WITH_EDITOR
-		GConfig->SetBool(TEXT("LiveCoding"), TEXT("Enabled"), bEnabled, GEditorPerProjectIni);
-#endif
+	Settings->bEnabled = bEnable;
+	EnableForSession(bEnable);
+}
+
+bool FLiveCodingModule::IsEnabledByDefault() const
+{
+	return Settings->bEnabled;
+}
 
-		if (bEnabled)
+void FLiveCodingModule::EnableForSession(bool bEnable)
+{
+	if (bEnable)
+	{
+		if(!bStarted)
 		{
+			StartLiveCoding();
 			ShowConsole();
 		}
-		else if(bStarted)
+	}
+	else 
+	{
+		if(bStarted)
 		{
 			UE_LOG(LogLiveCoding, Display, TEXT("Console will be hidden but remain running in the background. Restart to disable completely."));
 			LppSetActive(false);
@@ -115,34 +129,26 @@ void FLiveCodingModule::Enable(bool bInEnabled)
 	}
 }
 
-bool FLiveCodingModule::IsEnabled() const
+bool FLiveCodingModule::IsEnabledForSession() const
 {
-	return bEnabled;
+	return bStarted;
 }
 
 void FLiveCodingModule::ShowConsole()
 {
 	if (bStarted)
 	{
-		LppSetVisible(bEnabled);
-		LppSetActive(bEnabled);
+		LppSetVisible(true);
+		LppSetActive(true);
 		LppShowConsole();
 	}
-	else
-	{
-		bShouldStart = true;
-	}
 }
 
-void FLiveCodingModule::TriggerRecompile()
+void FLiveCodingModule::Compile()
 {
 	if(!GIsCompileActive)
 	{
-		if (!bStarted)
-		{
-			bShouldStart = true;
-			Tick();
-		}
+		EnableForSession(true);
 		if(bStarted)
 		{
 			LppTriggerRecompile();
@@ -158,54 +164,54 @@ bool FLiveCodingModule::IsCompiling() const
 
 void FLiveCodingModule::Tick()
 {
-	if (bShouldStart && !bStarted)
+	if (Settings->bEnabled != bEnabledLastTick && Settings->Startup != ELiveCodingStartupMode::Manual)
 	{
-		if(StartLiveCoding())
-		{
-			bStarted = true;
-		}
-		else
-		{
-			bShouldStart = false;
-		}
+		EnableForSession(Settings->bEnabled);
+		bEnabledLastTick = Settings->bEnabled;
 	}
 }
 
 bool FLiveCodingModule::StartLiveCoding()
 {
-	// Setup the console path
-	GLiveCodingConsolePath = ConsolePathVariable->GetString();
-	if (!FPaths::FileExists(GLiveCodingConsolePath))
+	if(!bStarted)
 	{
-		UE_LOG(LogLiveCoding, Error, TEXT("Unable to start live coding session. Missing executable '%s'. Use the LiveCoding.ConsolePath console variable to modify."), *GLiveCodingConsolePath);
-		return false;
-	}
+		// Setup the console path
+		GLiveCodingConsolePath = ConsolePathVariable->GetString();
+		if (!FPaths::FileExists(GLiveCodingConsolePath))
+		{
+			UE_LOG(LogLiveCoding, Error, TEXT("Unable to start live coding session. Missing executable '%s'. Use the LiveCoding.ConsolePath console variable to modify."), *GLiveCodingConsolePath);
+			return false;
+		}
 
-	UE_LOG(LogLiveCoding, Display, TEXT("Starting LiveCoding"));
+		UE_LOG(LogLiveCoding, Display, TEXT("Starting LiveCoding"));
 
-	// Enable external build system
-	LppUseExternalBuildSystem();
+		// Enable external build system
+		LppUseExternalBuildSystem();
 
-	// Enable the server
-	FString ProcessGroup = FString::Printf(TEXT("UE4_%s_0x%08x"), FApp::GetProjectName(), GetTypeHash(FPaths::ProjectDir()));
-	LppRegisterProcessGroup(TCHAR_TO_ANSI(*ProcessGroup));
+		// Enable the server
+		FString ProcessGroup = FString::Printf(TEXT("UE4_%s_0x%08x"), FApp::GetProjectName(), GetTypeHash(FPaths::ProjectDir()));
+		LppRegisterProcessGroup(TCHAR_TO_ANSI(*ProcessGroup));
 
-	// Build the command line
-	FString Arguments;
-	Arguments += FString::Printf(TEXT("%s"), FPlatformMisc::GetUBTPlatform());
-	Arguments += FString::Printf(TEXT(" %s"), EBuildConfigurations::ToString(FApp::GetBuildConfiguration()));
-	Arguments += FString::Printf(TEXT(" -TargetType=%s"), FPlatformMisc::GetUBTTarget());
-	if(FPaths::IsProjectFilePathSet())
-	{
-		Arguments += FString::Printf(TEXT(" -Project=\"%s\""), *FPaths::ConvertRelativePathToFull(FPaths::GetProjectFilePath()));
-	}
-	LppSetBuildArguments(*Arguments);
+		// Build the command line
+		FString Arguments;
+		Arguments += FString::Printf(TEXT("%s"), FPlatformMisc::GetUBTPlatform());
+		Arguments += FString::Printf(TEXT(" %s"), EBuildConfigurations::ToString(FApp::GetBuildConfiguration()));
+		Arguments += FString::Printf(TEXT(" -TargetType=%s"), FPlatformMisc::GetUBTTarget());
+		if(FPaths::IsProjectFilePathSet())
+		{
+			Arguments += FString::Printf(TEXT(" -Project=\"%s\""), *FPaths::ConvertRelativePathToFull(FPaths::GetProjectFilePath()));
+		}
+		LppSetBuildArguments(*Arguments);
+
+		// Configure all the current modules
+		UpdateModules();
 
-	// Configure all the current modules
-	UpdateModules();
+		// Register a delegate to listen for new modules loaded from this point onwards
+		ModulesChangedDelegateHandle = FModuleManager::Get().OnModulesChanged().AddRaw(this, &FLiveCodingModule::OnModulesChanged);
 
-	// Register a delegate to listen for new modules loaded from this point onwards
-	ModulesChangedDelegateHandle = FModuleManager::Get().OnModulesChanged().AddRaw(this, &FLiveCodingModule::OnModulesChanged);
+		// Mark it as started
+		bStarted = true;
+	}
 	return true;
 }
 
@@ -279,7 +285,6 @@ void FLiveCodingModule::ConfigureModule(const FName& Name, bool bIsProjectModule
 
 bool FLiveCodingModule::ShouldEnableModule(const FName& Name, bool bIsProjectModule, const FString& FullFilePath) const
 {
-	const ULiveCodingSettings* Settings = GetDefault<ULiveCodingSettings>();
 	if (Settings->ExcludeSpecificModules.Contains(Name))
 	{
 		return false;
diff --git a/Engine/Source/Developer/Windows/LiveCoding/Private/LiveCodingModule.h b/Engine/Source/Developer/Windows/LiveCoding/Private/LiveCodingModule.h
--- a/Engine/Source/Developer/Windows/LiveCoding/Private/LiveCodingModule.h
+++ b/Engine/Source/Developer/Windows/LiveCoding/Private/LiveCodingModule.h
@@ -8,6 +8,7 @@
 
 struct IConsoleCommand;
 class IConsoleVariable;
+class ULiveCodingSettings;
 
 class FLiveCodingModule final : public ILiveCodingModule
 {
@@ -19,16 +20,18 @@ public:
 	virtual void ShutdownModule() override;
 
 	// ILiveCodingModule implementation
-	virtual void Enable(bool bInEnabled) override;
-	virtual bool IsEnabled() const override;
+	virtual void EnableByDefault(bool bInEnabled) override;
+	virtual bool IsEnabledByDefault() const override;
+	virtual void EnableForSession(bool bInEnabled) override;
+	virtual bool IsEnabledForSession() const override;
 	virtual void ShowConsole() override;
-	virtual void TriggerRecompile() override;
+	virtual void Compile() override;
 	virtual bool IsCompiling() const override;
 	virtual void Tick() override;
 
 private:
-	bool bEnabled;
-	bool bShouldStart;
+	ULiveCodingSettings* Settings;
+	bool bEnabledLastTick;
 	bool bStarted;
 	TSet<FString> EnabledModules;
 
diff --git a/Engine/Source/Developer/Windows/LiveCoding/Private/LiveCodingSettings.h b/Engine/Source/Developer/Windows/LiveCoding/Private/LiveCodingSettings.h
--- a/Engine/Source/Developer/Windows/LiveCoding/Private/LiveCodingSettings.h
+++ b/Engine/Source/Developer/Windows/LiveCoding/Private/LiveCodingSettings.h
@@ -9,8 +9,9 @@
 UENUM()
 enum class ELiveCodingStartupMode : uint8
 {
-	Automatic,
-	Manual,
+	Automatic UMETA(DisplayName = "Start automatically and show console"),
+	AutomaticButHidden UMETA(DisplayName = "Start automatically but hide console until summoned"),
+	Manual UMETA(DisplayName = "Manual"),
 };
 
 UCLASS(config=EditorPerProjectUserSettings, meta=(DisplayName="Live Coding"))
@@ -19,28 +20,28 @@ class ULiveCodingSettings : public UObject
     GENERATED_BODY()
 
 public:
-    UPROPERTY(config, EditAnywhere, Category=Startup, Meta=(ConfigRestartRequired=true))
-    ELiveCodingStartupMode StartupMode;
-	
-    UPROPERTY(config, EditAnywhere, Category=Startup, Meta=(ConfigRestartRequired=true))
-    bool bShowConsole;
+    UPROPERTY(config, EditAnywhere, Category=General, Meta=(DisplayName="Enable Live Coding"))
+    bool bEnabled;
 
-    UPROPERTY(config, EditAnywhere, Category=Modules, Meta=(ConfigRestartRequired=true))
+    UPROPERTY(config, EditAnywhere, Category=General, Meta=(ConfigRestartRequired=true, EditCondition="bEnabled"))
+    ELiveCodingStartupMode Startup;
+
+	UPROPERTY(config, EditAnywhere, Category=Modules, Meta=(ConfigRestartRequired=true, EditCondition="bEnabled"))
     bool bIncludeEngineModules;
 
-    UPROPERTY(config, EditAnywhere, Category=Modules, Meta=(ConfigRestartRequired=true))
+    UPROPERTY(config, EditAnywhere, Category=Modules, Meta=(ConfigRestartRequired=true, EditCondition="bEnabled"))
     bool bIncludeEnginePluginModules;
 
-    UPROPERTY(config, EditAnywhere, Category=Modules, Meta=(ConfigRestartRequired=true))
+    UPROPERTY(config, EditAnywhere, Category=Modules, Meta=(ConfigRestartRequired=true, EditCondition="bEnabled"))
     bool bIncludeProjectModules;
 
-    UPROPERTY(config, EditAnywhere, Category=Modules, Meta=(ConfigRestartRequired=true))
+    UPROPERTY(config, EditAnywhere, Category=Modules, Meta=(ConfigRestartRequired=true, EditCondition="bEnabled"))
     bool bIncludeProjectPluginModules;
 
-    UPROPERTY(config, EditAnywhere, Category=Modules, Meta=(ConfigRestartRequired=true))
+    UPROPERTY(config, EditAnywhere, Category=Modules, Meta=(ConfigRestartRequired=true, EditCondition="bEnabled"))
     TArray<FName> IncludeSpecificModules;
 
-    UPROPERTY(config, EditAnywhere, Category=Modules, Meta=(ConfigRestartRequired=true))
+    UPROPERTY(config, EditAnywhere, Category=Modules, Meta=(ConfigRestartRequired=true, EditCondition="bEnabled"))
     TArray<FName> ExcludeSpecificModules;
 
 	ULiveCodingSettings(const FObjectInitializer& Initializer);
diff --git a/Engine/Source/Developer/Windows/LiveCoding/Public/ILiveCodingModule.h b/Engine/Source/Developer/Windows/LiveCoding/Public/ILiveCodingModule.h
--- a/Engine/Source/Developer/Windows/LiveCoding/Public/ILiveCodingModule.h
+++ b/Engine/Source/Developer/Windows/LiveCoding/Public/ILiveCodingModule.h
@@ -10,10 +10,12 @@
 class ILiveCodingModule : public IModuleInterface
 {
 public:
-	virtual void Enable(bool bEnabled) = 0;
-	virtual bool IsEnabled() const = 0;
+	virtual void EnableByDefault(bool bEnabled) = 0;
+	virtual bool IsEnabledByDefault() const = 0;
+	virtual void EnableForSession(bool bEnabled) = 0;
+	virtual bool IsEnabledForSession() const = 0;
 	virtual void ShowConsole() = 0;
-	virtual void TriggerRecompile() = 0;
+	virtual void Compile() = 0;
 	virtual bool IsCompiling() const = 0;
 	virtual void Tick() = 0;
 };
diff --git a/Engine/Source/Editor/LevelEditor/Private/LevelEditor.cpp b/Engine/Source/Editor/LevelEditor/Private/LevelEditor.cpp
--- a/Engine/Source/Editor/LevelEditor/Private/LevelEditor.cpp
+++ b/Engine/Source/Editor/LevelEditor/Private/LevelEditor.cpp
@@ -657,21 +657,25 @@ void FLevelEditorModule::BindGlobalLevelEditorCommands()
 
 #if WITH_LIVE_CODING
 	ActionList.MapAction( 
-		Commands.HotReloadMode_Legacy, 
-		FExecuteAction::CreateStatic( &FLevelEditorActionCallbacks::HotReload_EnableLiveCoding, false ),
+		Commands.LiveCoding_Enable, 
+		FExecuteAction::CreateStatic( &FLevelEditorActionCallbacks::LiveCoding_ToggleEnabled ),
 		FCanExecuteAction(),
-		FIsActionChecked::CreateStatic( &FLevelEditorActionCallbacks::HotReload_IsLegacyMode ) );
+		FIsActionChecked::CreateStatic( &FLevelEditorActionCallbacks::LiveCoding_IsEnabled ) );
 
 	ActionList.MapAction( 
-		Commands.HotReloadMode_LiveCoding, 
-		FExecuteAction::CreateStatic( &FLevelEditorActionCallbacks::HotReload_EnableLiveCoding, true ),
-		FCanExecuteAction(),
-		FIsActionChecked::CreateStatic( &FLevelEditorActionCallbacks::HotReload_IsLiveCodingMode ) );
+		Commands.LiveCoding_StartSession, 
+		FExecuteAction::CreateStatic( &FLevelEditorActionCallbacks::LiveCoding_StartSession_Clicked ),
+		FCanExecuteAction::CreateStatic( &FLevelEditorActionCallbacks::LiveCoding_CanStartSession ) );
+
+	ActionList.MapAction( 
+		Commands.LiveCoding_ShowConsole, 
+		FExecuteAction::CreateStatic( &FLevelEditorActionCallbacks::LiveCoding_ShowConsole_Clicked ),
+		FCanExecuteAction::CreateStatic( &FLevelEditorActionCallbacks::LiveCoding_CanShowConsole ) );
 
 	ActionList.MapAction( 
-		Commands.HotReload_ShowConsole, 
-		FExecuteAction::CreateStatic( &FLevelEditorActionCallbacks::HotReload_ShowConsole_Clicked ),
-		FCanExecuteAction::CreateStatic( &FLevelEditorActionCallbacks::HotReload_IsLiveCodingMode ) );
+		Commands.LiveCoding_Settings, 
+		FExecuteAction::CreateStatic( &FLevelEditorActionCallbacks::LiveCoding_Settings_Clicked ),
+		FCanExecuteAction());
 
 #endif
 
diff --git a/Engine/Source/Editor/LevelEditor/Private/LevelEditorActions.cpp b/Engine/Source/Editor/LevelEditor/Private/LevelEditorActions.cpp
--- a/Engine/Source/Editor/LevelEditor/Private/LevelEditorActions.cpp
+++ b/Engine/Source/Editor/LevelEditor/Private/LevelEditorActions.cpp
@@ -90,6 +90,7 @@
 #include "ActorGroupingUtils.h"
 #include "LevelUtils.h"
 #include "ISceneOutliner.h"
+#include "ISettingsModule.h"
 
 #if WITH_LIVE_CODING
 #include "ILiveCodingModule.h"
@@ -1032,9 +1033,9 @@ void FLevelEditorActionCallbacks::RecompileGameCode_Clicked()
 {
 #if WITH_LIVE_CODING
 	ILiveCodingModule* LiveCoding = FModuleManager::GetModulePtr<ILiveCodingModule>(LIVE_CODING_MODULE_NAME);
-	if (LiveCoding != nullptr && LiveCoding->IsEnabled())
+	if (LiveCoding != nullptr && LiveCoding->IsEnabledForSession())
 	{
-		LiveCoding->TriggerRecompile();
+		LiveCoding->Compile();
 		return;
 	}
 #endif
@@ -1052,7 +1053,7 @@ bool FLevelEditorActionCallbacks::Recompile_CanExecute()
 {
 #if WITH_LIVE_CODING
 	ILiveCodingModule* LiveCoding = FModuleManager::GetModulePtr<ILiveCodingModule>(LIVE_CODING_MODULE_NAME);
-	if (LiveCoding != nullptr && LiveCoding->IsEnabled())
+	if (LiveCoding != nullptr && LiveCoding->IsEnabledForSession())
 	{
 		return !LiveCoding->IsCompiling();
 	}
@@ -1071,27 +1072,37 @@ bool FLevelEditorActionCallbacks::Recompile_CanExecute()
 }
 
 #if WITH_LIVE_CODING
-void FLevelEditorActionCallbacks::HotReload_EnableLiveCoding(bool bEnable)
+void FLevelEditorActionCallbacks::LiveCoding_ToggleEnabled()
 {
 	ILiveCodingModule* LiveCoding = FModuleManager::GetModulePtr<ILiveCodingModule>(LIVE_CODING_MODULE_NAME);
 	if (LiveCoding != nullptr)
 	{
-		LiveCoding->Enable(bEnable);
+		LiveCoding->EnableByDefault(!LiveCoding->IsEnabledByDefault());
 	}
 }
 
-bool FLevelEditorActionCallbacks::HotReload_IsLegacyMode( )
+bool FLevelEditorActionCallbacks::LiveCoding_IsEnabled( )
 {
-	return !HotReload_IsLiveCodingMode();
+	ILiveCodingModule* LiveCoding = FModuleManager::GetModulePtr<ILiveCodingModule>(LIVE_CODING_MODULE_NAME);
+	return LiveCoding != nullptr && LiveCoding->IsEnabledByDefault();
+}
+
+void FLevelEditorActionCallbacks::LiveCoding_StartSession_Clicked()
+{
+	ILiveCodingModule* LiveCoding = FModuleManager::GetModulePtr<ILiveCodingModule>(LIVE_CODING_MODULE_NAME);
+	if (LiveCoding!= nullptr)
+	{
+		LiveCoding->EnableForSession(true);
+	}
 }
 
-bool FLevelEditorActionCallbacks::HotReload_IsLiveCodingMode( )
+bool FLevelEditorActionCallbacks::LiveCoding_CanStartSession()
 {
 	ILiveCodingModule* LiveCoding = FModuleManager::GetModulePtr<ILiveCodingModule>(LIVE_CODING_MODULE_NAME);
-	return LiveCoding != nullptr && LiveCoding->IsEnabled();
+	return LiveCoding != nullptr && LiveCoding->IsEnabledByDefault() && !LiveCoding->IsEnabledForSession();
 }
 
-void FLevelEditorActionCallbacks::HotReload_ShowConsole_Clicked()
+void FLevelEditorActionCallbacks::LiveCoding_ShowConsole_Clicked()
 {
 	ILiveCodingModule* LiveCoding = FModuleManager::GetModulePtr<ILiveCodingModule>(LIVE_CODING_MODULE_NAME);
 	if (LiveCoding!= nullptr)
@@ -1099,6 +1110,17 @@ void FLevelEditorActionCallbacks::HotReload_ShowConsole_Clicked()
 		LiveCoding->ShowConsole();
 	}
 }
+
+bool FLevelEditorActionCallbacks::LiveCoding_CanShowConsole()
+{
+	ILiveCodingModule* LiveCoding = FModuleManager::GetModulePtr<ILiveCodingModule>(LIVE_CODING_MODULE_NAME);
+	return LiveCoding!= nullptr && LiveCoding->IsEnabledForSession();
+}
+
+void FLevelEditorActionCallbacks::LiveCoding_Settings_Clicked()
+{
+	FModuleManager::LoadModuleChecked<ISettingsModule>("Settings").ShowViewer("Editor", "General", "Live Coding");
+}
 #endif
 
 void FLevelEditorActionCallbacks::ConnectToSourceControl_Clicked()
@@ -3148,9 +3170,10 @@ void FLevelEditorCommands::RegisterCommands()
 	UI_COMMAND( RecompileGameCode, "Recompile Game Code", "Recompiles and reloads C++ code for game systems on the fly", EUserInterfaceActionType::Button, FInputChord( EKeys::P, EModifierKey::Alt | EModifierKey::Control | EModifierKey::Shift ) );
 
 #if WITH_LIVE_CODING
-	UI_COMMAND( HotReloadMode_Legacy, "Legacy", "Use the legacy hot reload system", EUserInterfaceActionType::RadioButton, FInputChord() );
-	UI_COMMAND( HotReloadMode_LiveCoding, "Live Coding (Experimental)", "Hot-patches C++ function changes into the current process. Currently does not allow class layout changes.", EUserInterfaceActionType::RadioButton, FInputChord() );
-	UI_COMMAND( HotReload_ShowConsole, "Show Console", "Displays the live coding console window.", EUserInterfaceActionType::Button, FInputChord() );
+	UI_COMMAND( LiveCoding_Enable, "Enable Live Coding (Experimental)", "Hot-patches C++ function changes into the current process. Currently does not allow class layout changes.", EUserInterfaceActionType::ToggleButton, FInputChord() );
+	UI_COMMAND( LiveCoding_StartSession, "Start Session", "Starts a live coding session.", EUserInterfaceActionType::Button, FInputChord() );
+	UI_COMMAND( LiveCoding_ShowConsole, "Show Console", "Displays the live coding console window.", EUserInterfaceActionType::Button, FInputChord() );
+	UI_COMMAND( LiveCoding_Settings, "Settings...", "Open the live coding settings", EUserInterfaceActionType::Button, FInputChord() );
 #endif
 
 	UI_COMMAND( EditAsset, "Edit Asset", "Edits the asset associated with the selected actor", EUserInterfaceActionType::Button, FInputChord( EKeys::E, EModifierKey::Control ) );
diff --git a/Engine/Source/Editor/LevelEditor/Private/LevelEditorToolBar.cpp b/Engine/Source/Editor/LevelEditor/Private/LevelEditorToolBar.cpp
--- a/Engine/Source/Editor/LevelEditor/Private/LevelEditorToolBar.cpp
+++ b/Engine/Source/Editor/LevelEditor/Private/LevelEditorToolBar.cpp
@@ -1889,16 +1889,17 @@ TSharedRef< SWidget > FLevelEditorToolBar::GenerateCompileMenuContent( TSharedRe
 	const bool bShouldCloseWindowAfterMenuSelection = true;
 	FMenuBuilder MenuBuilder( bShouldCloseWindowAfterMenuSelection, InCommandList );
 
-	MenuBuilder.BeginSection("HotReloadMode", LOCTEXT( "HotReloadMode", "Mode" ) );
+	MenuBuilder.BeginSection("LiveCodingMode", LOCTEXT( "LiveCodingMode", "General" ) );
 	{
-		MenuBuilder.AddMenuEntry( FLevelEditorCommands::Get().HotReloadMode_Legacy );
-		MenuBuilder.AddMenuEntry( FLevelEditorCommands::Get().HotReloadMode_LiveCoding );
+		MenuBuilder.AddMenuEntry( FLevelEditorCommands::Get().LiveCoding_Enable );
 	}
 	MenuBuilder.EndSection();
 
-	MenuBuilder.BeginSection("HotReloadActions", LOCTEXT( "HotReloadActions", "Actions" ) );
+	MenuBuilder.BeginSection("LiveCodingActions", LOCTEXT( "LiveCodingActions", "Actions" ) );
 	{
-		MenuBuilder.AddMenuEntry( FLevelEditorCommands::Get().HotReload_ShowConsole );
+		MenuBuilder.AddMenuEntry( FLevelEditorCommands::Get().LiveCoding_StartSession );
+		MenuBuilder.AddMenuEntry( FLevelEditorCommands::Get().LiveCoding_ShowConsole );
+		MenuBuilder.AddMenuEntry( FLevelEditorCommands::Get().LiveCoding_Settings );
 	}
 	MenuBuilder.EndSection();
 
diff --git a/Engine/Source/Editor/LevelEditor/Public/LevelEditorActions.h b/Engine/Source/Editor/LevelEditor/Public/LevelEditorActions.h
--- a/Engine/Source/Editor/LevelEditor/Public/LevelEditorActions.h
+++ b/Engine/Source/Editor/LevelEditor/Public/LevelEditorActions.h
@@ -107,9 +107,10 @@ public:
 	TSharedPtr< FUICommandInfo > RecompileGameCode;
 
 #if WITH_LIVE_CODING
-	TSharedPtr< FUICommandInfo > HotReloadMode_Legacy;
-	TSharedPtr< FUICommandInfo > HotReloadMode_LiveCoding;
-	TSharedPtr< FUICommandInfo > HotReload_ShowConsole;
+	TSharedPtr< FUICommandInfo > LiveCoding_Enable;
+	TSharedPtr< FUICommandInfo > LiveCoding_StartSession;
+	TSharedPtr< FUICommandInfo > LiveCoding_ShowConsole;
+	TSharedPtr< FUICommandInfo > LiveCoding_Settings;
 #endif
 
 	/**
@@ -853,22 +854,37 @@ public:
 	/**
 	 * Enables live coding mode
 	 */
-	static void HotReload_EnableLiveCoding(bool bEnable);
+	static void LiveCoding_ToggleEnabled();
 
 	/**
-	 * Determines if hot reload is set to use legacy mode.
+	 * Determines if live coding is enabled
 	 */
-	static bool HotReload_IsLegacyMode();
+	static bool LiveCoding_IsEnabled();
 
 	/**
-	 * Determines if live coding is enabled
+	 * Starts live coding (in manual mode)
+	 */
+	static void LiveCoding_StartSession_Clicked();
+
+	/**
+	 * Determines whether we can manually start live coding for the current session
 	 */
-	static bool HotReload_IsLiveCodingMode();
+	static bool LiveCoding_CanStartSession();
 
 	/**
 	 * Shows the console
 	 */
-	static void HotReload_ShowConsole_Clicked();
+	static void LiveCoding_ShowConsole_Clicked();
+
+	/**
+	 * Determines whether the console can be shown
+	 */
+	static bool LiveCoding_CanShowConsole();
+
+	/**
+	 * Shows the settings panel
+	 */
+	static void LiveCoding_Settings_Clicked();
 #endif
 
 	/**
