LiveCoding: Make properties for toggling on/off for engine modules hidden in installed builds.

diff --git a/Engine/Source/Developer/Windows/LiveCoding/Private/LiveCodingSettings.cpp b/Engine/Source/Developer/Windows/LiveCoding/Private/LiveCodingSettings.cpp
--- a/Engine/Source/Developer/Windows/LiveCoding/Private/LiveCodingSettings.cpp
+++ b/Engine/Source/Developer/Windows/LiveCoding/Private/LiveCodingSettings.cpp
@@ -1,9 +1,22 @@
 // Copyright 1998-2019 Epic Games, Inc. All Rights Reserved.
 
 #include "LiveCodingSettings.h"
+#include "Misc/App.h"
 
 ULiveCodingSettings::ULiveCodingSettings(const FObjectInitializer& ObjectInitializer) : Super(ObjectInitializer)
 {
+	UProperty* EngineModulesProperty = StaticClass()->FindPropertyByName("bIncludeEngineModules");
+	check(EngineModulesProperty != nullptr);
+
+	UProperty* EnginePluginModulesProperty = StaticClass()->FindPropertyByName("bIncludeEnginePluginModules");
+	check(EnginePluginModulesProperty != nullptr);
+
+	if (FApp::IsEngineInstalled())
+	{
+		EngineModulesProperty->ClearPropertyFlags(CPF_Edit);
+		EnginePluginModulesProperty->ClearPropertyFlags(CPF_Edit);
+	}
+
 	bIncludeProjectModules = true;
 	bIncludeProjectPluginModules = true;
 }
