LiveCoding: Fix setting for enabling/disabling live coding not being saved when modified via the toolbar.

diff --git a/Engine/Source/Developer/Windows/LiveCoding/Private/LiveCodingModule.cpp b/Engine/Source/Developer/Windows/LiveCoding/Private/LiveCodingModule.cpp
--- a/Engine/Source/Developer/Windows/LiveCoding/Private/LiveCodingModule.cpp
+++ b/Engine/Source/Developer/Windows/LiveCoding/Private/LiveCodingModule.cpp
@@ -12,6 +12,7 @@
 #include "Misc/ConfigCacheIni.h"
 #include "LiveCodingSettings.h"
 #include "ISettingsModule.h"
+#include "ISettingsSection.h"
 #include "Windows/WindowsHWrapper.h"
 
 IMPLEMENT_MODULE(FLiveCodingModule, LiveCoding)
@@ -53,7 +54,7 @@ void FLiveCodingModule::StartupModule()
 	ISettingsModule* SettingsModule = FModuleManager::GetModulePtr<ISettingsModule>("Settings");
 	if (SettingsModule != nullptr)
 	{
-		SettingsModule->RegisterSettings("Editor", "General", "Live Coding",
+		SettingsSection = SettingsModule->RegisterSettings("Editor", "General", "Live Coding",
 			LOCTEXT("LiveCodingSettingsName", "Live Coding"),
 			LOCTEXT("LiveCodintSettingsDescription", "Settings for recompiling C++ code while the engine is running."),
 			GetMutableDefault<ULiveCodingSettings>()
@@ -99,7 +100,14 @@ void FLiveCodingModule::ShutdownModule()
 
 void FLiveCodingModule::EnableByDefault(bool bEnable)
 {
-	Settings->bEnabled = bEnable;
+	if(Settings->bEnabled != bEnable)
+	{
+		Settings->bEnabled = bEnable;
+		if(SettingsSection.IsValid())
+		{
+			SettingsSection->Save();
+		}
+	}
 	EnableForSession(bEnable);
 }
 
diff --git a/Engine/Source/Developer/Windows/LiveCoding/Private/LiveCodingModule.h b/Engine/Source/Developer/Windows/LiveCoding/Private/LiveCodingModule.h
--- a/Engine/Source/Developer/Windows/LiveCoding/Private/LiveCodingModule.h
+++ b/Engine/Source/Developer/Windows/LiveCoding/Private/LiveCodingModule.h
@@ -5,9 +5,11 @@
 #include "ILiveCodingModule.h"
 #include "Delegates/Delegate.h"
 #include "Modules/ModuleManager.h"
+#include "Templates/SharedPointer.h"
 
 struct IConsoleCommand;
 class IConsoleVariable;
+class ISettingsSection;
 class ULiveCodingSettings;
 
 class FLiveCodingModule final : public ILiveCodingModule
@@ -31,6 +33,7 @@ public:
 
 private:
 	ULiveCodingSettings* Settings;
+	TSharedPtr<ISettingsSection> SettingsSection;
 	bool bEnabledLastTick;
 	bool bStarted;
 	TSet<FString> EnabledModules;
