LiveCoding: Explicitly batch all EnableModules() commands, to prevent user command thread preempting the editor thread before it's finished adding everything. Prevents "Live coding ready" messages being spammed in the console window.

diff --git a/Engine/Source/Developer/Windows/LiveCoding/Private/External/LC_ClientUserCommandThread.cpp b/Engine/Source/Developer/Windows/LiveCoding/Private/External/LC_ClientUserCommandThread.cpp
--- a/Engine/Source/Developer/Windows/LiveCoding/Private/External/LC_ClientUserCommandThread.cpp
+++ b/Engine/Source/Developer/Windows/LiveCoding/Private/External/LC_ClientUserCommandThread.cpp
@@ -336,6 +336,18 @@ namespace
 	static CriticalSection g_userCommandQueueCS;
 }
 
+// BEGIN EPIC MOD - Allow manually batching commands
+void BeginCommandBatch()
+{
+	g_userCommandQueueCS.Enter();
+}
+
+void EndCommandBatch()
+{
+	g_userCommandQueueCS.Leave();
+}
+// END EPIC MOD
+
 
 ClientUserCommandThread::ClientUserCommandThread(DuplexPipeClient* pipeClient, DuplexPipeClient* exceptionPipeClient)
 	: m_thread(INVALID_HANDLE_VALUE)
diff --git a/Engine/Source/Developer/Windows/LiveCoding/Private/LiveCodingModule.cpp b/Engine/Source/Developer/Windows/LiveCoding/Private/LiveCodingModule.cpp
--- a/Engine/Source/Developer/Windows/LiveCoding/Private/LiveCodingModule.cpp
+++ b/Engine/Source/Developer/Windows/LiveCoding/Private/LiveCodingModule.cpp
@@ -240,6 +240,9 @@ void FLiveCodingModule::UpdateModules()
 	TArray<FModuleStatus> ModuleStatuses;
 	FModuleManager::Get().QueryModules(ModuleStatuses);
 
+	extern void BeginCommandBatch();
+	BeginCommandBatch();
+
 	for (const FModuleStatus& ModuleStatus : ModuleStatuses)
 	{
 		if (ModuleStatus.bIsLoaded)
@@ -248,6 +251,9 @@ void FLiveCodingModule::UpdateModules()
 			ConfigureModule(FName(*ModuleStatus.Name), ModuleStatus.bIsGameModule, FullFilePath);
 		}
 	}
+
+	extern void EndCommandBatch();
+	EndCommandBatch();
 #endif
 }
 
