LiveCoding: Allow editor thread to keep running while server is processing modules.

diff --git a/Engine/Source/Developer/Windows/LiveCoding/Private/External/LC_ClientUserCommandThread.cpp b/Engine/Source/Developer/Windows/LiveCoding/Private/External/LC_ClientUserCommandThread.cpp
--- a/Engine/Source/Developer/Windows/LiveCoding/Private/External/LC_ClientUserCommandThread.cpp
+++ b/Engine/Source/Developer/Windows/LiveCoding/Private/External/LC_ClientUserCommandThread.cpp
@@ -662,6 +662,48 @@ unsigned int __stdcall ClientUserCommandThread::ThreadProxy(void* context)
 }
 
 
+// BEGIN EPIC MOD - Temporarily release lock to prevent hangs
+class LeaveableScopedLock
+{
+public:
+	explicit LeaveableScopedLock(CriticalSection* cs)
+		: m_cs(cs)
+		, m_hasLock(true)
+	{
+		cs->Enter();
+	}
+
+	void Enter()
+	{
+		if (!m_hasLock)
+		{
+			m_cs->Enter();
+			m_hasLock = true;
+		}
+	}
+
+	void Leave()
+	{
+		if (m_hasLock)
+		{
+			m_cs->Leave();
+			m_hasLock = false;
+		}
+	}
+
+	~LeaveableScopedLock(void)
+	{
+		Leave();
+	}
+
+private:
+	CriticalSection* m_cs;
+	bool m_hasLock;
+};
+// END EPIC MOD
+
+
+
 unsigned int ClientUserCommandThread::ThreadFunction(Event* waitForStartEvent, CriticalSection* pipeAccessCS)
 {
 	// wait until we get the signal that the thread can start
@@ -702,7 +744,9 @@ unsigned int ClientUserCommandThread::ThreadFunction(Event* waitForStartEvent, C
 
 		// lock critical section for accessing the queue.
 		// user code might be calling other exported functions in the mean time.
-		CriticalSection::ScopedLock queueLock(&g_userCommandQueueCS);
+		// BEGIN EPIC MOD - Rearrange lock scope to prevent hangs
+		LeaveableScopedLock queueLock(&g_userCommandQueueCS);
+		// END EPIC MOD
 
 		if (g_userCommandQueue.size() == 0u)
 		{
@@ -743,6 +787,10 @@ unsigned int ClientUserCommandThread::ThreadFunction(Event* waitForStartEvent, C
 			}
 		}
 
+		// BEGIN EPIC MOD - Temporarily release lock to prevent hangs
+		queueLock.Leave();
+		// BEGIN EPIC MOD - Pre
+
 		// send out scoped commands first
 		{
 			const size_t count = enabledScopedCommands.size();
@@ -795,6 +843,14 @@ unsigned int ClientUserCommandThread::ThreadFunction(Event* waitForStartEvent, C
 			}
 		}
 
+		// BEGIN EPIC MOD - Temporarily release lock to prevent hangs
+		queueLock.Enter();
+		if(g_userCommandQueue.size() > 0)
+		{
+			continue;
+		}
+		// BEGIN EPIC MOD
+
 		m_itemInQueueEvent->Reset();
 	}
 
