LiveCoding: Fix the "Show Console" button being disabled until the first compile.

diff --git a/Engine/Source/Developer/Windows/LiveCoding/Private/LiveCodingModule.cpp b/Engine/Source/Developer/Windows/LiveCoding/Private/LiveCodingModule.cpp
--- a/Engine/Source/Developer/Windows/LiveCoding/Private/LiveCodingModule.cpp
+++ b/Engine/Source/Developer/Windows/LiveCoding/Private/LiveCodingModule.cpp
@@ -226,6 +226,7 @@ bool FLiveCodingModule::StartLiveCoding()
 
 		// Mark it as started
 		bStarted = true;
+		bEnabledForSession = true;
 	}
 	return true;
 }
