LiveCoding: Fix determination of whether a module is part of a project plugin. Such modules report "false" for FModuleStatus.bIsGamePlugin.

diff --git a/Engine/Source/Developer/Windows/LiveCoding/Private/LiveCodingModule.cpp b/Engine/Source/Developer/Windows/LiveCoding/Private/LiveCodingModule.cpp
--- a/Engine/Source/Developer/Windows/LiveCoding/Private/LiveCodingModule.cpp
+++ b/Engine/Source/Developer/Windows/LiveCoding/Private/LiveCodingModule.cpp
@@ -27,6 +27,9 @@ FLiveCodingModule::FLiveCodingModule()
 	: bEnabledLastTick(false)
 	, bEnabledForSession(false)
 	, bStarted(false)
+	, FullEnginePluginsDir(FPaths::ConvertRelativePathToFull(FPaths::EnginePluginsDir()))
+	, FullProjectDir(FPaths::ConvertRelativePathToFull(FPaths::ProjectDir()))
+	, FullProjectPluginsDir(FPaths::ConvertRelativePathToFull(FPaths::ProjectPluginsDir()))
 {
 }
 
@@ -249,7 +252,7 @@ void FLiveCodingModule::UpdateModules()
 		if (ModuleStatus.bIsLoaded)
 		{
 			FString FullFilePath = FPaths::ConvertRelativePathToFull(ModuleStatus.FilePath);
-			ConfigureModule(FName(*ModuleStatus.Name), ModuleStatus.bIsGameModule, FullFilePath);
+			ConfigureModule(FName(*ModuleStatus.Name), FullFilePath);
 		}
 	}
 
@@ -285,16 +288,16 @@ void FLiveCodingModule::OnModulesChanged(FName ModuleName, EModuleChangeReason R
 		if (FModuleManager::Get().QueryModule(ModuleName, Status))
 		{
 			FString FullFilePath = FPaths::ConvertRelativePathToFull(Status.FilePath);
-			ConfigureModule(ModuleName, Status.bIsGameModule, FullFilePath);
+			ConfigureModule(ModuleName, FullFilePath);
 		}
 	}
 #endif
 }
 
-void FLiveCodingModule::ConfigureModule(const FName& Name, bool bIsProjectModule, const FString& FullFilePath)
+void FLiveCodingModule::ConfigureModule(const FName& Name, const FString& FullFilePath)
 {
 #if !IS_MONOLITHIC
-	if (ShouldEnableModule(Name, bIsProjectModule, FullFilePath))
+	if (ShouldEnableModule(Name, FullFilePath))
 	{
 		EnableModule(FullFilePath);
 	}
@@ -305,7 +308,7 @@ void FLiveCodingModule::ConfigureModule(const FName& Name, bool bIsProjectModule
 #endif
 }
 
-bool FLiveCodingModule::ShouldEnableModule(const FName& Name, bool bIsProjectModule, const FString& FullFilePath) const
+bool FLiveCodingModule::ShouldEnableModule(const FName& Name, const FString& FullFilePath) const
 {
 	if (Settings->ExcludeSpecificModules.Contains(Name))
 	{
@@ -316,14 +319,13 @@ bool FLiveCodingModule::ShouldEnableModule(const FName& Name, bool bIsProjectMod
 		return true;
 	}
 
-	if (bIsProjectModule)
+	if (FullFilePath.StartsWith(FullProjectDir))
 	{
 		if (Settings->bIncludeProjectModules == Settings->bIncludeProjectPluginModules)
 		{
 			return Settings->bIncludeProjectModules;
 		}
 
-		FString FullProjectPluginsDir = FPaths::ConvertRelativePathToFull(FPaths::ProjectPluginsDir());
 		if(FullFilePath.StartsWith(FullProjectPluginsDir))
 		{
 			return Settings->bIncludeProjectPluginModules;
@@ -345,7 +347,6 @@ bool FLiveCodingModule::ShouldEnableModule(const FName& Name, bool bIsProjectMod
 			return Settings->bIncludeEngineModules;
 		}
 
-		FString FullEnginePluginsDir = FPaths::ConvertRelativePathToFull(FPaths::EnginePluginsDir());
 		if(FullFilePath.StartsWith(FullEnginePluginsDir))
 		{
 			return Settings->bIncludeEnginePluginModules;
diff --git a/Engine/Source/Developer/Windows/LiveCoding/Private/LiveCodingModule.h b/Engine/Source/Developer/Windows/LiveCoding/Private/LiveCodingModule.h
--- a/Engine/Source/Developer/Windows/LiveCoding/Private/LiveCodingModule.h
+++ b/Engine/Source/Developer/Windows/LiveCoding/Private/LiveCodingModule.h
@@ -40,6 +40,10 @@ private:
 	bool bStarted;
 	TSet<FString> EnabledModules;
 
+	const FString FullEnginePluginsDir;
+	const FString FullProjectDir;
+	const FString FullProjectPluginsDir;
+
 	IConsoleCommand* EnableCommand;
 	IConsoleVariable* ConsolePathVariable;
 	FDelegateHandle EndFrameDelegateHandle;
@@ -47,14 +51,13 @@ private:
 
 	bool StartLiveCoding();
 
-	void OnEndFrame();
 	void OnModulesChanged(FName ModuleName, EModuleChangeReason Reason);
 
 	void UpdateModules();
 	void EnableModule(const FString& FullFilePath);
 	void DisableModule(const FString& FullFilePath);
 
-	void ConfigureModule(const FName& Name, bool bIsProjectModule, const FString& FullFilePath);
-	bool ShouldEnableModule(const FName& Name, bool bIsProjectModule, const FString& FilePath) const;
+	void ConfigureModule(const FName& Name, const FString& FullFilePath);
+	bool ShouldEnableModule(const FName& Name, const FString& FullFilePath) const;
 };
 
