LiveCoding: Merging fixes to resolve issues with global symbols being reconstructed when loading patch DLLs. Was causing asserts when log channels were re-registered etc...

diff --git a/Engine/Source/Developer/Windows/LiveCoding/Private/External/LC_ClientCommandActions.cpp b/Engine/Source/Developer/Windows/LiveCoding/Private/External/LC_ClientCommandActions.cpp
--- a/Engine/Source/Developer/Windows/LiveCoding/Private/External/LC_ClientCommandActions.cpp
+++ b/Engine/Source/Developer/Windows/LiveCoding/Private/External/LC_ClientCommandActions.cpp
@@ -10,7 +10,7 @@
 #include "LC_Logging.h"
 
 
-bool actions::RegisterProcessFinished::Execute(CommandType* command, const DuplexPipe* pipe, void* context)
+bool actions::RegisterProcessFinished::Execute(const CommandType* command, const DuplexPipe* pipe, void* context)
 {
 	pipe->SendAck();
 
@@ -22,7 +22,7 @@ bool actions::RegisterProcessFinished::Execute(CommandType* command, const Duple
 }
 
 
-bool actions::EnableModuleFinished::Execute(CommandType* command, const DuplexPipe* pipe, void*)
+bool actions::EnableModuleFinished::Execute(const CommandType* command, const DuplexPipe* pipe, void*)
 {
 	Event* event = static_cast<Event*>(command->token);
 	event->Signal();
@@ -32,7 +32,7 @@ bool actions::EnableModuleFinished::Execute(CommandType* command, const DuplexPi
 }
 
 
-bool actions::DisableModuleFinished::Execute(CommandType* command, const DuplexPipe* pipe, void*)
+bool actions::DisableModuleFinished::Execute(const CommandType* command, const DuplexPipe* pipe, void*)
 {
 	Event* event = static_cast<Event*>(command->token);
 	event->Signal();
@@ -42,7 +42,7 @@ bool actions::DisableModuleFinished::Execute(CommandType* command, const DuplexP
 }
 
 
-bool actions::EnableAllModulesFinished::Execute(CommandType* command, const DuplexPipe* pipe, void*)
+bool actions::EnableAllModulesFinished::Execute(const CommandType* command, const DuplexPipe* pipe, void*)
 {
 	Event* event = static_cast<Event*>(command->token);
 	event->Signal();
@@ -52,7 +52,7 @@ bool actions::EnableAllModulesFinished::Execute(CommandType* command, const Dupl
 }
 
 
-bool actions::DisableAllModulesFinished::Execute(CommandType* command, const DuplexPipe* pipe, void*)
+bool actions::DisableAllModulesFinished::Execute(const CommandType* command, const DuplexPipe* pipe, void*)
 {
 	Event* event = static_cast<Event*>(command->token);
 	event->Signal();
@@ -62,7 +62,7 @@ bool actions::DisableAllModulesFinished::Execute(CommandType* command, const Dup
 }
 
 
-bool actions::EnterSyncPoint::Execute(CommandType*, const DuplexPipe* pipe, void*)
+bool actions::EnterSyncPoint::Execute(const CommandType*, const DuplexPipe* pipe, void*)
 {
 	syncPoint::Enter();
 	pipe->SendAck();
@@ -71,7 +71,7 @@ bool actions::EnterSyncPoint::Execute(CommandType*, const DuplexPipe* pipe, void
 }
 
 
-bool actions::LeaveSyncPoint::Execute(CommandType*, const DuplexPipe* pipe, void*)
+bool actions::LeaveSyncPoint::Execute(const CommandType*, const DuplexPipe* pipe, void*)
 {
 	syncPoint::Leave();
 	pipe->SendAck();
@@ -80,7 +80,7 @@ bool actions::LeaveSyncPoint::Execute(CommandType*, const DuplexPipe* pipe, void
 }
 
 
-bool actions::CallHooks::Execute(CommandType* command, const DuplexPipe* pipe, void*)
+bool actions::CallHooks::Execute(const CommandType* command, const DuplexPipe* pipe, void*)
 {
 	for (const hook::Function* hook = command->first; hook < command->last; ++hook)
 	{
@@ -98,7 +98,7 @@ bool actions::CallHooks::Execute(CommandType* command, const DuplexPipe* pipe, v
 }
 
 
-bool actions::GetModule::Execute(CommandType* command, const DuplexPipe* pipe, void*)
+bool actions::GetModule::Execute(const CommandType* command, const DuplexPipe* pipe, void*)
 {
 	HMODULE module = ::GetModuleHandleW(command->path);
 	if (module)
@@ -136,7 +136,7 @@ extern FChunkedFixedUObjectArray*& GObjectArrayForDebugVisualizers;
 // END EPIC MOD
 
 
-bool actions::LoadPatch::Execute(CommandType* command, const DuplexPipe* pipe, void*)
+bool actions::LoadPatch::Execute(const CommandType* command, const DuplexPipe* pipe, void*)
 {
 	// load library into this process
 	HMODULE module = ::LoadLibraryW(command->path);
@@ -163,7 +163,7 @@ bool actions::LoadPatch::Execute(CommandType* command, const DuplexPipe* pipe, v
 }
 
 
-bool actions::UnloadPatch::Execute(CommandType* command, const DuplexPipe* pipe, void*)
+bool actions::UnloadPatch::Execute(const CommandType* command, const DuplexPipe* pipe, void*)
 {
 	// unload library from this process
 	::FreeLibrary(command->module);
@@ -173,7 +173,7 @@ bool actions::UnloadPatch::Execute(CommandType* command, const DuplexPipe* pipe,
 }
 
 
-bool actions::CallEntryPoint::Execute(CommandType* command, const DuplexPipe* pipe, void*)
+bool actions::CallEntryPoint::Execute(const CommandType* command, const DuplexPipe* pipe, void*)
 {
 	executable::CallDllEntryPoint(command->moduleBase, command->entryPointRva);
 	pipe->SendAck();
@@ -182,7 +182,7 @@ bool actions::CallEntryPoint::Execute(CommandType* command, const DuplexPipe* pi
 }
 
 
-bool actions::LogOutput::Execute(CommandType* command, const DuplexPipe* pipe, void*)
+bool actions::LogOutput::Execute(const CommandType* command, const DuplexPipe* pipe, void*)
 {
 	logging::LogNoFormat<logging::Channel::USER>(command->buffer);
 	pipe->SendAck();
@@ -194,7 +194,7 @@ bool actions::LogOutput::Execute(CommandType* command, const DuplexPipe* pipe, v
 extern bool GIsCompileActive;
 // END EPIC MOD
 
-bool actions::CompilationFinished::Execute(CommandType*, const DuplexPipe* pipe, void*)
+bool actions::CompilationFinished::Execute(const CommandType*, const DuplexPipe* pipe, void*)
 {
 	pipe->SendAck();
 
@@ -207,7 +207,7 @@ bool actions::CompilationFinished::Execute(CommandType*, const DuplexPipe* pipe,
 }
 
 
-bool actions::HandleExceptionFinished::Execute(CommandType* command, const DuplexPipe* pipe, void* context)
+bool actions::HandleExceptionFinished::Execute(const CommandType* command, const DuplexPipe* pipe, void* context)
 {
 	ClientUserCommandThread::ExceptionResult* resultContext = static_cast<ClientUserCommandThread::ExceptionResult*>(context);
 	resultContext->returnAddress = command->returnAddress;
diff --git a/Engine/Source/Developer/Windows/LiveCoding/Private/External/LC_ClientCommandActions.h b/Engine/Source/Developer/Windows/LiveCoding/Private/External/LC_ClientCommandActions.h
--- a/Engine/Source/Developer/Windows/LiveCoding/Private/External/LC_ClientCommandActions.h
+++ b/Engine/Source/Developer/Windows/LiveCoding/Private/External/LC_ClientCommandActions.h
@@ -10,11 +10,11 @@ class DuplexPipe;
 
 namespace actions
 {
-	#define DECLARE_ACTION(_name)																	\
-		struct _name																				\
-		{																							\
-			typedef ::commands::_name CommandType;													\
-			static bool Execute(CommandType* command, const DuplexPipe* pipe, void* context);		\
+	#define DECLARE_ACTION(_name)																		\
+		struct _name																					\
+		{																								\
+			typedef ::commands::_name CommandType;														\
+			static bool Execute(const CommandType* command, const DuplexPipe* pipe, void* context);		\
 		}
 
 	DECLARE_ACTION(RegisterProcessFinished);
diff --git a/Engine/Source/Developer/Windows/LiveCoding/Private/External/LC_ClientStartupThread.cpp b/Engine/Source/Developer/Windows/LiveCoding/Private/External/LC_ClientStartupThread.cpp
--- a/Engine/Source/Developer/Windows/LiveCoding/Private/External/LC_ClientStartupThread.cpp
+++ b/Engine/Source/Developer/Windows/LiveCoding/Private/External/LC_ClientStartupThread.cpp
@@ -176,14 +176,14 @@ void ClientStartupThread::TriggerRecompile(void)
 }
 
 
-void ClientStartupThread::BuildPatch(const wchar_t* moduleNames[], const wchar_t* objPaths[], unsigned int count)
+void ClientStartupThread::BuildPatch(const wchar_t* moduleNames[], const wchar_t* objPaths[], const wchar_t* amalgamatedObjPaths[], unsigned int count)
 {
 	// we cannot wait for commands in the user command thread as long as startup hasn't finished
 	Join();
 
 	if (m_userCommandThread)
 	{
-		m_userCommandThread->BuildPatch(moduleNames, objPaths, count);
+		m_userCommandThread->BuildPatch(moduleNames, objPaths, amalgamatedObjPaths, count);
 	}
 }
 
@@ -445,7 +445,7 @@ unsigned int ClientStartupThread::ThreadFunction(const std::wstring& processGrou
 	m_userCommandThread->Start(processGroupName, m_startEvent, m_pipeClientCS);
 
 	// register this process with Live++
-	m_pipeClient->SendCommandAndWaitForAck(commands::RegisterProcess { process::GetId(), process::GetBase(), commandThreadId });
+	m_pipeClient->SendCommandAndWaitForAck(commands::RegisterProcess { process::GetBase(), process::GetId(), commandThreadId });
 
 	// handle commands until registration is finished
 	{
diff --git a/Engine/Source/Developer/Windows/LiveCoding/Private/External/LC_ClientStartupThread.h b/Engine/Source/Developer/Windows/LiveCoding/Private/External/LC_ClientStartupThread.h
--- a/Engine/Source/Developer/Windows/LiveCoding/Private/External/LC_ClientStartupThread.h
+++ b/Engine/Source/Developer/Windows/LiveCoding/Private/External/LC_ClientStartupThread.h
@@ -35,7 +35,7 @@ public:
 
 	void WaitForToken(void* token);
 	void TriggerRecompile(void);
-	void BuildPatch(const wchar_t* moduleNames[], const wchar_t* objPaths[], unsigned int count);
+	void BuildPatch(const wchar_t* moduleNames[], const wchar_t* objPaths[], const wchar_t* amalgamatedObjPaths[], unsigned int count);
 
 	void InstallExceptionHandler(void);
 
diff --git a/Engine/Source/Developer/Windows/LiveCoding/Private/External/LC_ClientUserCommandThread.cpp b/Engine/Source/Developer/Windows/LiveCoding/Private/External/LC_ClientUserCommandThread.cpp
--- a/Engine/Source/Developer/Windows/LiveCoding/Private/External/LC_ClientUserCommandThread.cpp
+++ b/Engine/Source/Developer/Windows/LiveCoding/Private/External/LC_ClientUserCommandThread.cpp
@@ -47,6 +47,8 @@ namespace userCommands
 	{
 		EnableModuleCommand(void)
 			: BaseCommand(Scope::ENABLE_MODULES)
+			, token(nullptr)
+			, moduleName()
 		{
 		}
 
@@ -70,6 +72,8 @@ namespace userCommands
 	{
 		EnableAllModulesCommand(void)
 			: BaseCommand(Scope::ENABLE_MODULES)
+			, token(nullptr)
+			, moduleName()
 		{
 		}
 
@@ -93,6 +97,8 @@ namespace userCommands
 	{
 		DisableModuleCommand(void)
 			: BaseCommand(Scope::DISABLE_MODULES)
+			, token(nullptr)
+			, moduleName()
 		{
 		}
 
@@ -116,6 +122,8 @@ namespace userCommands
 	{
 		DisableAllModulesCommand(void)
 			: BaseCommand(Scope::DISABLE_MODULES)
+			, token(nullptr)
+			, moduleName()
 		{
 		}
 
@@ -262,6 +270,10 @@ namespace userCommands
 	{
 		BuildPatchCommand(void)
 			: BaseCommand(Scope::NONE)
+			, count(0u)
+			, moduleNames()
+			, objPaths()
+			, amalgamatedObjPaths()
 		{
 		}
 
@@ -278,6 +290,7 @@ namespace userCommands
 				commands::BuildPatchPacket packet;
 				wcscpy_s(packet.moduleName, moduleNames[i].c_str());
 				wcscpy_s(packet.objPath, objPaths[i].c_str());
+				wcscpy_s(packet.amalgamatedObjPath, amalgamatedObjPaths[i].c_str());
 
 				pipe->SendCommandAndWaitForAck(packet);
 			}
@@ -286,12 +299,15 @@ namespace userCommands
 		unsigned int count;
 		std::vector<std::wstring> moduleNames;
 		std::vector<std::wstring> objPaths;
+		std::vector<std::wstring> amalgamatedObjPaths;
 	};
 
 	struct ApplySettingBoolCommand : public BaseCommand
 	{
 		ApplySettingBoolCommand(void)
 			: BaseCommand(Scope::NONE)
+			, settingName()
+			, value(0)
 		{
 		}
 
@@ -313,6 +329,8 @@ namespace userCommands
 	{
 		ApplySettingIntCommand(void)
 			: BaseCommand(Scope::NONE)
+			, settingName()
+			, value(0)
 		{
 		}
 
@@ -334,6 +352,8 @@ namespace userCommands
 	{
 		ApplySettingStringCommand(void)
 			: BaseCommand(Scope::NONE)
+			, settingName()
+			, value()
 		{
 		}
 
@@ -510,17 +530,28 @@ void ClientUserCommandThread::TriggerRecompile(void)
 }
 
 
-void ClientUserCommandThread::BuildPatch(const wchar_t* moduleNames[], const wchar_t* objPaths[], unsigned int count)
+void ClientUserCommandThread::BuildPatch(const wchar_t* moduleNames[], const wchar_t* objPaths[], const wchar_t* amalgamatedObjPaths[], unsigned int count)
 {
 	userCommands::BuildPatchCommand* command = new userCommands::BuildPatchCommand;
 	command->count = count;
 	command->moduleNames.reserve(count);
 	command->objPaths.reserve(count);
+	command->amalgamatedObjPaths.reserve(count);
 
 	for (unsigned int i = 0u; i < count; ++i)
 	{
 		command->moduleNames.push_back(moduleNames[i]);
 		command->objPaths.push_back(objPaths[i]);
+
+		// the amalgamated object paths are optional
+		if (amalgamatedObjPaths && amalgamatedObjPaths[i])
+		{
+			command->amalgamatedObjPaths.push_back(amalgamatedObjPaths[i]);
+		}
+		else
+		{
+			command->amalgamatedObjPaths.push_back(std::wstring());
+		}
 	}
 	{
 		CriticalSection::ScopedLock lock(&g_userCommandQueueCS);
diff --git a/Engine/Source/Developer/Windows/LiveCoding/Private/External/LC_ClientUserCommandThread.h b/Engine/Source/Developer/Windows/LiveCoding/Private/External/LC_ClientUserCommandThread.h
--- a/Engine/Source/Developer/Windows/LiveCoding/Private/External/LC_ClientUserCommandThread.h
+++ b/Engine/Source/Developer/Windows/LiveCoding/Private/External/LC_ClientUserCommandThread.h
@@ -41,7 +41,7 @@ public:
 
 	void WaitForToken(void* token);
 	void TriggerRecompile(void);
-	void BuildPatch(const wchar_t* moduleNames[], const wchar_t* objPaths[], unsigned int count);
+	void BuildPatch(const wchar_t* moduleNames[], const wchar_t* objPaths[], const wchar_t* amalgamatedObjPaths[], unsigned int count);
 
 	void InstallExceptionHandler(void);
 	ExceptionResult HandleException(EXCEPTION_RECORD* exception, CONTEXT* context, unsigned int threadId);
diff --git a/Engine/Source/Developer/Windows/LiveCoding/Private/External/LC_Commands.h b/Engine/Source/Developer/Windows/LiveCoding/Private/External/LC_Commands.h
--- a/Engine/Source/Developer/Windows/LiveCoding/Private/External/LC_Commands.h
+++ b/Engine/Source/Developer/Windows/LiveCoding/Private/External/LC_Commands.h
@@ -23,8 +23,8 @@ namespace commands
 	{
 		static const uint32_t ID = Acknowledge::ID + 1u;
 
-		unsigned int processId;
 		void* processBase;
+		unsigned int processId;
 		unsigned int threadId;		// thread ID of Live++ thread running in host
 	};
 
@@ -256,6 +256,7 @@ namespace commands
 
 		wchar_t moduleName[MAX_PATH];
 		wchar_t objPath[MAX_PATH];
+		wchar_t amalgamatedObjPath[MAX_PATH];
 	};
 
 	// tell Live++ to handle an exception
diff --git a/Engine/Source/Developer/Windows/LiveCoding/Private/External/LC_CriticalSection.h b/Engine/Source/Developer/Windows/LiveCoding/Private/External/LC_CriticalSection.h
--- a/Engine/Source/Developer/Windows/LiveCoding/Private/External/LC_CriticalSection.h
+++ b/Engine/Source/Developer/Windows/LiveCoding/Private/External/LC_CriticalSection.h
@@ -4,6 +4,7 @@
 
 #include "CoreTypes.h"
 #include "Windows/MinimalWindowsApi.h"
+#include "LC_Platform.h"
 
 class CriticalSection
 {
@@ -29,6 +30,11 @@ public:
 		}
 
 	private:
+		LC_DISABLE_COPY(ScopedLock);
+		LC_DISABLE_MOVE(ScopedLock);
+		LC_DISABLE_ASSIGNMENT(ScopedLock);
+		LC_DISABLE_MOVE_ASSIGNMENT(ScopedLock);
+
 		CriticalSection* m_cs;
 	};
 
diff --git a/Engine/Source/Developer/Windows/LiveCoding/Private/External/LC_EntryPoint.cpp b/Engine/Source/Developer/Windows/LiveCoding/Private/External/LC_EntryPoint.cpp
--- a/Engine/Source/Developer/Windows/LiveCoding/Private/External/LC_EntryPoint.cpp
+++ b/Engine/Source/Developer/Windows/LiveCoding/Private/External/LC_EntryPoint.cpp
@@ -125,9 +125,9 @@ LPP_DLL_API void __cdecl LppTriggerRecompile(void)
 }
 
 
-LPP_DLL_API void __cdecl LppBuildPatch(const wchar_t* moduleNames[], const wchar_t* objPaths[], unsigned int count)
+LPP_DLL_API void __cdecl LppBuildPatch(const wchar_t* moduleNames[], const wchar_t* objPaths[], const wchar_t* amalgamatedObjPaths[], unsigned int count)
 {
-	g_startupThread->BuildPatch(moduleNames, objPaths, count);
+	g_startupThread->BuildPatch(moduleNames, objPaths, amalgamatedObjPaths, count);
 }
 
 
diff --git a/Engine/Source/Developer/Windows/LiveCoding/Private/External/LC_Executable.cpp b/Engine/Source/Developer/Windows/LiveCoding/Private/External/LC_Executable.cpp
--- a/Engine/Source/Developer/Windows/LiveCoding/Private/External/LC_Executable.cpp
+++ b/Engine/Source/Developer/Windows/LiveCoding/Private/External/LC_Executable.cpp
@@ -26,14 +26,14 @@ namespace detail
 		void* base = image->base;
 
 		// PE image start with a DOS header
-		IMAGE_DOS_HEADER* dosHeader = static_cast<IMAGE_DOS_HEADER*>(base);
+		const IMAGE_DOS_HEADER* dosHeader = static_cast<IMAGE_DOS_HEADER*>(base);
 		if (dosHeader->e_magic != IMAGE_DOS_SIGNATURE)
 		{
 			LC_ERROR_USER("Image has unknown file format");
 			return nullptr;
 		}
 
-		IMAGE_NT_HEADERS* ntHeader = pointer::Offset<IMAGE_NT_HEADERS*>(dosHeader, dosHeader->e_lfanew);
+		const IMAGE_NT_HEADERS* ntHeader = pointer::Offset<const IMAGE_NT_HEADERS*>(dosHeader, dosHeader->e_lfanew);
 		if (ntHeader->Signature != IMAGE_NT_SIGNATURE)
 		{
 			LC_ERROR_USER("Invalid .exe file");
@@ -208,7 +208,7 @@ void executable::RebaseImage(Image* image, PreferredBase preferredBase)
 
 			// PE spec: Block size: The total number of bytes in the base relocation block, *including* the Page RVA and
 			// Block Size fields and the Type/Offset fields that follow.
-			DWORD numberOfEntriesInThisBlock = (blockSize - sizeof(IMAGE_BASE_RELOCATION)) / sizeof(WORD);
+			const DWORD numberOfEntriesInThisBlock = (blockSize - sizeof(IMAGE_BASE_RELOCATION)) / sizeof(WORD);
 			const WORD* entries = pointer::Offset<const WORD*>(baseRelocations, sizeof(IMAGE_BASE_RELOCATION));
 			for (DWORD i = 0u; i < numberOfEntriesInThisBlock; ++i)
 			{
diff --git a/Engine/Source/Developer/Windows/LiveCoding/Private/External/LC_InterprocessMutex.h b/Engine/Source/Developer/Windows/LiveCoding/Private/External/LC_InterprocessMutex.h
--- a/Engine/Source/Developer/Windows/LiveCoding/Private/External/LC_InterprocessMutex.h
+++ b/Engine/Source/Developer/Windows/LiveCoding/Private/External/LC_InterprocessMutex.h
@@ -28,6 +28,11 @@ public:
 		}
 
 	private:
+		LC_DISABLE_COPY(ScopedLock);
+		LC_DISABLE_MOVE(ScopedLock);
+		LC_DISABLE_ASSIGNMENT(ScopedLock);
+		LC_DISABLE_MOVE_ASSIGNMENT(ScopedLock);
+
 		InterprocessMutex* m_mutex;
 	};
 
diff --git a/Engine/Source/Developer/Windows/LiveCoding/Private/External/LC_Process.cpp b/Engine/Source/Developer/Windows/LiveCoding/Private/External/LC_Process.cpp
--- a/Engine/Source/Developer/Windows/LiveCoding/Private/External/LC_Process.cpp
+++ b/Engine/Source/Developer/Windows/LiveCoding/Private/External/LC_Process.cpp
@@ -634,11 +634,13 @@ namespace process
 			if (pointer::Offset<const void*>(scan, size) >= upperBound)
 			{
 				// outside of range to scan
+				LC_ERROR_DEV("Could not find memory range that fits 0x%X bytes with alignment 0x%X in range from 0x%p to 0x%p (scan: 0x%p)", size, alignment, lowerBound, upperBound, scan);
 				return nullptr;
 			}
 			else if (scan < lowerBound)
 			{
 				// outside of range (possible wrap-around)
+				LC_ERROR_DEV("Could not find memory range that fits 0x%X bytes with alignment 0x%X in range from 0x%p to 0x%p (scan: 0x%p)", size, alignment, lowerBound, upperBound, scan);
 				return nullptr;
 			}
 
@@ -737,6 +739,8 @@ namespace process
 			{
 				// buffer is too small, try again
 				::free(processSnapshot);
+				processSnapshot = nullptr;
+
 				bufferSize *= 2u;
 			}
 			else if (status < 0)
@@ -744,7 +748,7 @@ namespace process
 				// something went wrong
 				LC_ERROR_USER("Cannot enumerate threads in process (PID: %d)", processId);
 				::free(processSnapshot);
-				
+
 				return threadIds;
 			}
 		}
@@ -778,10 +782,13 @@ namespace process
 			}
 
 			// record all threads belonging to the given process
-			for (ULONG i = 0u; i < processInfo->uThreadCount; ++i)
+			if (processInfo)
 			{
-				const DWORD threadId = static_cast<DWORD>(reinterpret_cast<DWORD_PTR>(processInfo->Threads[i].ClientId.UniqueThread));
-				threadIds.push_back(threadId);
+				for (ULONG i = 0u; i < processInfo->uThreadCount; ++i)
+				{
+					const DWORD threadId = static_cast<DWORD>(reinterpret_cast<DWORD_PTR>(processInfo->Threads[i].ClientId.UniqueThread));
+					threadIds.push_back(threadId);
+				}
 			}
 		}
 
diff --git a/Engine/Source/Developer/Windows/LiveCodingServer/Private/External/LC_AppSettings.cpp b/Engine/Source/Developer/Windows/LiveCodingServer/Private/External/LC_AppSettings.cpp
--- a/Engine/Source/Developer/Windows/LiveCodingServer/Private/External/LC_AppSettings.cpp
+++ b/Engine/Source/Developer/Windows/LiveCodingServer/Private/External/LC_AppSettings.cpp
@@ -30,8 +30,9 @@ namespace
 
 	static std::wstring g_cachedCompilerPath;
 	static std::wstring g_cachedLinkerPath;
+	static types::vector<std::wstring> g_cachedAmalgamatedCppFileExtensions;
 
-	static std::wstring DeterminePath(SettingString* setting, const wchar_t* const type, const wchar_t* const vs2017Path, const wchar_t* const vs2015AndEarlierPath, const wchar_t* exeName)
+	static std::wstring DeterminePath(const SettingString* setting, const wchar_t* const type, const wchar_t* const vs2017Path, const wchar_t* const vs2015AndEarlierPath, const wchar_t* exeName)
 	{
 		// absolute paths can be used as they are, relative paths are supposed to be relative to the Live++ executable
 		std::wstring path;
@@ -442,7 +443,7 @@ void appSettings::Startup(const wchar_t* group)
 		L"amalgamation_split_into_single_parts",
 		L"Split into single parts",
 		L"Specifies whether amalgamated/unity files are automatically split into single files",
-		false
+		true
 	);
 
 	g_amalgamationSplitMinCppCount = new SettingInt
@@ -453,6 +454,15 @@ void appSettings::Startup(const wchar_t* group)
 		L"Minimum number of .cpp files that must be included in an amalgamated/unity file before it is split",
 		3
 	);	
+
+	g_amalgamationCppFileExtensions = new SettingString
+	(
+		group,
+		L"amalgamation_cpp_file_extensions",
+		L"C/C++ file extensions",
+		L"File extensions treated as C/C++ files when splitting amalgamated/unity files",
+		L".cpp;.c;.cc;.c++;.cp;.cxx"
+	);
 }
 
 
@@ -502,6 +512,7 @@ void appSettings::Shutdown(void)
 
 	delete g_amalgamationSplitIntoSingleParts;
 	delete g_amalgamationSplitMinCppCount;
+	delete g_amalgamationCppFileExtensions;
 }
 
 
@@ -571,6 +582,12 @@ std::wstring appSettings::GetLinkerPath(void)
 }
 
 
+const types::vector<std::wstring>& appSettings::GetAmalgamatedCppFileExtensions(void)
+{
+	return g_cachedAmalgamatedCppFileExtensions;
+}
+
+
 void appSettings::UpdateCompilerPathCache(void)
 {
 	g_cachedCompilerPath = DeterminePath(g_compilerPath, L"compiler", VS2017_COMPILER_PATH, VS2015_AND_EARLIER_COMPILER_PATH, COMPILER_EXE);
@@ -590,6 +607,27 @@ void appSettings::UpdatePathCache(void)
 }
 
 
+void appSettings::UpdateAmalgamatedCppFileExtensions(void)
+{
+	g_cachedAmalgamatedCppFileExtensions.clear();
+	g_cachedAmalgamatedCppFileExtensions.reserve(16u);
+
+	const wchar_t DELIMITER = L';';
+	const std::wstring extensions = g_amalgamationCppFileExtensions->GetValue();
+	size_t start = extensions.find_first_not_of(DELIMITER);
+	size_t end = start;
+
+	while (start != std::string::npos)
+	{
+		// first find the next occurrence of the delimiter, then push the token into the vector,
+		// then skip all occurrences of the delimiter until we find the start of the next token
+		end = extensions.find(DELIMITER, start);
+		g_cachedAmalgamatedCppFileExtensions.push_back(extensions.substr(start, end - start));
+		start = extensions.find_first_not_of(DELIMITER, end);
+	}
+}
+
+
 void appSettings::ApplySettingBool(const char* const settingName, bool value)
 {
 	const unsigned int COUNT = 20u;
@@ -670,7 +708,7 @@ void appSettings::ApplySettingString(const char* const settingName, const wchar_
 {
 	// try string settings first
 	{
-		const unsigned int COUNT = 9u;
+		const unsigned int COUNT = 10u;
 		SettingString* settings[COUNT] =
 		{
 			g_playSoundOnSuccess,
@@ -681,7 +719,8 @@ void appSettings::ApplySettingString(const char* const settingName, const wchar_
 			g_linkerOptions,
 			g_continuousCompilationPath,
 			g_virtualDriveLetter,
-			g_virtualDrivePath
+			g_virtualDrivePath,
+			g_amalgamationCppFileExtensions
 		};
 
 		const SettingString* setting = ApplySetting(settings, COUNT, settingName, value);
@@ -697,12 +736,16 @@ void appSettings::ApplySettingString(const char* const settingName, const wchar_
 				UpdateLinkerPathCache();
 			}
 
+			// update cache when amalgamated C++ file extensions change
+			else if (setting == g_amalgamationCppFileExtensions)
+			{
+				UpdateAmalgamatedCppFileExtensions();
+			}
+
 			return;
 		}
 	}
 
-	const std::wstring wideSettingName(string::ToWideString(settingName));
-
 	// try proxies second
 	{
 		const unsigned int COUNT = 2u;
@@ -778,3 +821,4 @@ extern SettingBool* appSettings::g_installCompiledPatchesMultiProcess = nullptr;
 
 extern SettingBool* appSettings::g_amalgamationSplitIntoSingleParts = nullptr;
 extern SettingInt* appSettings::g_amalgamationSplitMinCppCount = nullptr;
+extern SettingString* appSettings::g_amalgamationCppFileExtensions = nullptr;
diff --git a/Engine/Source/Developer/Windows/LiveCodingServer/Private/External/LC_AppSettings.h b/Engine/Source/Developer/Windows/LiveCodingServer/Private/External/LC_AppSettings.h
--- a/Engine/Source/Developer/Windows/LiveCodingServer/Private/External/LC_AppSettings.h
+++ b/Engine/Source/Developer/Windows/LiveCodingServer/Private/External/LC_AppSettings.h
@@ -40,10 +40,15 @@ namespace appSettings
 	// returns the overridden linker path
 	std::wstring GetLinkerPath(void);
 
+	// returns an array of file extensions to treat as C++ files when splitting amalgamated files
+	const types::vector<std::wstring>& GetAmalgamatedCppFileExtensions(void);
+
 	void UpdateCompilerPathCache(void);
 	void UpdateLinkerPathCache(void);
 	void UpdatePathCache(void);
 
+	void UpdateAmalgamatedCppFileExtensions(void);
+
 	// apply a new value to any of the boolean settings
 	void ApplySettingBool(const char* const settingName, bool value);
 
@@ -79,14 +84,14 @@ namespace appSettings
 	extern SettingBool* g_enableDevLogCompilands;
 
 	// compiler
-	extern SettingString* g_compilerPath;					// DO NOT USE directly, use GetCompilerPath instead!
+	extern SettingString* g_compilerPath;						// DO NOT USE directly, use GetCompilerPath instead!
 	extern SettingBool* g_useCompilerOverrideAsFallback;
 	extern SettingBool* g_useCompilerEnvironment;
 	extern SettingString* g_compilerOptions;
 	extern SettingBool* g_compilerForcePchPdbs;
 
 	// linker
-	extern SettingString* g_linkerPath;						// DO NOT USE directly, use GetLinkerPath instead!
+	extern SettingString* g_linkerPath;							// DO NOT USE directly, use GetLinkerPath instead!
 	extern SettingBool* g_useLinkerOverrideAsFallback;
 	extern SettingBool* g_useLinkerEnvironment;
 	extern SettingString* g_linkerOptions;
@@ -107,4 +112,5 @@ namespace appSettings
 	// amalgamated/unity builds
 	extern SettingBool* g_amalgamationSplitIntoSingleParts;
 	extern SettingInt* g_amalgamationSplitMinCppCount;
+	extern SettingString* g_amalgamationCppFileExtensions;		// DO NOT USE directly, use GetAmalgamatedCppFileExtensions instead!
 }
diff --git a/Engine/Source/Developer/Windows/LiveCodingServer/Private/External/LC_Environment.cpp b/Engine/Source/Developer/Windows/LiveCodingServer/Private/External/LC_Environment.cpp
--- a/Engine/Source/Developer/Windows/LiveCodingServer/Private/External/LC_Environment.cpp
+++ b/Engine/Source/Developer/Windows/LiveCodingServer/Private/External/LC_Environment.cpp
@@ -61,10 +61,17 @@ namespace environment
 			const char* start = memory;
 
 			// search for carriage return
-			while (memory[0] != '\r')
+			while ((memory < memoryEnd) && (memory[0] != '\r'))
 			{
 				++memory;
 			}
+
+			if (memory >= memoryEnd)
+			{
+				// reached the end
+				break;
+			}
+
 			const char* end = memory;
 
 			std::string line(start, end);
@@ -83,7 +90,7 @@ namespace environment
 			}
 
 			// skip carriage return and new line
-			while ((memory[0] == '\r') || (memory[0] == '\n'))
+			while ((memory < memoryEnd) && ((memory[0] == '\r') || (memory[0] == '\n')))
 			{
 				++memory;
 			}
@@ -162,7 +169,11 @@ namespace environment
 		value.resize(MAX_PATH);
 
 		size_t bytes = 0u;
-		_wgetenv_s(&bytes, &value[0], value.size(), variable);
+		const errno_t error = _wgetenv_s(&bytes, &value[0], value.size(), variable);
+		if (error != 0)
+		{
+			LC_LOG_DEV("Could not retrieve environment variable %S (Error: %d)", variable, error);
+		}
 
 		return value;
 	}
diff --git a/Engine/Source/Developer/Windows/LiveCodingServer/Private/External/LC_FileUtil.cpp b/Engine/Source/Developer/Windows/LiveCodingServer/Private/External/LC_FileUtil.cpp
--- a/Engine/Source/Developer/Windows/LiveCodingServer/Private/External/LC_FileUtil.cpp
+++ b/Engine/Source/Developer/Windows/LiveCodingServer/Private/External/LC_FileUtil.cpp
@@ -73,12 +73,59 @@ namespace detail
 namespace
 {
 	static detail::NormalizedFilenameCache* g_normalizedFilenameCache = nullptr;
+
+	static file::DriveType::Enum g_driveTypeCache['z' - 'a' + 1u] = {};
 }
 
 
 void file::Startup(void)
 {
 	g_normalizedFilenameCache = new detail::NormalizedFilenameCache;
+
+	// fill cache of drive types
+	char root[4u] = { 'a', ':', '\\', '\0' };
+	for (char drive = 'a'; drive <= 'z'; ++drive)
+	{
+		const int index = drive - 'a';
+		
+		root[0] = drive;
+		const UINT driveType = ::GetDriveTypeA(root);
+
+		switch (driveType)
+		{
+			case DRIVE_UNKNOWN:
+				g_driveTypeCache[index] = DriveType::UNKNOWN;
+				break;
+
+			case DRIVE_NO_ROOT_DIR:
+				g_driveTypeCache[index] = DriveType::UNKNOWN;
+				break;
+
+			case DRIVE_REMOVABLE:
+				g_driveTypeCache[index] = DriveType::REMOVABLE;
+				break;
+
+			case DRIVE_FIXED:
+				g_driveTypeCache[index] = DriveType::FIXED;
+				break;
+
+			case DRIVE_REMOTE:
+				g_driveTypeCache[index] = DriveType::REMOTE;
+				break;
+
+			case DRIVE_CDROM:
+				g_driveTypeCache[index] = DriveType::OPTICAL;
+				break;
+
+			case DRIVE_RAMDISK:
+				g_driveTypeCache[index] = DriveType::RAMDISK;
+				break;
+
+			default:
+				g_driveTypeCache[index] = DriveType::UNKNOWN;
+				break;
+		}
+	}
 }
 
 
@@ -88,6 +135,37 @@ void file::Shutdown(void)
 }
 
 
+file::DriveType::Enum file::GetDriveType(char driveLetter)
+{
+	if ((driveLetter >= 'a') && (driveLetter <= 'z'))
+	{
+		return g_driveTypeCache[driveLetter - 'a'];
+	}
+	else if ((driveLetter >= 'A') && (driveLetter <= 'Z'))
+	{
+		return g_driveTypeCache[driveLetter - 'A'];
+	}
+
+	return DriveType::UNKNOWN;
+}
+
+
+file::DriveType::Enum file::GetDriveType(const wchar_t* path)
+{
+	const wchar_t driveLetter = path[0];
+	if ((driveLetter >= L'a') && (driveLetter <= L'z'))
+	{
+		return g_driveTypeCache[driveLetter - L'a'];
+	}
+	else if ((driveLetter >= L'A') && (driveLetter <= L'Z'))
+	{
+		return g_driveTypeCache[driveLetter - L'A'];
+	}
+
+	return DriveType::UNKNOWN;
+}
+
+
 file::Attributes file::GetAttributes(const wchar_t* path)
 {
 	const WIN32_FILE_ATTRIBUTE_DATA initialData = { INVALID_FILE_ATTRIBUTES };
@@ -170,7 +248,7 @@ bool file::IsRelativePath(const wchar_t* path)
 		return false;
 	}
 
-	return (::PathIsRelativeW(path) == Windows::TRUE);
+	return (::PathIsRelativeW(path) != Windows::FALSE);
 }
 
 
diff --git a/Engine/Source/Developer/Windows/LiveCodingServer/Private/External/LC_FileUtil.h b/Engine/Source/Developer/Windows/LiveCodingServer/Private/External/LC_FileUtil.h
--- a/Engine/Source/Developer/Windows/LiveCodingServer/Private/External/LC_FileUtil.h
+++ b/Engine/Source/Developer/Windows/LiveCodingServer/Private/External/LC_FileUtil.h
@@ -13,9 +13,25 @@ namespace file
 		WIN32_FILE_ATTRIBUTE_DATA data;
 	};
 
+	struct DriveType
+	{
+		enum Enum
+		{
+			UNKNOWN = 0,
+			REMOVABLE,
+			FIXED,
+			REMOTE,
+			OPTICAL,
+			RAMDISK
+		};
+	};
+
 	void Startup(void);
 	void Shutdown(void);
 
+	// Returns the type of a drive. Drive letters can be both lower- or upper-case
+	DriveType::Enum GetDriveType(char driveLetter);
+	DriveType::Enum GetDriveType(const wchar_t* path);
 
 	Attributes GetAttributes(const wchar_t* path);
 	uint64_t GetLastModificationTime(const Attributes& attributes);
diff --git a/Engine/Source/Developer/Windows/LiveCodingServer/Private/External/LC_GrowingMemoryBlock.cpp b/Engine/Source/Developer/Windows/LiveCodingServer/Private/External/LC_GrowingMemoryBlock.cpp
--- a/Engine/Source/Developer/Windows/LiveCodingServer/Private/External/LC_GrowingMemoryBlock.cpp
+++ b/Engine/Source/Developer/Windows/LiveCodingServer/Private/External/LC_GrowingMemoryBlock.cpp
@@ -35,12 +35,6 @@ bool GrowingMemoryBlock::Insert(const void* data, size_t size)
 	{
 		// grow to accommodate the new size, copy over the existing data, and delete the old storage
 		uint8_t* newData = new uint8_t[m_capacity];
-		if (!newData)
-		{
-			// could not allocate memory, bail out
-			return false;
-		}
-
 		memcpy(newData, m_data, m_size);
 		delete[] m_data;
 		m_data = newData;
diff --git a/Engine/Source/Developer/Windows/LiveCodingServer/Private/External/LC_LiveModule.cpp b/Engine/Source/Developer/Windows/LiveCodingServer/Private/External/LC_LiveModule.cpp
--- a/Engine/Source/Developer/Windows/LiveCodingServer/Private/External/LC_LiveModule.cpp
+++ b/Engine/Source/Developer/Windows/LiveCodingServer/Private/External/LC_LiveModule.cpp
@@ -167,20 +167,20 @@ namespace
 	{
 		// MSVC's link.exe is much more common, so treat this as our default
 		const std::wstring linkerPath = GetLinkerPath(linkerDb);
-		const std::wstring lowerCaseFilename = string::ToLower(file::GetFilename(linkerPath));
-		if (string::Contains(linkerPath.c_str(), L"lld"))
+		const std::wstring lowerCaseLinkerPath = string::ToLower(file::GetFilename(linkerPath));
+		if (string::Contains(lowerCaseLinkerPath.c_str(), L"lld"))
 		{
 			return coff::SymbolRemovalStrategy::LLD_COMPATIBLE;
 		}
-		else if (string::Contains(linkerPath.c_str(), L"lld-link"))
+		else if (string::Contains(lowerCaseLinkerPath.c_str(), L"lld-link"))
 		{
 			return coff::SymbolRemovalStrategy::LLD_COMPATIBLE;
 		}
-		else if (string::Contains(linkerPath.c_str(), L"ld.lld"))
+		else if (string::Contains(lowerCaseLinkerPath.c_str(), L"ld.lld"))
 		{
 			return coff::SymbolRemovalStrategy::LLD_COMPATIBLE;
 		}
-		else if (string::Contains(linkerPath.c_str(), L"ld64.lld"))
+		else if (string::Contains(lowerCaseLinkerPath.c_str(), L"ld64.lld"))
 		{
 			return coff::SymbolRemovalStrategy::LLD_COMPATIBLE;
 		}
@@ -366,12 +366,35 @@ namespace
 	}
 
 
-	// helper function that returns or generates the unique ID of an optional compiland
-	static inline uint32_t GetCompilandId(const symbols::Compiland* compiland, const wchar_t* const objPath)
+	// helper function that returns or generates the unique ID of an optional compiland.
+	// for files split off from amalgamated files, we need to use the original object path of the amalgamated file here,
+	// otherwise names of symbols would differ, leading to constructors of global instances being called again.
+	static inline uint32_t GetCompilandId(const symbols::Compiland* compiland, const wchar_t* const objPath, const types::vector<LiveModule::ModifiedObjFile>& modifiedObjFiles)
 	{
-		return compiland
-			? compiland->uniqueId									// compiland exists
-			: uniqueId::Generate(file::NormalizePath(objPath));		// new compiland, generate new unique ID
+		// try to find the given .obj path in the array of modified object files to check if there's an original amalgamated object path for it
+		for (size_t i = 0u; i < modifiedObjFiles.size(); ++i)
+		{
+			const LiveModule::ModifiedObjFile& objFile = modifiedObjFiles[i];
+
+			// don't bother checking strings if the amalgamated object path is empty anyway
+			if (!objFile.amalgamatedObjPath.empty())
+			{
+				if (string::Matches(objPath, objFile.objPath.c_str()))
+				{
+					return uniqueId::Generate(file::NormalizePath(objFile.amalgamatedObjPath.c_str()));
+				}
+			}
+		}
+
+		if (compiland)
+		{
+			// the compiland already exists
+			return compiland->uniqueId;
+		}
+		else
+		{
+			return uniqueId::Generate(file::NormalizePath(objPath));
+		}
 	}
 
 
@@ -441,7 +464,7 @@ namespace
 
 
 	template <typename T>
-	static types::vector<symbols::ObjPath> UpdateCoffCache(const T& compilands, CoffCache<coff::CoffDB>* coffCache, CacheUpdate::Enum updateType, coff::ReadFlags::Enum coffReadFlags)
+	static types::vector<symbols::ObjPath> UpdateCoffCache(const T& compilands, CoffCache<coff::CoffDB>* coffCache, CacheUpdate::Enum updateType, coff::ReadFlags::Enum coffReadFlags, const types::vector<LiveModule::ModifiedObjFile>& modifiedOrNewObjFiles)
 	{
 		LC_LOG_INDENT_DEV;
 
@@ -458,7 +481,7 @@ namespace
 			symbols::ObjPath objPath = it->first;
 			const std::wstring& wideObjPath = string::ToWideString(objPath);
 			const symbols::Compiland* compiland = it->second;
-			const uint32_t compilandUniqueId = GetCompilandId(compiland, wideObjPath.c_str());
+			const uint32_t compilandUniqueId = GetCompilandId(compiland, wideObjPath.c_str(), modifiedOrNewObjFiles);
 
 			const bool shouldUpdate = (updateType == CacheUpdate::NON_EXISTANT)
 				? (coffCache->Lookup(objPath) == nullptr)				// NON-EXISTANT: update cache only for files which don't have an entry yet
@@ -521,7 +544,10 @@ namespace
 		const size_t MODULE_ALIGNMENT = 64u * 1024u;
 		void* preferredBase = process::ScanMemoryRange(processHandle, lowerBound, upperBound, patchSize, MODULE_ALIGNMENT);
 
-		return pointer::AsInteger<executable::PreferredBase>(preferredBase);
+		const executable::PreferredBase preferredImageBase = pointer::AsInteger<executable::PreferredBase>(preferredBase);
+		LC_LOG_DEV("Preferred base address for image: 0x%" PRIX64 " (PID: %d)", preferredImageBase, processId);
+
+		return preferredImageBase;
 	}
 #endif
 
@@ -680,9 +706,11 @@ namespace
 			const size_t count = hookData.data->processes.size();
 			for (size_t p = 0u; p < count; ++p)
 			{
-				const unsigned int pid = hookData.data->processes[p].processId;
-				void* moduleBase = hookData.data->processes[p].moduleBase;
-				const DuplexPipe* pipe = hookData.data->processes[p].pipe;
+				const ModuleCache::ProcessData& processData = hookData.data->processes[p];
+
+				const unsigned int pid = processData.processId;
+				void* moduleBase = processData.moduleBase;
+				const DuplexPipe* pipe = processData.pipe;
 
 				LC_LOG_USER("Calling compile start hooks (PID: %d)", pid);
 				pipe->SendCommandAndWaitForAck(commands::CallHooks { hook::MakeFunction(moduleBase, hookData.firstRva), hook::MakeFunction(moduleBase, hookData.lastRva) });
@@ -705,9 +733,11 @@ namespace
 			const size_t count = hookData.data->processes.size();
 			for (size_t p = 0u; p < count; ++p)
 			{
-				const unsigned int pid = hookData.data->processes[p].processId;
-				void* moduleBase = hookData.data->processes[p].moduleBase;
-				const DuplexPipe* pipe = hookData.data->processes[p].pipe;
+				const ModuleCache::ProcessData& processData = hookData.data->processes[p];
+
+				const unsigned int pid = processData.processId;
+				void* moduleBase = processData.moduleBase;
+				const DuplexPipe* pipe = processData.pipe;
 
 				LC_LOG_USER("Calling compile success hooks (PID: %d)", pid);				
 				pipe->SendCommandAndWaitForAck(commands::CallHooks { hook::MakeFunction(moduleBase, hookData.firstRva), hook::MakeFunction(moduleBase, hookData.lastRva) });
@@ -730,9 +760,11 @@ namespace
 			const size_t count = hookData.data->processes.size();
 			for (size_t p = 0u; p < count; ++p)
 			{
-				const unsigned int pid = hookData.data->processes[p].processId;
-				void* moduleBase = hookData.data->processes[p].moduleBase;
-				const DuplexPipe* pipe = hookData.data->processes[p].pipe;
+				const ModuleCache::ProcessData& processData = hookData.data->processes[p];
+
+				const unsigned int pid = processData.processId;
+				void* moduleBase = processData.moduleBase;
+				const DuplexPipe* pipe = processData.pipe;
 
 				LC_LOG_USER("Calling compile error hooks (PID: %d)", pid);
 				pipe->SendCommandAndWaitForAck(commands::CallHooks { hook::MakeFunction(moduleBase, hookData.firstRva), hook::MakeFunction(moduleBase, hookData.lastRva) });
@@ -1191,11 +1223,11 @@ void LiveModule::UpdateDirectoryCache(DirectoryCache* cache)
 }
 
 
-LiveModule::ErrorType::Enum LiveModule::Update(FileAttributeCache* fileCache, DirectoryCache* directoryCache, UpdateType::Enum updateType, const std::vector<std::wstring>& modifiedOrNewObjFiles)
+LiveModule::ErrorType::Enum LiveModule::Update(FileAttributeCache* fileCache, DirectoryCache* directoryCache, UpdateType::Enum updateType, const types::vector<ModifiedObjFile>& modifiedOrNewObjFiles)
 {
 	telemetry::Scope updateScope("Update live module");
 
-	LC_LOG_DEV("\nLiveModule::Update -------------------------------------------\n");
+	LC_LOG_DEV("LiveModule Update: %S", m_moduleName.c_str());
 
 	GLiveCodingServer->GetStatusChangeDelegate().ExecuteIfBound(L"Checking modified files...");
 
@@ -1331,7 +1363,7 @@ LiveModule::ErrorType::Enum LiveModule::Update(FileAttributeCache* fileCache, Di
 	{
 		for (size_t i = 0u; i < modifiedOrNewObjFiles.size(); ++i)
 		{
-			LC_LOG_USER("File %S was modified or is new", modifiedOrNewObjFiles[i].c_str());
+			LC_LOG_USER("File %S was modified or is new", modifiedOrNewObjFiles[i].objPath.c_str());
 		}
 
 		LC_LOG_USER("Building patch from %zu file(s) for Live Coding module %S", modifiedOrNewObjFiles.size(), m_moduleName.c_str());
@@ -1802,7 +1834,7 @@ LiveModule::ErrorType::Enum LiveModule::Update(FileAttributeCache* fileCache, Di
 			const size_t count = modifiedOrNewObjFiles.size();
 			for (size_t i = 0u; i < count; ++i)
 			{
-				const std::wstring& wideObjPath = modifiedOrNewObjFiles[i];
+				const std::wstring& wideObjPath = modifiedOrNewObjFiles[i].objPath;
 				const symbols::ObjPath& objPath = string::ToUtf8String(wideObjPath);
 				symbols::Compiland* compiland = symbols::FindCompiland(m_compilandDB, objPath);
 
@@ -1945,11 +1977,11 @@ LiveModule::ErrorType::Enum LiveModule::Update(FileAttributeCache* fileCache, Di
 		{
 			for (auto it = modifiedOrNewObjFiles.begin(); it != modifiedOrNewObjFiles.end(); ++it)
 			{
-				const symbols::ObjPath objPath(string::ToUtf8String(*it));
+				const symbols::ObjPath objPath(string::ToUtf8String(it->objPath));
 				const symbols::Compiland* compiland = symbols::FindCompiland(m_compilandDB, objPath);
 
 				// new compilands won't be found in the database, so there's no unique ID yet that we can use
-				const uint32_t compilandUniqueId = GetCompilandId(compiland, it->c_str());
+				const uint32_t compilandUniqueId = GetCompilandId(compiland, it->objPath.c_str(), modifiedOrNewObjFiles);
 
 				// this file was either modified or is new. in any case, the new .OBJ needs to be linked in, even
 				// though the file might be contained in a library.
@@ -2026,11 +2058,11 @@ LiveModule::ErrorType::Enum LiveModule::Update(FileAttributeCache* fileCache, Di
 						// there is no entry yet for this COFF in the cache.
 						// this means that this .obj was not recompiled (otherwise it would have an entry already),
 						// but has been pulled in for the first time due to unresolved symbols.
-						auto task = scheduler::CreateTask(taskRoot, [this, objPath, coffReadFlags]()
+						auto task = scheduler::CreateTask(taskRoot, [this, objPath, coffReadFlags, &modifiedOrNewObjFiles]()
 						{
 							const symbols::Compiland* compiland = symbols::FindCompiland(m_compilandDB, objPath);
 							const std::wstring& wideObjPath = string::ToWideString(objPath);
-							const uint32_t compilandUniqueId = GetCompilandId(compiland, wideObjPath.c_str());
+							const uint32_t compilandUniqueId = GetCompilandId(compiland, wideObjPath.c_str(), modifiedOrNewObjFiles);
 
 								LC_LOG_DEV("Need %s for the first time, updating COFF cache", objPath.c_str());
 
@@ -2094,7 +2126,7 @@ LiveModule::ErrorType::Enum LiveModule::Update(FileAttributeCache* fileCache, Di
 	}
 
 	// update the COFF cache for all compiled files
-	UpdateCoffCache(m_compiledCompilands, m_coffCache, CacheUpdate::ALL, coffReadFlags);
+	UpdateCoffCache(m_compiledCompilands, m_coffCache, CacheUpdate::ALL, coffReadFlags, modifiedOrNewObjFiles);
 
 	GLiveCodingServer->GetStatusChangeDelegate().ExecuteIfBound(L"Stripping COFFs...");
 
@@ -2123,7 +2155,7 @@ LiveModule::ErrorType::Enum LiveModule::Update(FileAttributeCache* fileCache, Di
 			symbols::Compiland* compiland = symbols::FindCompiland(m_compilandDB, objPath);
 			const std::wstring wideObjPath = string::ToWideString(objPath);
 
-			const uint32_t compilandUniqueId = GetCompilandId(compiland, wideObjPath.c_str());
+			const uint32_t compilandUniqueId = GetCompilandId(compiland, wideObjPath.c_str(), modifiedOrNewObjFiles);
 
 			coff::ObjFile* objFile = coff::OpenObj(wideObjPath.c_str());
 			if (objFile && objFile->memoryFile)
@@ -2651,7 +2683,7 @@ LiveModule::ErrorType::Enum LiveModule::Update(FileAttributeCache* fileCache, Di
 		const symbols::ObjPath& objPath = compilandIt->first;
 		const std::wstring& wideObjPath = string::ToWideString(objPath);
 		const symbols::Compiland* compiland = compilandIt->second;
-		const uint32_t compilandUniqueId = GetCompilandId(compiland, wideObjPath.c_str());
+		const uint32_t compilandUniqueId = GetCompilandId(compiland, wideObjPath.c_str(), modifiedOrNewObjFiles);
 
 		coff::ObjFile* coffFile = coff::OpenObj(wideObjPath.c_str());
 		if (coffFile && coffFile->memoryFile)
@@ -3238,7 +3270,7 @@ LiveModule::ErrorType::Enum LiveModule::Update(FileAttributeCache* fileCache, Di
 		// update the COFF cache for new patch compilands.
 		// there may be files for which we don't have a database yet, even though we updated the database for all compiled files.
 		// this can happen when a new .obj that is part of a library is linked in for the first time.
-		types::vector<symbols::ObjPath> updatedCoffs = UpdateCoffCache(patch_compilandDB->compilands, m_coffCache, CacheUpdate::NON_EXISTANT, coffReadFlags);
+		types::vector<symbols::ObjPath> updatedCoffs = UpdateCoffCache(patch_compilandDB->compilands, m_coffCache, CacheUpdate::NON_EXISTANT, coffReadFlags, modifiedOrNewObjFiles);
 
 
 		// similarly, reconstruct symbols and dynamic initializers for new .obj that have been pulled in for the first time.
@@ -3390,9 +3422,11 @@ LiveModule::ErrorType::Enum LiveModule::Update(FileAttributeCache* fileCache, Di
 				const size_t count = hookData.data->processes.size();
 				for (size_t p = 0u; p < count; ++p)
 				{
-					const unsigned int pid = hookData.data->processes[p].processId;
-					void* moduleBase = hookData.data->processes[p].moduleBase;
-					const DuplexPipe* pipe = hookData.data->processes[p].pipe;
+					const ModuleCache::ProcessData& hookProcessData = hookData.data->processes[p];
+
+					const unsigned int pid = hookProcessData.processId;
+					void* moduleBase = hookProcessData.moduleBase;
+					const DuplexPipe* pipe = hookProcessData.pipe;
 
 					LC_LOG_USER("Calling pre-patch hooks (PID: %d)", pid);
 					pipe->SendCommandAndWaitForAck(commands::CallHooks { hook::MakeFunction(moduleBase, hookData.firstRva), hook::MakeFunction(moduleBase, hookData.lastRva) });
@@ -3822,9 +3856,11 @@ LiveModule::ErrorType::Enum LiveModule::Update(FileAttributeCache* fileCache, Di
 				{
 					++functionsPatchedCount;
 
-					const unsigned int pid = originalData.data->processes[p].processId;
-					void* moduleBase = originalData.data->processes[p].moduleBase;
-					process::Handle processHandle = originalData.data->processes[p].processHandle;
+					const ModuleCache::ProcessData& hookProcessData = originalData.data->processes[p];
+
+					const unsigned int pid = hookProcessData.processId;
+					void* moduleBase = hookProcessData.moduleBase;
+					process::Handle processHandle = hookProcessData.processHandle;
 
 					char* originalAddress = pointer::Offset<char*>(moduleBase, originalData.symbol->rva);
 					char* patchAddress = pointer::Offset<char*>(loadedPatches[p], patchSymbol->rva);
@@ -4004,7 +4040,7 @@ bool LiveModule::InstallCompiledPatches(LiveProcess* liveProcess, void* original
 		return true;
 	}
 
-	LC_LOG_DEV("\nLiveModule::InstallCompiledPatches ---------------------------\n");
+	LC_LOG_DEV("LiveModule InstallCompiledPatches: %S", m_moduleName.c_str());
 
 	telemetry::Scope wholeScope("Installing patches");
 
@@ -4251,7 +4287,7 @@ bool LiveModule::LoadPatchInfoAction::Execute(CommandType* command, const Duplex
 void LiveModule::UpdateDirectoryCache(const ImmutableString& path, symbols::Dependency* dependency, DirectoryCache* cache)
 {
 	const std::wstring& directoryOnly = file::GetDirectory(string::ToWideString(path));
-	dependency->parentDirectory = cache->AddDirectory(directoryOnly.c_str());
+	dependency->parentDirectory = cache->AddDirectory(directoryOnly);
 }
 
 
diff --git a/Engine/Source/Developer/Windows/LiveCodingServer/Private/External/LC_LiveModule.h b/Engine/Source/Developer/Windows/LiveCodingServer/Private/External/LC_LiveModule.h
--- a/Engine/Source/Developer/Windows/LiveCodingServer/Private/External/LC_LiveModule.h
+++ b/Engine/Source/Developer/Windows/LiveCodingServer/Private/External/LC_LiveModule.h
@@ -51,6 +51,12 @@ public:
 		};
 	};
 
+	struct ModifiedObjFile
+	{
+		std::wstring objPath;
+		std::wstring amalgamatedObjPath;	// optional
+	};
+
 	struct CompileResult
 	{
 		unsigned int exitCode;
@@ -76,7 +82,7 @@ public:
 	// in DEFAULT mode, this checks for file modifications, compiles files automatically, builds a patch containing changes, and loads them into the host application.
 	// in EXTERNAL_BUILD_SYSTEM mode, this does not compile files but builds a patch containing modified .objs, loading the patch into the host application.
 	// optionally, an array of modified or new .objs can be given in this mode, which builds a patch containing these files, not checking for any other modifications.
-	ErrorType::Enum Update(FileAttributeCache* fileCache, DirectoryCache* directoryCache, UpdateType::Enum updateType, const std::vector<std::wstring>& modifiedOrNewObjFiles);
+	ErrorType::Enum Update(FileAttributeCache* fileCache, DirectoryCache* directoryCache, UpdateType::Enum updateType, const types::vector<ModifiedObjFile>& modifiedOrNewObjFiles);
 	bool InstallCompiledPatches(LiveProcess* liveProcess, void* originalModuleBase);
 
 	const std::wstring& GetModuleName(void) const;
diff --git a/Engine/Source/Developer/Windows/LiveCodingServer/Private/External/LC_PoolAllocator.cpp b/Engine/Source/Developer/Windows/LiveCodingServer/Private/External/LC_PoolAllocator.cpp
--- a/Engine/Source/Developer/Windows/LiveCodingServer/Private/External/LC_PoolAllocator.cpp
+++ b/Engine/Source/Developer/Windows/LiveCodingServer/Private/External/LC_PoolAllocator.cpp
@@ -38,7 +38,7 @@ namespace freeList
 		// the alignment will always satisfy restrictions while producing the least amount of wasted/unused memory.
 		// furthermore, we must always be able to store a Freelist* in each entry.
 		const size_t minimumSize = sizeof(Slot*) > originalElementSize ? sizeof(Slot*) : originalElementSize;
-		const size_t elementSize = RoundUpToMultiple(minimumSize, alignment);
+		const size_t elementSize = freelist::RoundUpToMultiple(minimumSize, alignment);
 
 		// in order to satisfy alignment restrictions, it suffices to offset the start of the free list only
 		as_void = memoryBlock;
diff --git a/Engine/Source/Developer/Windows/LiveCodingServer/Private/External/LC_Scheduler.cpp b/Engine/Source/Developer/Windows/LiveCodingServer/Private/External/LC_Scheduler.cpp
--- a/Engine/Source/Developer/Windows/LiveCodingServer/Private/External/LC_Scheduler.cpp
+++ b/Engine/Source/Developer/Windows/LiveCodingServer/Private/External/LC_Scheduler.cpp
@@ -43,7 +43,7 @@ namespace
 				result = ::GetLogicalProcessorInformation(buffer, &bytesNeeded);
 
 				// return the number of logical processors in case anything went wrong
-				if (result == Windows::FALSE)
+				if ((result == Windows::FALSE) || (buffer == nullptr))
 				{
 					::free(buffer);
 					return GetLogicalProcessorCount();
diff --git a/Engine/Source/Developer/Windows/LiveCodingServer/Private/External/LC_ServerCommandThread.cpp b/Engine/Source/Developer/Windows/LiveCodingServer/Private/External/LC_ServerCommandThread.cpp
--- a/Engine/Source/Developer/Windows/LiveCodingServer/Private/External/LC_ServerCommandThread.cpp
+++ b/Engine/Source/Developer/Windows/LiveCodingServer/Private/External/LC_ServerCommandThread.cpp
@@ -41,7 +41,11 @@ namespace
 	{
 		InitializeCOM(void)
 		{
-			::CoInitialize(NULL);
+			const HRESULT success = ::CoInitialize(NULL);
+			if (success != S_OK)
+			{
+				LC_LOG_DEV("Could not initialize COM. Error: %d", success);
+			}
 		}
 
 		~InitializeCOM(void)
@@ -125,6 +129,7 @@ ServerCommandThread::~ServerCommandThread(void)
 	// this is only called when Live++ is being torn down anyway, so we leave cleanup to the OS.
 	// otherwise we could run into races when trying to terminate the thread that might currently be doing
 	// some intensive work.
+	delete m_directoryCache;
 }
 
 
@@ -703,7 +708,8 @@ void ServerCommandThread::CompileChanges(bool didAllProcessesMakeProgress)
 		GLiveCodingServer->GetStatusChangeDelegate().ExecuteIfBound(L"Compiling changes for live coding...");
 
 		TMap<FString, TArray<FString>> ModuleToObjectFiles;
-		if (!CompileDelegate.Execute(Targets, ModuleToObjectFiles))
+		TMap<FString, FString> ObjectFileToUnityObjectFile;
+		if (!CompileDelegate.Execute(Targets, ModuleToObjectFiles, ObjectFileToUnityObjectFile))
 		{
 			GLiveCodingServer->GetCompileFinishedDelegate().ExecuteIfBound(ELiveCodingResult::Error, L"Compilation error.");
 			return;
@@ -766,11 +772,19 @@ void ServerCommandThread::CompileChanges(bool didAllProcessesMakeProgress)
 				return;
 			}
 
-			types::vector<std::wstring> ObjectFiles;
+			types::vector<LiveModule::ModifiedObjFile> ObjectFiles;
 			for(const FString& ObjectFile : Pair.Value)
 			{
-				std::wstring NormalizedObjectFile = file::NormalizePath(*ObjectFile);
-				ObjectFiles.push_back(std::move(NormalizedObjectFile));
+				LiveModule::ModifiedObjFile ModifiedObjFile;
+				ModifiedObjFile.objPath = file::NormalizePath(*ObjectFile);
+
+				const FString* UnityObjectFile = ObjectFileToUnityObjectFile.Find(ObjectFile);
+				if (UnityObjectFile != nullptr)
+				{
+					ModifiedObjFile.amalgamatedObjPath = file::NormalizePath(**UnityObjectFile);
+				}
+
+				ObjectFiles.push_back(std::move(ModifiedObjFile));
 			}
 
 			m_liveModuleToModifiedOrNewObjFiles.insert(std::make_pair(ModuleFileName, std::move(ObjectFiles)));
@@ -823,14 +837,14 @@ void ServerCommandThread::CompileChanges(bool didAllProcessesMakeProgress)
 			else
 			{
 				// build a patch with the given list of .objs for this module
-				const std::vector<std::wstring>& objFiles = objFilesIt->second;
+				const types::vector<LiveModule::ModifiedObjFile>& objFiles = objFilesIt->second;
 				moduleUpdateError = liveModule->Update(&fileCache, m_directoryCache, updateType, objFiles);
 			}
 		}
 		else
 		{
 			// no optional .objs were given, update all live modules regularly
-			std::vector<std::wstring> emptyObjs;
+			types::vector<LiveModule::ModifiedObjFile> emptyObjs;
 			moduleUpdateError = liveModule->Update(&fileCache, m_directoryCache, updateType, emptyObjs);
 		}
 
@@ -1323,7 +1337,10 @@ bool ServerCommandThread::BuildPatchAction::Execute(CommandType* command, const
 
 		pipe->SendAck();
 
-		commandThread->m_liveModuleToModifiedOrNewObjFiles[packetCommand.moduleName].push_back(packetCommand.objPath);
+		LiveModule::ModifiedObjFile modifiedObjFile;
+		modifiedObjFile.objPath = packetCommand.objPath;
+		modifiedObjFile.amalgamatedObjPath = packetCommand.amalgamatedObjPath;
+		commandThread->m_liveModuleToModifiedOrNewObjFiles[packetCommand.moduleName].push_back(modifiedObjFile);
 	}
 
 	commandThread->m_manualRecompileTriggered = true;
diff --git a/Engine/Source/Developer/Windows/LiveCodingServer/Private/External/LC_ServerCommandThread.h b/Engine/Source/Developer/Windows/LiveCodingServer/Private/External/LC_ServerCommandThread.h
--- a/Engine/Source/Developer/Windows/LiveCodingServer/Private/External/LC_ServerCommandThread.h
+++ b/Engine/Source/Developer/Windows/LiveCodingServer/Private/External/LC_ServerCommandThread.h
@@ -13,6 +13,7 @@
 #include "LC_Executable.h"
 #include "LC_RunMode.h"
 #include "LC_Types.h"
+#include "LC_LiveModule.h"
 
 
 class MainFrame;
@@ -261,5 +262,5 @@ private:
 
 	// for triggering recompiles using the API
 	bool m_manualRecompileTriggered;
-	types::unordered_map<std::wstring, types::vector<std::wstring>> m_liveModuleToModifiedOrNewObjFiles;
+	types::unordered_map<std::wstring, types::vector<LiveModule::ModifiedObjFile>> m_liveModuleToModifiedOrNewObjFiles;
 };
diff --git a/Engine/Source/Developer/Windows/LiveCodingServer/Private/External/LC_Settings.cpp b/Engine/Source/Developer/Windows/LiveCodingServer/Private/External/LC_Settings.cpp
--- a/Engine/Source/Developer/Windows/LiveCodingServer/Private/External/LC_Settings.cpp
+++ b/Engine/Source/Developer/Windows/LiveCodingServer/Private/External/LC_Settings.cpp
@@ -253,6 +253,7 @@ SettingString::SettingString(const wchar_t* group, const wchar_t* name, const wc
 	, m_shortDescription(shortDescription)
 	, m_description(description)
 	, m_value(initialValue)
+	, m_previousValue(initialValue)
 {
 	const std::wstring value = LoadSetting(m_group, m_name, initialValue);
 
@@ -273,6 +274,7 @@ void SettingString::SetValue(const wchar_t* value)
 
 void SettingString::SetValueWithoutSaving(const wchar_t* value)
 {
+	m_previousValue = m_value;
 	m_value = value;
 }
 
@@ -301,6 +303,12 @@ const wchar_t* SettingString::GetValue(void) const
 }
 
 
+const wchar_t* SettingString::GetPreviousValue(void) const
+{
+	return m_previousValue.c_str();
+}
+
+
 SettingIntProxy::SettingIntProxy(SettingInt* setting)
 	: Setting(Type::INTEGER_PROXY)
 	, m_setting(setting)
diff --git a/Engine/Source/Developer/Windows/LiveCodingServer/Private/External/LC_Settings.h b/Engine/Source/Developer/Windows/LiveCodingServer/Private/External/LC_Settings.h
--- a/Engine/Source/Developer/Windows/LiveCodingServer/Private/External/LC_Settings.h
+++ b/Engine/Source/Developer/Windows/LiveCodingServer/Private/External/LC_Settings.h
@@ -111,6 +111,7 @@ public:
 	const wchar_t* GetShortDescription(void) const;
 	const wchar_t* GetDescription(void) const;
 	const wchar_t* GetValue(void) const;
+	const wchar_t* GetPreviousValue(void) const;
 
 private:
 	const wchar_t* m_group;
@@ -118,6 +119,7 @@ private:
 	const wchar_t* m_shortDescription;
 	const wchar_t* m_description;
 	std::wstring m_value;
+	std::wstring m_previousValue;
 };
 
 
diff --git a/Engine/Source/Developer/Windows/LiveCodingServer/Private/External/LC_SymbolInfo.cpp b/Engine/Source/Developer/Windows/LiveCodingServer/Private/External/LC_SymbolInfo.cpp
new file mode 100644
--- /dev/null
+++ b/Engine/Source/Developer/Windows/LiveCodingServer/Private/External/LC_SymbolInfo.cpp
@@ -0,0 +1,14 @@
+// Copyright 2011-2019 Molecular Matters GmbH, all rights reserved.
+
+#include "LC_SymbolInfo.h"
+#include <stdio.h>
+
+
+SymbolInfo::SymbolInfo(const char* const function, const char* const filename, unsigned int line)
+	: m_function()
+	, m_filename()
+	, m_line(line)
+{
+	strcpy_s(m_function, function);
+	strcpy_s(m_filename, filename);
+}
diff --git a/Engine/Source/Developer/Windows/LiveCodingServer/Private/External/LC_SymbolInfo.h b/Engine/Source/Developer/Windows/LiveCodingServer/Private/External/LC_SymbolInfo.h
new file mode 100644
--- /dev/null
+++ b/Engine/Source/Developer/Windows/LiveCodingServer/Private/External/LC_SymbolInfo.h
@@ -0,0 +1,34 @@
+// Copyright 2011-2019 Molecular Matters GmbH, all rights reserved.
+
+#pragma once
+
+#include "CoreTypes.h"
+#include "LC_Platform.h"
+
+class SymbolInfo
+{
+public:
+	SymbolInfo(const char* const function, const char* const filename, unsigned int line);
+
+	inline const char* GetFunction(void) const
+	{
+		return m_function;
+	}
+
+	inline const char* GetFilename(void) const
+	{
+		return m_filename;
+	}
+
+	inline unsigned int GetLine(void) const
+	{
+		return m_line;
+	}
+
+private:
+	LC_DISABLE_ASSIGNMENT(SymbolInfo);
+
+	char m_function[512u];
+	char m_filename[WINDOWS_MAX_PATH];
+	unsigned int m_line;
+};
diff --git a/Engine/Source/Developer/Windows/LiveCodingServer/Private/External/LC_SymbolResolution.cpp b/Engine/Source/Developer/Windows/LiveCodingServer/Private/External/LC_SymbolResolution.cpp
new file mode 100644
--- /dev/null
+++ b/Engine/Source/Developer/Windows/LiveCodingServer/Private/External/LC_SymbolResolution.cpp
@@ -0,0 +1,65 @@
+// Copyright 2011-2019 Molecular Matters GmbH, all rights reserved.
+
+#include "LC_SymbolResolution.h"
+#include "LC_PointerUtil.h"
+#include "LC_Logging.h"
+#include "Windows/WindowsHWrapper.h"
+
+#pragma push_macro("OPTIONAL")
+#ifndef OPTIONAL
+#define OPTIONAL
+#endif
+#include <dbghelp.h>
+#pragma pop_macro("OPTIONAL")
+
+namespace symbolResolution
+{
+	void Startup(void)
+	{
+		::SymSetOptions(SYMOPT_DEFERRED_LOADS | SYMOPT_LOAD_LINES | SYMOPT_UNDNAME | SYMOPT_INCLUDE_32BIT_MODULES);
+		const BOOL success = ::SymInitialize(::GetCurrentProcess(), "", true);
+		if (success == Windows::FALSE)
+		{
+			const DWORD error = ::GetLastError();
+			LC_ERROR_DEV("SymInitialize failed with error 0x%X", error);
+		}
+	}
+
+
+	void Shutdown(void)
+	{
+		const BOOL success = ::SymCleanup(::GetCurrentProcess());
+		if (success == Windows::FALSE)
+		{
+			const DWORD error = ::GetLastError();
+			LC_ERROR_DEV("SymCleanup failed with error 0x%X", error);
+		}
+	}
+
+
+	SymbolInfo ResolveSymbolsForAddress(const void* const address)
+	{
+		// retrieve function
+		DWORD64 displacement64 = 0;
+		ULONG64 buffer[(sizeof(SYMBOL_INFO) + MAX_SYM_NAME * sizeof(TCHAR) + sizeof(ULONG64) - 1) / sizeof(ULONG64)] = {};
+		SYMBOL_INFO* symbolInfo = reinterpret_cast<SYMBOL_INFO*>(buffer);
+		symbolInfo->SizeOfStruct = sizeof(SYMBOL_INFO);
+		symbolInfo->MaxNameLen = MAX_SYM_NAME;
+
+		const HANDLE process = ::GetCurrentProcess();
+		const DWORD64 intAddress = pointer::AsInteger<DWORD64>(address);
+
+		if (::SymFromAddr(process, intAddress, &displacement64, symbolInfo))
+		{
+			// retrieve filename and line number
+			DWORD displacement = 0;
+			IMAGEHLP_LINE64 line = {};
+			line.SizeOfStruct = sizeof(IMAGEHLP_LINE64);
+
+			if (::SymGetLineFromAddr64(process, intAddress, &displacement, &line))
+				return SymbolInfo(symbolInfo->Name, line.FileName, line.LineNumber);
+		}
+
+		return SymbolInfo("unknown", "unknown", 0);
+	}
+}
diff --git a/Engine/Source/Developer/Windows/LiveCodingServer/Private/External/LC_SymbolResolution.h b/Engine/Source/Developer/Windows/LiveCodingServer/Private/External/LC_SymbolResolution.h
new file mode 100644
--- /dev/null
+++ b/Engine/Source/Developer/Windows/LiveCodingServer/Private/External/LC_SymbolResolution.h
@@ -0,0 +1,15 @@
+// Copyright 2011-2019 Molecular Matters GmbH, all rights reserved.
+
+#pragma once
+
+#include "LC_SymbolInfo.h"
+
+
+namespace symbolResolution
+{
+	void Startup(void);
+	void Shutdown(void);
+
+	// Resolves the symbol information for a given address
+	SymbolInfo ResolveSymbolsForAddress(const void* const address);
+}
diff --git a/Engine/Source/Developer/Windows/LiveCodingServer/Private/External/LC_Symbols.cpp b/Engine/Source/Developer/Windows/LiveCodingServer/Private/External/LC_Symbols.cpp
--- a/Engine/Source/Developer/Windows/LiveCodingServer/Private/External/LC_Symbols.cpp
+++ b/Engine/Source/Developer/Windows/LiveCodingServer/Private/External/LC_Symbols.cpp
@@ -337,30 +337,15 @@ namespace
 
 	static bool IsCppOrCFile(const std::wstring& normalizedLowercaseFilename)
 	{
-		const std::wstring& extension = file::GetExtension(normalizedLowercaseFilename);
-		if (string::Matches(extension.c_str(), L".cpp"))
-		{
-			return true;
-		}
-		else if (string::Matches(extension.c_str(), L".c"))
-		{
-			return true;
-		}
-		else if (string::Matches(extension.c_str(), L".cc"))
-		{
-			return true;
-		}
-		else if (string::Matches(extension.c_str(), L".c++"))
-		{
-			return true;
-		}
-		else if (string::Matches(extension.c_str(), L".cp"))
-		{
-			return true;
-		}
-		else if (string::Matches(extension.c_str(), L".cxx"))
+		const std::wstring& filenameExtension = file::GetExtension(normalizedLowercaseFilename);
+		const types::vector<std::wstring>& amalgamatedCppFileExtensions = appSettings::GetAmalgamatedCppFileExtensions();
+
+		for (size_t i = 0u; i < amalgamatedCppFileExtensions.size(); ++i)
 		{
-			return true;
+			if (string::Matches(filenameExtension.c_str(), amalgamatedCppFileExtensions[i].c_str()))
+			{
+				return true;
+			}
 		}
 
 		return false;
@@ -418,6 +403,10 @@ namespace symbols
 		{
 			const types::vector<IDiaSymbol*>& publicSymbols = dia::GatherChildSymbols(provider->globalScope, SymTagPublicSymbol);
 			const size_t symbolCount = publicSymbols.size();
+			symbolDB->symbolsByName.reserve(symbolCount);
+			symbolDB->symbolsByRva.reserve(symbolCount);
+			symbolDB->patchableFunctionSymbols.reserve(symbolCount);
+
 			for (size_t i = 0u; i < symbolCount; ++i)
 			{
 				IDiaSymbol* publicSymbol = publicSymbols[i];
@@ -649,7 +638,8 @@ namespace symbols
 			const dia::SymbolName& diaCompilandPath = dia::GetSymbolName(diaSymbol);
 			std::wstring compilandPath(diaCompilandPath.GetString());
 			const std::wstring& uppercaseCompilandPath = string::ToUpper(compilandPath);
-			const bool isObjPath = string::Contains(uppercaseCompilandPath.c_str(), L".OBJ");
+
+			const bool isObjPath = string::Contains(uppercaseCompilandPath.c_str(), L".OBJ") || string::Contains(uppercaseCompilandPath.c_str(), L".O");
 			if (isObjPath)
 			{
 				// a valid compiland, gather more information.
@@ -752,9 +742,20 @@ namespace symbols
 					continue;
 				}
 
-				// only add compilands that exist on disk
+				// only add compilands that exist on disk.
+				// optimization: ignore files stored on optical drives because they cannot be changed anyway.
 				{
 					// test the compiland path first
+					if (file::GetDriveType(compilandPath.c_str()) == file::DriveType::OPTICAL)
+					{
+						if (generateLogs)
+						{
+							LC_LOG_DEV("Ignoring file %S on optical drive", compilandPath.c_str());
+						}
+
+						continue;
+					}
+
 					FileAttributeCache::Data cacheData = fileCache.UpdateCacheData(compilandPath);
 					if (!cacheData.exists)
 					{
@@ -773,6 +774,16 @@ namespace symbols
 							testPath += L"\\";
 							testPath += file::GetFilename(compilandPath);
 
+							if (file::GetDriveType(testPath.c_str()) == file::DriveType::OPTICAL)
+							{
+								if (generateLogs)
+								{
+									LC_LOG_DEV("Ignoring file %S on optical drive", testPath.c_str());
+								}
+
+								continue;
+							}
+
 							cacheData = fileCache.UpdateCacheData(testPath);
 							if (!cacheData.exists)
 							{
@@ -794,6 +805,17 @@ namespace symbols
 								testPath = optionsCache[2];
 								testPath += L"\\";
 								testPath += compilandPath;
+
+								if (file::GetDriveType(testPath.c_str()) == file::DriveType::OPTICAL)
+								{
+									if (generateLogs)
+									{
+										LC_LOG_DEV("Ignoring file %S on optical drive", testPath.c_str());
+									}
+
+									continue;
+								}
+
 								cacheData = fileCache.UpdateCacheData(testPath);
 								testFileExists = cacheData.exists;
 							}
@@ -899,6 +921,7 @@ namespace symbols
 				// remote filenames in distributed builds.
 				// if we find a file dependency matching the given source file, we take that one instead to get
 				// full absolute file paths.
+				// optimization: ignore all files on optical drives because they cannot be changed anyway
 				const std::wstring srcFileOnlyLowercase = string::ToLower(file::GetFilename(optionsCache[0]));
 
 				const ObjPath objPath(string::ToUtf8String(normalizedCompilandPath));
@@ -908,6 +931,16 @@ namespace symbols
 					// we are only interested in tracking .obj files. we will never be able to recompile files
 					// and we don't know anything about source files, dependencies, etc.
 					// but we still use our dependency tracking system by letting each .obj depend on itself.
+					if (file::GetDriveType(normalizedCompilandPath.c_str()) == file::DriveType::OPTICAL)
+					{
+						if (generateLogs)
+						{
+							LC_LOG_DEV("Ignoring file %S on optical drive", normalizedCompilandPath.c_str());
+						}
+
+						continue;
+					}
+
 					const FileAttributeCache::Data& cacheData = fileCache.UpdateCacheData(normalizedCompilandPath);
 					if (cacheData.exists)
 					{
@@ -992,26 +1025,36 @@ namespace symbols
 						const FileInfo& fileInfo = fileInfos[j];
 						const std::wstring& normalizedFilename = fileInfo.normalizedFilename;
 
-						const FileAttributeCache::Data& cacheData = fileCache.UpdateCacheData(normalizedFilename);
-						if (cacheData.exists)
+						if (file::GetDriveType(normalizedFilename.c_str()) == file::DriveType::OPTICAL)
 						{
 							if (generateLogs)
 							{
-								LC_LOG_DEV("Dependency %S", normalizedFilename.c_str());
+								LC_LOG_DEV("Ignoring file %S on optical drive", normalizedFilename.c_str());
 							}
+						}
+						else
+						{
+							const FileAttributeCache::Data& cacheData = fileCache.UpdateCacheData(normalizedFilename);
+							if (cacheData.exists)
+							{
+								if (generateLogs)
+								{
+									LC_LOG_DEV("Dependency %S", normalizedFilename.c_str());
+								}
 
-							const ImmutableString& sourceFilePath = string::ToUtf8String(normalizedFilename);
-							AddFileDependency(compilandDb, sourceFilePath, objPath, cacheData.lastModificationTime);
+								const ImmutableString& sourceFilePath = string::ToUtf8String(normalizedFilename);
+								AddFileDependency(compilandDb, sourceFilePath, objPath, cacheData.lastModificationTime);
 
-							if (splitAmalgamatedFiles)
+								if (splitAmalgamatedFiles)
+								{
+									compiland->sourceFiles->files.push_back(sourceFilePath);
+								}
+							}
+							else if (generateLogs)
 							{
-								compiland->sourceFiles->files.push_back(sourceFilePath);
+								LC_LOG_DEV("Missing dependency %S", normalizedFilename.c_str());
 							}
 						}
-						else if (generateLogs)
-						{
-							LC_LOG_DEV("Missing dependency %S", normalizedFilename.c_str());
-						}
 					}
 
 					compilandDb->compilands.insert(std::make_pair(objPath, compiland));
@@ -1042,19 +1085,29 @@ namespace symbols
 							{
 								LC_LOG_DEV("Main .cpp %S", normalizedFilename.c_str());
 
-								const FileAttributeCache::Data& cacheData = fileCache.UpdateCacheData(normalizedFilename);
-								if (cacheData.exists)
+								if (file::GetDriveType(normalizedFilename.c_str()) == file::DriveType::OPTICAL)
 								{
 									if (generateLogs)
 									{
-										LC_LOG_DEV("Dependency %S", normalizedFilename.c_str());
+										LC_LOG_DEV("Ignoring file %S on optical drive", normalizedFilename.c_str());
 									}
-
-									AddFileDependency(compilandDb, sourceFilePath, objPath, cacheData.lastModificationTime);
 								}
-								else if (generateLogs)
+								else
 								{
-									LC_LOG_DEV("Missing dependency %S", normalizedFilename.c_str());
+									const FileAttributeCache::Data& cacheData = fileCache.UpdateCacheData(normalizedFilename);
+									if (cacheData.exists)
+									{
+										if (generateLogs)
+										{
+											LC_LOG_DEV("Dependency %S", normalizedFilename.c_str());
+										}
+
+										AddFileDependency(compilandDb, sourceFilePath, objPath, cacheData.lastModificationTime);
+									}
+									else if (generateLogs)
+									{
+										LC_LOG_DEV("Missing dependency %S", normalizedFilename.c_str());
+									}
 								}
 							}
 							else
@@ -1064,95 +1117,115 @@ namespace symbols
 								// this source file will not trigger a build of the amalgamated file.
 								LC_LOG_DEV("Included .cpp %S", normalizedFilename.c_str());
 
-								const FileAttributeCache::Data& cacheData = fileCache.UpdateCacheData(normalizedFilename);
-								if (cacheData.exists)
+								if (file::GetDriveType(normalizedFilename.c_str()) == file::DriveType::OPTICAL)
 								{
 									if (generateLogs)
 									{
-										LC_LOG_DEV("Dependency %S", normalizedFilename.c_str());
+										LC_LOG_DEV("Ignoring file %S on optical drive", normalizedFilename.c_str());
 									}
+								}
+								else
+								{
+									const FileAttributeCache::Data& cacheData = fileCache.UpdateCacheData(normalizedFilename);
+									if (cacheData.exists)
+									{
+										if (generateLogs)
+										{
+											LC_LOG_DEV("Dependency %S", normalizedFilename.c_str());
+										}
 
-									// create new .obj path by appending this file name to the real .obj, e.g.
-									// Amalgamated.obj turns into Amalgamated.lpp_part.ASingleFile.obj.
-									const std::wstring newObjPart = amalgamation::CreateObjPart(file::NormalizePath(normalizedFilename.c_str()));
-									const std::wstring newObjPath = amalgamation::CreateObjPath(compilandPath, newObjPart);
-									const std::wstring normalizedNewObjPath = file::NormalizePath(newObjPath.c_str());
-
-									AddFileDependency(compilandDb, sourceFilePath, string::ToUtf8String(normalizedNewObjPath), cacheData.lastModificationTime);
+										// create new .obj path by appending this file name to the real .obj, e.g.
+										// Amalgamated.obj turns into Amalgamated.lpp_part.ASingleFile.obj.
+										const std::wstring newObjPart = amalgamation::CreateObjPart(file::NormalizePath(normalizedFilename.c_str()));
+										const std::wstring newObjPath = amalgamation::CreateObjPath(compilandPath, newObjPart);
+										const std::wstring normalizedNewObjPath = file::NormalizePath(newObjPath.c_str());
 
-									// create a new compiland matching this .obj.
-									// we could use different PDBs for different files when no PCHs are being used, but with
-									// our automatic multi-processor compilation this doesn't really gain anything performance-wise
-									// and just complicates things.
+										AddFileDependency(compilandDb, sourceFilePath, string::ToUtf8String(normalizedNewObjPath), cacheData.lastModificationTime);
 
-									// adapt command line to accommodate new .obj path
-									const std::wstring& newCommandLine = string::Replace(optionsCache[4], L".obj", newObjPart);
+										// create a new compiland matching this .obj.
+										// we could use different PDBs for different files when no PCHs are being used, but with
+										// our automatic multi-processor compilation this doesn't really gain anything performance-wise
+										// and just complicates things.
 
-									Compiland* newCompiland = LC_NEW(&g_compilandAllocator, Compiland)
-									{
-										string::ToUtf8String(newObjPath),
-										string::ToUtf8String(normalizedFilename),
-										envPdbPath,
-										envCompilerPath,
-										string::ToUtf8String(newCommandLine),
-										envCompilerWorkingDirectory,
-										objPath,										// .obj of the amalgamation
-										nullptr,										// file indices
-
-										// note that for the purpose of disambiguating symbols in COFF files,
-										// we treat these files as being the amalgamated file.
-										// symbols originally coming from amalgamated files need to have the same
-										// name as symbols from individual files.
-										uniqueId::Generate(normalizedCompilandPath),
-
-										// same for the DIA symbol index, for the same reason
-										static_cast<uint32_t>(i),
-
-										Compiland::Type::PART_OF_AMALGAMATION,			// type of file
-										isPartOfLibrary,								// isPartOfLibrary
-										false											// wasRecompiled
-									};
-
-									compilandDb->compilands.insert(std::make_pair(string::ToUtf8String(normalizedNewObjPath), newCompiland));
-
-									// try updating the amalgamated compiland for the given file and create a new one in case none exists yet
-									{
-										const auto& insertPair = compilandDb->amalgamatedCompilands.emplace(objPath, nullptr);
-										symbols::AmalgamatedCompiland*& amalgamatedCompiland = insertPair.first->second;
+										// adapt command line to accommodate new .obj path
+										const std::wstring& newCommandLine = string::Replace(optionsCache[4], L".obj", newObjPart);
 
-										if (insertPair.second)
+										Compiland* newCompiland = LC_NEW(&g_compilandAllocator, Compiland)
 										{
-											// insertion was successful, create a new amalgamated compiland
-											amalgamatedCompiland = LC_NEW(&g_amalgamatedCompilandAllocator, symbols::AmalgamatedCompiland);
-											amalgamatedCompiland->isSplit = false;
+											string::ToUtf8String(newObjPath),
+											string::ToUtf8String(normalizedFilename),
+											envPdbPath,
+											envCompilerPath,
+											string::ToUtf8String(newCommandLine),
+											envCompilerWorkingDirectory,
+											objPath,										// .obj of the amalgamation
+											nullptr,										// file indices
+
+											// note that for the purpose of disambiguating symbols in COFF files,
+											// we treat these files as being the amalgamated file.
+											// symbols originally coming from amalgamated files need to have the same
+											// name as symbols from individual files.
+											uniqueId::Generate(normalizedCompilandPath),
+
+											// same for the DIA symbol index, for the same reason
+											static_cast<uint32_t>(i),
+
+											Compiland::Type::PART_OF_AMALGAMATION,			// type of file
+											isPartOfLibrary,								// isPartOfLibrary
+											false											// wasRecompiled
+										};
+
+										compilandDb->compilands.insert(std::make_pair(string::ToUtf8String(normalizedNewObjPath), newCompiland));
+
+										// try updating the amalgamated compiland for the given file and create a new one in case none exists yet
+										{
+											const auto& insertPair = compilandDb->amalgamatedCompilands.emplace(objPath, nullptr);
+											symbols::AmalgamatedCompiland*& amalgamatedCompiland = insertPair.first->second;
+
+											if (insertPair.second)
+											{
+												// insertion was successful, create a new amalgamated compiland
+												amalgamatedCompiland = LC_NEW(&g_amalgamatedCompilandAllocator, symbols::AmalgamatedCompiland);
+												amalgamatedCompiland->isSplit = false;
+											}
+
+											// update entry
+											amalgamatedCompiland->singleParts.push_back(string::ToUtf8String(normalizedNewObjPath));
 										}
-
-										// update entry
-										amalgamatedCompiland->singleParts.push_back(string::ToUtf8String(normalizedNewObjPath));
 									}
-								}
-								else if (generateLogs)
-								{
-									LC_LOG_DEV("Missing dependency %S", normalizedFilename.c_str());
+									else if (generateLogs)
+									{
+										LC_LOG_DEV("Missing dependency %S", normalizedFilename.c_str());
+									}
 								}
 							}
 						}
 						else
 						{
 							// this is a header file. add it as regular dependency for the main amalgamated .obj file
-							const FileAttributeCache::Data& cacheData = fileCache.UpdateCacheData(normalizedFilename);
-							if (cacheData.exists)
+							if (file::GetDriveType(normalizedFilename.c_str()) == file::DriveType::OPTICAL)
 							{
 								if (generateLogs)
 								{
-									LC_LOG_DEV("Dependency %S", normalizedFilename.c_str());
+									LC_LOG_DEV("Ignoring file %S on optical drive", normalizedFilename.c_str());
 								}
-
-								AddFileDependency(compilandDb, sourceFilePath, objPath, cacheData.lastModificationTime);
 							}
-							else if (generateLogs)
+							else
 							{
-								LC_LOG_DEV("Missing dependency %S", normalizedFilename.c_str());
+								const FileAttributeCache::Data& cacheData = fileCache.UpdateCacheData(normalizedFilename);
+								if (cacheData.exists)
+								{
+									if (generateLogs)
+									{
+										LC_LOG_DEV("Dependency %S", normalizedFilename.c_str());
+									}
+
+									AddFileDependency(compilandDb, sourceFilePath, objPath, cacheData.lastModificationTime);
+								}
+								else if (generateLogs)
+								{
+									LC_LOG_DEV("Missing dependency %S", normalizedFilename.c_str());
+								}
 							}
 						}
 					}
@@ -1464,14 +1537,14 @@ namespace symbols
 		const symbols::ImageSection* firstSection = symbols::FindImageSectionByRVA(imageSectionDb, firstRva);
 		if (!firstSection)
 		{
-			LC_ERROR_DEV("Cannot find image section holding start of dynamic initializer range");
+			LC_ERROR_USER("Could not find image section containing dynamic initializers.\nThis will lead to constructors of global and static variables being called again for the next patch, likely leading to unexpected behaviour.");
 			return initializerDb;
 		}
 
 		const symbols::ImageSection* lastSection = symbols::FindImageSectionByRVA(imageSectionDb, lastRva);
 		if (!lastSection)
 		{
-			LC_ERROR_DEV("Cannot find image section holding end of dynamic initializer range");
+			LC_ERROR_USER("Could not find image section containing dynamic initializers.\nThis will lead to constructors of global and static variables being called again for the next patch, likely leading to unexpected behaviour.");
 			return initializerDb;
 		}
 
diff --git a/Engine/Source/Developer/Windows/LiveCodingServer/Private/External/LC_Symbols.h b/Engine/Source/Developer/Windows/LiveCodingServer/Private/External/LC_Symbols.h
--- a/Engine/Source/Developer/Windows/LiveCodingServer/Private/External/LC_Symbols.h
+++ b/Engine/Source/Developer/Windows/LiveCodingServer/Private/External/LC_Symbols.h
@@ -27,6 +27,7 @@ namespace symbols
 		uint32_t rva;
 
 		LC_DISABLE_ASSIGNMENT(Symbol);
+		LC_DISABLE_COPY(Symbol);
 	};
 
 	struct SymbolDB
@@ -88,6 +89,7 @@ namespace symbols
 		bool wasRecompiled;
 
 		LC_DISABLE_ASSIGNMENT(Compiland);
+		LC_DISABLE_COPY(Compiland);
 	};
 
 	struct AmalgamatedCompiland
diff --git a/Engine/Source/Developer/Windows/LiveCodingServer/Private/External/LC_Telemetry.cpp b/Engine/Source/Developer/Windows/LiveCodingServer/Private/External/LC_Telemetry.cpp
--- a/Engine/Source/Developer/Windows/LiveCodingServer/Private/External/LC_Telemetry.cpp
+++ b/Engine/Source/Developer/Windows/LiveCodingServer/Private/External/LC_Telemetry.cpp
@@ -31,12 +31,15 @@ namespace
 telemetry::Scope::Scope(const char* name)
 	: m_name(name)
 	, m_start(std::chrono::high_resolution_clock::now())
+	, m_cs()
 {
 }
 
 
 telemetry::Scope::~Scope(void)
 {
+	CriticalSection::ScopedLock lock(&m_cs);
+
 	if (m_name)
 	{
 		Print(m_name, m_start);
@@ -70,6 +73,8 @@ void telemetry::Scope::Restart(void)
 
 void telemetry::Scope::End(void)
 {
+	CriticalSection::ScopedLock lock(&m_cs);
+
 	Print(m_name, m_start);
 
 	// do not print again when going out of scope
diff --git a/Engine/Source/Developer/Windows/LiveCodingServer/Private/External/LC_Telemetry.h b/Engine/Source/Developer/Windows/LiveCodingServer/Private/External/LC_Telemetry.h
--- a/Engine/Source/Developer/Windows/LiveCodingServer/Private/External/LC_Telemetry.h
+++ b/Engine/Source/Developer/Windows/LiveCodingServer/Private/External/LC_Telemetry.h
@@ -4,6 +4,7 @@
 
 #include "CoreTypes.h"
 #include <chrono>
+#include "LC_CriticalSection.h"
 
 namespace telemetry
 {
@@ -23,6 +24,7 @@ namespace telemetry
 	private:
 		const char* m_name;
 		std::chrono::high_resolution_clock::time_point m_start;
+		CriticalSection m_cs;
 	};
 
 
diff --git a/Engine/Source/Developer/Windows/LiveCodingServer/Public/ILiveCodingServer.h b/Engine/Source/Developer/Windows/LiveCodingServer/Public/ILiveCodingServer.h
--- a/Engine/Source/Developer/Windows/LiveCodingServer/Public/ILiveCodingServer.h
+++ b/Engine/Source/Developer/Windows/LiveCodingServer/Public/ILiveCodingServer.h
@@ -44,7 +44,8 @@ public:
 	virtual FLogOutputDelegate& GetLogOutputDelegate() = 0;
 
 	typedef TMap<FString, TArray<FString>> FModuleToObjectFiles;
-	DECLARE_DELEGATE_RetVal_TwoParams(bool, FCompileDelegate, const TArray<FString>&, FModuleToObjectFiles&)
+	typedef TMap<FString, FString> FObjectFileToUnityObjectFile;
+	DECLARE_DELEGATE_RetVal_ThreeParams(bool, FCompileDelegate, const TArray<FString>&, FModuleToObjectFiles&, FObjectFileToUnityObjectFile&)
 	virtual FCompileDelegate& GetCompileDelegate() = 0;
 
 	DECLARE_DELEGATE(FCompileStartedDelegate);
diff --git a/Engine/Source/Programs/LiveCodingConsole/Private/LiveCodingConsole.cpp b/Engine/Source/Programs/LiveCodingConsole/Private/LiveCodingConsole.cpp
--- a/Engine/Source/Programs/LiveCodingConsole/Private/LiveCodingConsole.cpp
+++ b/Engine/Source/Programs/LiveCodingConsole/Private/LiveCodingConsole.cpp
@@ -179,7 +179,7 @@ private:
 		LogWidget->AppendLine(GetLogColor(Verbosity), MoveTemp(Text));
 	}
 
-	bool CompilePatch(const TArray<FString>& Targets, TMap<FString, TArray<FString>>& ModuleToObjectFiles)
+	bool CompilePatch(const TArray<FString>& Targets, TMap<FString, TArray<FString>>& ModuleToObjectFiles, TMap<FString, FString>& ObjectFileToUnityObjectFile)
 	{
 		// Update the compile start time. This gets copied into the last patch time once a patch has been confirmed to have been applied.
 		NextPatchStartTime = FDateTime::UtcNow();
@@ -250,6 +250,9 @@ private:
 				}
 			}
 		}
+
+		// Save off the mapping to unity file
+		ObjectFileToUnityObjectFile = MoveTemp(Manifest.ObjectFileToUnityObjectFile);
 		return true;
 	}
 
diff --git a/Engine/Source/Programs/LiveCodingConsole/Private/LiveCodingManifest.cpp b/Engine/Source/Programs/LiveCodingConsole/Private/LiveCodingManifest.cpp
--- a/Engine/Source/Programs/LiveCodingConsole/Private/LiveCodingManifest.cpp
+++ b/Engine/Source/Programs/LiveCodingConsole/Private/LiveCodingManifest.cpp
@@ -95,5 +95,37 @@ bool FLiveCodingManifest::Parse(FJsonObject& Object, FString& OutFailReason)
 			InputFiles.Add(Input->AsString());
 		}
 	}
+
+	const TArray<TSharedPtr<FJsonValue>>* AdaptiveUnityFilesArray;
+	if (Object.TryGetArrayField(TEXT("AdaptiveUnityFiles"), AdaptiveUnityFilesArray))
+	{
+		for (const TSharedPtr<FJsonValue>& AdaptiveUnityFileValue : *AdaptiveUnityFilesArray)
+		{
+			if (AdaptiveUnityFileValue->Type != EJson::Object)
+			{
+				OutFailReason = TEXT("invalid module object");
+				return false;
+			}
+
+			const FJsonObject& AdaptiveUnityFileObject = *AdaptiveUnityFileValue->AsObject();
+
+			FString ObjectFile;
+			if (!AdaptiveUnityFileObject.TryGetStringField(TEXT("ObjectFile"), ObjectFile))
+			{
+				OutFailReason = TEXT("missing module 'ObjectFile' field");
+				return false;
+			}
+
+			FString UnityObjectFile;
+			if (!AdaptiveUnityFileObject.TryGetStringField(TEXT("UnityObjectFile"), UnityObjectFile))
+			{
+				OutFailReason = TEXT("missing module 'UnityObjectFile' field");
+				return false;
+			}
+
+			ObjectFileToUnityObjectFile.Emplace(MoveTemp(ObjectFile), MoveTemp(UnityObjectFile));
+		}
+	}
+
 	return true;
 }
diff --git a/Engine/Source/Programs/LiveCodingConsole/Private/LiveCodingManifest.h b/Engine/Source/Programs/LiveCodingConsole/Private/LiveCodingManifest.h
--- a/Engine/Source/Programs/LiveCodingConsole/Private/LiveCodingManifest.h
+++ b/Engine/Source/Programs/LiveCodingConsole/Private/LiveCodingManifest.h
@@ -11,6 +11,7 @@ struct FLiveCodingManifest
 	FString LinkerPath;
 	TMap<FString, FString> LinkerEnvironment;
 	TMap<FString, TArray<FString>> BinaryToObjectFiles;
+	TMap<FString, FString> ObjectFileToUnityObjectFile;
 
 	bool Read(const TCHAR* FileName, FString& OutFailReason);
 	bool Parse(FJsonObject& Object, FString& OutFailReason);
diff --git a/Engine/Source/Programs/UnrealBuildTool/Modes/BuildMode.cs b/Engine/Source/Programs/UnrealBuildTool/Modes/BuildMode.cs
--- a/Engine/Source/Programs/UnrealBuildTool/Modes/BuildMode.cs
+++ b/Engine/Source/Programs/UnrealBuildTool/Modes/BuildMode.cs
@@ -221,9 +221,11 @@ namespace UnrealBuildTool
 			// Output the manifest
 			for(int TargetIdx = 0; TargetIdx < TargetDescriptors.Count; TargetIdx++)
 			{
-				if(TargetDescriptors[TargetIdx].LiveCodingManifest != null)
+				TargetDescriptor TargetDescriptor = TargetDescriptors[TargetIdx];
+				if(TargetDescriptor.LiveCodingManifest != null)
 				{
-					HotReload.WriteLiveCodeManifest(TargetDescriptors[TargetIdx].LiveCodingManifest, Makefiles[TargetIdx].Actions);
+					TargetMakefile Makefile = Makefiles[TargetIdx];
+					HotReload.WriteLiveCodeManifest(TargetDescriptor.LiveCodingManifest, Makefile.SourceFileToUnityFile, Makefile.Actions);
 				}
 			}
 
diff --git a/Engine/Source/Programs/UnrealBuildTool/System/HotReload.cs b/Engine/Source/Programs/UnrealBuildTool/System/HotReload.cs
--- a/Engine/Source/Programs/UnrealBuildTool/System/HotReload.cs
+++ b/Engine/Source/Programs/UnrealBuildTool/System/HotReload.cs
@@ -727,8 +727,9 @@ namespace UnrealBuildTool
 		/// Writes a manifest containing all the information needed to create a live coding patch
 		/// </summary>
 		/// <param name="ManifestFile">File to write to</param>
+		/// <param name="SourceFileToUnityFile">Mapping of source file to the unity file that would contain it</param>
 		/// <param name="Actions">List of actions that are part of the graph</param>
-		public static void WriteLiveCodeManifest(FileReference ManifestFile, List<Action> Actions)
+		public static void WriteLiveCodeManifest(FileReference ManifestFile, Dictionary<FileItem, FileItem> SourceFileToUnityFile, List<Action> Actions)
 		{
 			// Find all the output object files
 			HashSet<FileItem> ObjectFiles = new HashSet<FileItem>();
@@ -740,8 +741,32 @@ namespace UnrealBuildTool
 				}
 			}
 
+			// Find any action that has a non-unity cpp prerequisite, and map it to the appropriate 
+			Dictionary<FileItem, FileItem> ObjectFileToUnityObjectFile = new Dictionary<FileItem, FileItem>();
+			foreach(Action Action in Actions)
+			{
+				if (Action.ActionType == ActionType.Compile)
+				{
+					foreach(FileItem PrerequisiteItem in Action.PrerequisiteItems)
+					{
+						FileItem UnityFile;
+						if(SourceFileToUnityFile.TryGetValue(PrerequisiteItem, out UnityFile))
+						{
+							foreach(FileItem ProducedItem in Action.ProducedItems)
+							{
+								if(ProducedItem.HasExtension(".obj"))
+								{
+									FileItem UnityObjFile = FileItem.GetItemByFileReference(new FileReference(UnityFile.Location.FullName + ".obj"));
+									ObjectFileToUnityObjectFile[ProducedItem] = UnityObjFile;
+								}
+							}
+						}
+					}
+				}
+			}
+
 			// Write the output manifest
-			using(JsonWriter Writer = new JsonWriter(ManifestFile))
+			using (JsonWriter Writer = new JsonWriter(ManifestFile))
 			{
 				Writer.WriteObjectStart();
 
@@ -785,6 +810,26 @@ namespace UnrealBuildTool
 				}
 				Writer.WriteArrayEnd();
 
+				Writer.WriteArrayStart("AdaptiveUnityFiles");
+				foreach (Action Action in Actions)
+				{
+					if (Action.ActionType == ActionType.Link && Action.ProducedItems.Any(x => x.HasExtension(".exe") || x.HasExtension(".dll")))
+					{
+						foreach (FileItem InputFile in Action.PrerequisiteItems)
+						{
+							FileItem UnityFile;
+							if (ObjectFileToUnityObjectFile.TryGetValue(InputFile, out UnityFile))
+							{
+								Writer.WriteObjectStart();
+								Writer.WriteValue("ObjectFile", InputFile.AbsolutePath);
+								Writer.WriteValue("UnityObjectFile", UnityFile.AbsolutePath);
+								Writer.WriteObjectEnd();
+							}
+						}
+					}
+				}
+				Writer.WriteArrayEnd();
+
 				Writer.WriteObjectEnd();
 			}
 		}
diff --git a/Engine/Source/Programs/UnrealBuildTool/System/TargetMakefile.cs b/Engine/Source/Programs/UnrealBuildTool/System/TargetMakefile.cs
--- a/Engine/Source/Programs/UnrealBuildTool/System/TargetMakefile.cs
+++ b/Engine/Source/Programs/UnrealBuildTool/System/TargetMakefile.cs
@@ -20,7 +20,7 @@ namespace UnrealBuildTool
 		/// <summary>
 		/// The version number to write
 		/// </summary>
-		public const int CurrentVersion = 11;
+		public const int CurrentVersion = 12;
 
 		/// <summary>
 		/// The time at which the makefile was created
@@ -113,6 +113,11 @@ namespace UnrealBuildTool
 		/// </summary>
 		public HashSet<FileItem> CandidatesForWorkingSet = new HashSet<FileItem>();
 
+		/// <summary>
+		/// Mapping from source file to the unity file that contains it. Used for Live Coding. Note that this mapping includes files that are part of the working set, even if they are now a separate compiland.
+		/// </summary>
+		public Dictionary<FileItem, FileItem> SourceFileToUnityFile = new Dictionary<FileItem, FileItem>();
+
 		/// <summary>
 		/// Maps each target to a list of UObject module info structures
 		/// </summary>
@@ -190,6 +195,7 @@ namespace UnrealBuildTool
 			DirectoryToSourceFiles = Reader.ReadDictionary(() => Reader.ReadDirectoryItem(), () => Reader.ReadArray(() => Reader.ReadFileItem()));
 			WorkingSet = Reader.ReadHashSet(() => Reader.ReadFileItem());
 			CandidatesForWorkingSet = Reader.ReadHashSet(() => Reader.ReadFileItem());
+			SourceFileToUnityFile = Reader.ReadDictionary(() => Reader.ReadFileItem(), () => Reader.ReadFileItem());
 			UObjectModules = Reader.ReadList(() => new UHTModuleInfo(Reader));
 			UObjectModuleHeaders = Reader.ReadList(() => new UHTModuleHeaderInfo(Reader));
 			PluginFiles = Reader.ReadHashSet(() => Reader.ReadFileItem());
@@ -220,6 +226,7 @@ namespace UnrealBuildTool
 			Writer.WriteDictionary(DirectoryToSourceFiles, k => Writer.WriteDirectoryItem(k), v => Writer.WriteArray(v, e => Writer.WriteFileItem(e)));
 			Writer.WriteHashSet(WorkingSet, x => Writer.WriteFileItem(x));
 			Writer.WriteHashSet(CandidatesForWorkingSet, x => Writer.WriteFileItem(x));
+			Writer.WriteDictionary(SourceFileToUnityFile, k => Writer.WriteFileItem(k), v => Writer.WriteFileItem(v));
 			Writer.WriteList(UObjectModules, e => e.Write(Writer));
 			Writer.WriteList(UObjectModuleHeaders, x => x.Write(Writer));
 			Writer.WriteHashSet(PluginFiles, x => Writer.WriteFileItem(x));
diff --git a/Engine/Source/Programs/UnrealBuildTool/System/Unity.cs b/Engine/Source/Programs/UnrealBuildTool/System/Unity.cs
--- a/Engine/Source/Programs/UnrealBuildTool/System/Unity.cs
+++ b/Engine/Source/Programs/UnrealBuildTool/System/Unity.cs
@@ -28,6 +28,7 @@ namespace UnrealBuildTool
 		public class FileCollection
 		{
 			public List<FileItem> Files { get; private set; }
+			public List<FileItem> VirtualFiles { get; private set; }
 			public long TotalLength { get; private set; }
 
 			/// The length of this file collection, plus any additional virtual space needed for bUseAdapativeUnityBuild.
@@ -37,6 +38,7 @@ namespace UnrealBuildTool
 			public FileCollection()
 			{
 				Files = new List<FileItem>();
+				VirtualFiles = new List<FileItem>();
 				TotalLength = 0;
 				VirtualLength = 0;
 			}
@@ -56,10 +58,11 @@ namespace UnrealBuildTool
 			/// one of that module's unity blobs.  Basically, it can prevent dozens of files from being recompiled after the first
 			/// time building after your working set of source files changes
 			/// </summary>
-			/// <param name="VirtualFileLength">Length of the virtual file to add to this file collection</param>
-			public void AddVirtualFile(long VirtualFileLength)
+			/// <param name="File">The virtual file to add to the collection</param>
+			public void AddVirtualFile(FileItem File)
 			{
-				VirtualLength += VirtualFileLength;
+				VirtualFiles.Add(File);
+				VirtualLength += File.Length;
 			}
 		}
 
@@ -105,7 +108,7 @@ namespace UnrealBuildTool
 			/// <param name="File">The file to add virtually.  Only the size of the file is tracked.</param>
 			public void AddVirtualFile(FileItem File)
 			{
-				CurrentCollection.AddVirtualFile(File.Length);
+				CurrentCollection.AddVirtualFile(File);
 				if (SplitLength != -1 && CurrentCollection.VirtualLength > SplitLength)
 				{
 					EndCurrentUnityFile();
@@ -334,6 +337,16 @@ namespace UnrealBuildTool
 				// Write the unity file to the intermediate folder.
 				FileItem UnityCPPFile = FileItem.CreateIntermediateTextFile(UnityCPPFilePath, OutputUnityCPPWriter.ToString());
 				NewCPPFiles.Add(UnityCPPFile);
+
+				// Store the mapping of source files to unity files in the makefile
+				foreach(FileItem SourceFile in UnityFile.Files)
+				{
+					Makefile.SourceFileToUnityFile[SourceFile] = UnityCPPFile;
+				}
+				foreach (FileItem SourceFile in UnityFile.VirtualFiles)
+				{
+					Makefile.SourceFileToUnityFile[SourceFile] = UnityCPPFile;
+				}
 			}
 
 			return NewCPPFiles;
