Force inline CachePlatformAudioCookOverrides to avoid symbol duplication in monolithic builds.

diff --git a/Engine/Source/Developer/Windows/WindowsTargetPlatform/Private/GenericWindowsTargetPlatform.h b/Engine/Source/Developer/Windows/WindowsTargetPlatform/Private/GenericWindowsTargetPlatform.h
--- a/Engine/Source/Developer/Windows/WindowsTargetPlatform/Private/GenericWindowsTargetPlatform.h
+++ b/Engine/Source/Developer/Windows/WindowsTargetPlatform/Private/GenericWindowsTargetPlatform.h
@@ -22,7 +22,7 @@
 namespace Windows
 {
 #if WITH_ENGINE
-	void CachePlatformAudioCookOverrides(FPlatformAudioCookOverrides& OutOverrides)
+	FORCEINLINE void CachePlatformAudioCookOverrides(FPlatformAudioCookOverrides& OutOverrides)
 	{
 		const TCHAR* CategoryName = TEXT("/Script/WindowsTargetPlatform.WindowsTargetSettings");
 
