Live Coding: Fix failure to build patch when triggering Live Coding via the editor compile button with startup set to Manual. Was not registering the modules with the sidecar process before triggering the compile.

#rnx

diff --git a/Engine/Source/Developer/Windows/LiveCoding/Private/LiveCodingModule.cpp b/Engine/Source/Developer/Windows/LiveCoding/Private/LiveCodingModule.cpp
--- a/Engine/Source/Developer/Windows/LiveCoding/Private/LiveCodingModule.cpp
+++ b/Engine/Source/Developer/Windows/LiveCoding/Private/LiveCodingModule.cpp
@@ -217,6 +217,7 @@ void FLiveCodingModule::Compile()
 		EnableForSession(true);
 		if(bStarted)
 		{
+			UpdateModules(); // Need to do this immediately rather than waiting until next tick
 			LppTriggerRecompile();
 			GIsCompileActive = true;
 		}
