Live Coding: Fix patches not being reapplied during a quick restart for lazy loaded modules. Now forces such modules to be enabled immediately if they have pending patches.

#rnx

diff --git a/Engine/Source/Developer/Windows/LiveCoding/Private/External/LC_API.cpp b/Engine/Source/Developer/Windows/LiveCoding/Private/External/LC_API.cpp
--- a/Engine/Source/Developer/Windows/LiveCoding/Private/External/LC_API.cpp
+++ b/Engine/Source/Developer/Windows/LiveCoding/Private/External/LC_API.cpp
@@ -205,10 +205,10 @@ LPP_DLL_API(void) LppSetBuildArguments(const wchar_t* arguments)
 // END EPIC MOD
 
 // BEGIN EPIC MOD - Support for lazy-loading modules
-LPP_DLL_API(void) LppEnableLazyLoadedModule(const wchar_t* nameOfExeOrDll)
+LPP_DLL_API(void*) LppEnableLazyLoadedModule(const wchar_t* nameOfExeOrDll)
 {
 	HMODULE baseAddress = GetModuleHandle(nameOfExeOrDll);
-	g_startupThread->EnableLazyLoadedModule(nameOfExeOrDll, baseAddress);
+	return g_startupThread->EnableLazyLoadedModule(nameOfExeOrDll, baseAddress);
 }
 // END EPIC MOD
 
diff --git a/Engine/Source/Developer/Windows/LiveCoding/Private/External/LC_API.h b/Engine/Source/Developer/Windows/LiveCoding/Private/External/LC_API.h
--- a/Engine/Source/Developer/Windows/LiveCoding/Private/External/LC_API.h
+++ b/Engine/Source/Developer/Windows/LiveCoding/Private/External/LC_API.h
@@ -44,7 +44,7 @@ LPP_DLL_API(void) LppShowConsole();
 LPP_DLL_API(void) LppSetVisible(bool visible);
 LPP_DLL_API(void) LppSetActive(bool active);
 LPP_DLL_API(void) LppSetBuildArguments(const wchar_t* arguments);
-LPP_DLL_API(void) LppEnableLazyLoadedModule(const wchar_t* nameOfExeOrDll);
+LPP_DLL_API(void*) LppEnableLazyLoadedModule(const wchar_t* nameOfExeOrDll);
 // END EPIC MOD
 LPP_DLL_API(void) LppApplySettingBool(const char* settingName, int value);
 LPP_DLL_API(void) LppApplySettingInt(const char* settingName, int value);
diff --git a/Engine/Source/Developer/Windows/LiveCoding/Private/External/LC_ClientStartupThread.cpp b/Engine/Source/Developer/Windows/LiveCoding/Private/External/LC_ClientStartupThread.cpp
--- a/Engine/Source/Developer/Windows/LiveCoding/Private/External/LC_ClientStartupThread.cpp
+++ b/Engine/Source/Developer/Windows/LiveCoding/Private/External/LC_ClientStartupThread.cpp
@@ -331,15 +331,17 @@ void ClientStartupThread::SetBuildArguments(const wchar_t* arguments)
 // END EPIC MOD
 
 // BEGIN EPIC MOD - Support for lazy-loading modules
-void ClientStartupThread::EnableLazyLoadedModule(const wchar_t* fileName, Windows::HMODULE moduleBase)
+void* ClientStartupThread::EnableLazyLoadedModule(const wchar_t* fileName, Windows::HMODULE moduleBase)
 {
 	// we cannot wait for commands in the user command thread as long as startup hasn't finished
 	Join();
 
 	if (m_userCommandThread)
 	{
-		m_userCommandThread->EnableLazyLoadedModule(fileName, moduleBase);
+		return m_userCommandThread->EnableLazyLoadedModule(fileName, moduleBase);
 	}
+
+	return nullptr;
 }
 // END EPIC MOD
 
diff --git a/Engine/Source/Developer/Windows/LiveCoding/Private/External/LC_ClientStartupThread.h b/Engine/Source/Developer/Windows/LiveCoding/Private/External/LC_ClientStartupThread.h
--- a/Engine/Source/Developer/Windows/LiveCoding/Private/External/LC_ClientStartupThread.h
+++ b/Engine/Source/Developer/Windows/LiveCoding/Private/External/LC_ClientStartupThread.h
@@ -63,7 +63,7 @@ public:
 	// END EPIC MOD
 
 	// BEGIN EPIC MOD - Support for lazy-loading modules
-	void EnableLazyLoadedModule(const wchar_t* fileName, Windows::HMODULE moduleBase);
+	void* EnableLazyLoadedModule(const wchar_t* fileName, Windows::HMODULE moduleBase);
 	// END EPIC MOD
 
 private:
diff --git a/Engine/Source/Developer/Windows/LiveCoding/Private/External/LC_ClientUserCommandThread.cpp b/Engine/Source/Developer/Windows/LiveCoding/Private/External/LC_ClientUserCommandThread.cpp
--- a/Engine/Source/Developer/Windows/LiveCoding/Private/External/LC_ClientUserCommandThread.cpp
+++ b/Engine/Source/Developer/Windows/LiveCoding/Private/External/LC_ClientUserCommandThread.cpp
@@ -468,14 +468,17 @@ void ClientUserCommandThread::SetBuildArguments(const wchar_t* arguments)
 // END EPIC MOD
 
 // BEGIN EPIC MOD - Adding support for lazy-loading modules
-void ClientUserCommandThread::EnableLazyLoadedModule(const wchar_t* fileName, Windows::HMODULE moduleBase)
+void* ClientUserCommandThread::EnableLazyLoadedModule(const wchar_t* fileName, Windows::HMODULE moduleBase)
 {
-	ProxyCommand<commands::EnableLazyLoadedModule>* proxy = new ProxyCommand<commands::EnableLazyLoadedModule>(false, 0u);
+	ProxyCommand<commands::EnableLazyLoadedModule>* proxy = new ProxyCommand<commands::EnableLazyLoadedModule>(true, 0u);
 	proxy->m_command.processId = process::GetId();
 	wcscpy_s(proxy->m_command.fileName, fileName);
 	proxy->m_command.moduleBase = moduleBase;
+	proxy->m_command.token = new Event(nullptr, Event::Type::AUTO_RESET);
 
 	PushUserCommand(proxy);
+
+	return proxy->m_command.token;
 }
 // END EPIC MOD
 
diff --git a/Engine/Source/Developer/Windows/LiveCoding/Private/External/LC_ClientUserCommandThread.h b/Engine/Source/Developer/Windows/LiveCoding/Private/External/LC_ClientUserCommandThread.h
--- a/Engine/Source/Developer/Windows/LiveCoding/Private/External/LC_ClientUserCommandThread.h
+++ b/Engine/Source/Developer/Windows/LiveCoding/Private/External/LC_ClientUserCommandThread.h
@@ -80,7 +80,7 @@ public:
 	// END EPIC MOD
 
 	// BEGIN EPIC MOD - Support for lazy-loading modules
-	void EnableLazyLoadedModule(const wchar_t* fileName, Windows::HMODULE moduleBase);
+	void* EnableLazyLoadedModule(const wchar_t* fileName, Windows::HMODULE moduleBase);
 	// END EPIC MOD
 
 	void ApplySettingBool(const char* settingName, int value);
diff --git a/Engine/Source/Developer/Windows/LiveCoding/Private/External/LC_Commands.h b/Engine/Source/Developer/Windows/LiveCoding/Private/External/LC_Commands.h
--- a/Engine/Source/Developer/Windows/LiveCoding/Private/External/LC_Commands.h
+++ b/Engine/Source/Developer/Windows/LiveCoding/Private/External/LC_Commands.h
@@ -271,6 +271,7 @@ namespace commands
 		unsigned int processId;
 		wchar_t fileName[260];
 		Windows::HMODULE moduleBase;
+		void* token;
 	};
 
 	struct FinishedLazyLoadingModules
diff --git a/Engine/Source/Developer/Windows/LiveCoding/Private/LiveCodingModule.cpp b/Engine/Source/Developer/Windows/LiveCoding/Private/LiveCodingModule.cpp
--- a/Engine/Source/Developer/Windows/LiveCoding/Private/LiveCodingModule.cpp
+++ b/Engine/Source/Developer/Windows/LiveCoding/Private/LiveCodingModule.cpp
@@ -362,7 +362,7 @@ void FLiveCodingModule::UpdateModules()
 					}
 					else
 					{
-						LppEnableLazyLoadedModule(*FullFilePath);
+						LppWaitForToken(LppEnableLazyLoadedModule(*FullFilePath));
 					}
 					ConfiguredModules.Add(ModuleName);
 				}
@@ -376,7 +376,7 @@ void FLiveCodingModule::UpdateModules()
 			{
 				EnableModuleFileNames.Add(*EnableModule);
 			}
-			LppEnableModules(EnableModuleFileNames.GetData(), EnableModuleFileNames.Num());
+			LppWaitForToken(LppEnableModules(EnableModuleFileNames.GetData(), EnableModuleFileNames.Num()));
 		}
 #endif
 	}
diff --git a/Engine/Source/Developer/Windows/LiveCodingServer/Private/External/LC_ServerCommandThread.cpp b/Engine/Source/Developer/Windows/LiveCodingServer/Private/External/LC_ServerCommandThread.cpp
--- a/Engine/Source/Developer/Windows/LiveCodingServer/Private/External/LC_ServerCommandThread.cpp
+++ b/Engine/Source/Developer/Windows/LiveCodingServer/Private/External/LC_ServerCommandThread.cpp
@@ -1477,18 +1477,40 @@ bool ServerCommandThread::actions::EnableLazyLoadedModule::Execute(const Command
 	// protect against accepting this command while compilation is already in progress
 	CriticalSection::ScopedLock lock(&commandThread->m_actionCS);
 
+	// Check if this module is already enabled - it may have been lazy-loaded, then fully loaded, by a restarted process. If so, translate this into a call to EnableModules.
+	const std::wstring modulePath = file::NormalizePath(command->fileName);
+	for (LiveModule* module : commandThread->m_liveModules)
+	{
+		if(module->GetModuleName() == modulePath)
+		{
+			EnableModules::CommandType EnableCmd = { };
+			EnableCmd.moduleCount = 1;
+			EnableCmd.processId = command->processId;
+			EnableCmd.token = command->token;
+
+			commands::ModuleData Module;
+			Module.base = command->moduleBase;
+			wcscpy_s(Module.path, command->fileName);
+
+			return EnableModules::Execute(&EnableCmd, pipe, context, &Module, sizeof(Module));
+		}
+	}
+
+	// Acknowledge the command
+	pipe->SendAck();
+
+	// Register the module for lazy loading
 	for (LiveProcess* process : commandThread->m_liveProcesses)
 	{
 		if (process->GetProcessId() == command->processId)
 		{
-			const std::wstring modulePath = file::NormalizePath(command->fileName);
 			process->AddLazyLoadedModule(modulePath, command->moduleBase);
 			LC_LOG_DEV("Registered module %S for lazy-loading", modulePath.c_str());
 		}
 	}
 
-	pipe->SendAck();
-
+	// Tell the client we're done
+	pipe->SendCommandAndWaitForAck(commands::EnableModulesFinished { command->token }, nullptr, 0u);
 	return true;
 }
 // END EPIC MOD
