Live Coding: Use FPlatformMisc::RequestExit(true) when quick restarting an application, to ensure that logs are flushed and the session is not considered an abnormal termination.

#rnx

diff --git a/Engine/Source/Developer/Windows/LiveCoding/Private/External/LC_Restart.cpp b/Engine/Source/Developer/Windows/LiveCoding/Private/External/LC_Restart.cpp
--- a/Engine/Source/Developer/Windows/LiveCoding/Private/External/LC_Restart.cpp
+++ b/Engine/Source/Developer/Windows/LiveCoding/Private/External/LC_Restart.cpp
@@ -5,6 +5,7 @@
 #include "LC_PrimitiveNames.h"
 #include "LC_Process.h"
 #include "LC_Memory.h"
+#include "HAL/PlatformMisc.h"
 
 
 namespace
@@ -55,6 +56,12 @@ void restart::Execute(lpp::RestartBehaviour behaviour, unsigned int exitCode)
 
 	switch (behaviour)
 	{
+		// BEGIN EPIC MODS - Use UE4 codepath for termination to ensure logs are flushed and session analytics are sent
+		case lpp::LPP_RESTART_BEHAVIOR_REQUEST_EXIT:
+			FPlatformMisc::RequestExit(true);
+			break;
+		// END EPIC MODS
+
 		// https://docs.microsoft.com/en-us/windows/win32/api/processthreadsapi/nf-processthreadsapi-exitprocess
 		case lpp::LPP_RESTART_BEHAVIOUR_DEFAULT_EXIT:
 			ExitProcess(exitCode);
diff --git a/Engine/Source/Developer/Windows/LiveCoding/Private/External/LPP_API.h b/Engine/Source/Developer/Windows/LiveCoding/Private/External/LPP_API.h
--- a/Engine/Source/Developer/Windows/LiveCoding/Private/External/LPP_API.h
+++ b/Engine/Source/Developer/Windows/LiveCoding/Private/External/LPP_API.h
@@ -80,6 +80,9 @@ LPP_NS_BEGIN
 
 enum RestartBehaviour
 {
+	// BEGIN EPIC MODS - Use UE4 codepath for termination to ensure logs are flushed and session analytics are sent
+	LPP_RESTART_BEHAVIOR_REQUEST_EXIT,				// FPlatforMisc::RequestExit(true)
+	// END EPIC MODS
 	LPP_RESTART_BEHAVIOUR_DEFAULT_EXIT,				// ExitProcess()
 	LPP_RESTART_BEHAVIOUR_EXIT_WITH_FLUSH,			// exit()
 	LPP_RESTART_BEHAVIOUR_EXIT,						// _Exit()
diff --git a/Engine/Source/Developer/Windows/LiveCoding/Private/LiveCodingModule.cpp b/Engine/Source/Developer/Windows/LiveCoding/Private/LiveCodingModule.cpp
--- a/Engine/Source/Developer/Windows/LiveCoding/Private/LiveCodingModule.cpp
+++ b/Engine/Source/Developer/Windows/LiveCoding/Private/LiveCodingModule.cpp
@@ -233,7 +233,7 @@ void FLiveCodingModule::Tick()
 {
 	if (LppWantsRestart())
 	{
-		LppRestart(lpp::LPP_RESTART_BEHAVIOUR_INSTANT_TERMINATION, 0);
+		LppRestart(lpp::LPP_RESTART_BEHAVIOR_REQUEST_EXIT, 0);
 	}
 
 	if (Settings->bEnabled != bEnabledLastTick && Settings->Startup != ELiveCodingStartupMode::Manual)
