Live Coding: Prevent orphaned console instances if a process never reattaches (due to manual startup, etc...). Job object handle is now duplicated into the child process at startup, and closed immediately.

#rnx

diff --git a/Engine/Source/Developer/Windows/LiveCodingServer/Private/External/LC_LiveProcess.cpp b/Engine/Source/Developer/Windows/LiveCodingServer/Private/External/LC_LiveProcess.cpp
--- a/Engine/Source/Developer/Windows/LiveCodingServer/Private/External/LC_LiveProcess.cpp
+++ b/Engine/Source/Developer/Windows/LiveCodingServer/Private/External/LC_LiveProcess.cpp
@@ -176,7 +176,7 @@ bool LiveProcess::PrepareForRestart(void)
 }
 
 
-void LiveProcess::Restart(void)
+void LiveProcess::Restart(void* restartJob)
 {
 	if (m_restartState == RestartState::SUCCESSFUL_PREPARE)
 	{
@@ -189,7 +189,12 @@ void LiveProcess::Restart(void)
 		process::Wait(m_processHandle);
 
 		// restart the target application
-		process::Spawn(m_imagePath.c_str(), m_workingDirectory.c_str(), m_commandLine.c_str(), m_environment.GetData(), process::SpawnFlags::NONE);
+		// BEGIN EPIC MOD - Prevent orphaned console instances if processes fail to restart. Job object will be duplicated into child process.
+		process::Context* context = process::Spawn(m_imagePath.c_str(), m_workingDirectory.c_str(), m_commandLine.c_str(), m_environment.GetData(), process::SpawnFlags::SUSPENDED);
+		void* TargetHandle;
+		DuplicateHandle(GetCurrentProcess(), restartJob, context->pi.hProcess, &TargetHandle, 0, Windows::TRUE, DUPLICATE_SAME_ACCESS);
+		ResumeThread(context->pi.hThread);
+		// END EPIC MOD
 
 		m_restartState = RestartState::SUCCESSFUL;
 	}
diff --git a/Engine/Source/Developer/Windows/LiveCodingServer/Private/External/LC_LiveProcess.h b/Engine/Source/Developer/Windows/LiveCodingServer/Private/External/LC_LiveProcess.h
--- a/Engine/Source/Developer/Windows/LiveCodingServer/Private/External/LC_LiveProcess.h
+++ b/Engine/Source/Developer/Windows/LiveCodingServer/Private/External/LC_LiveProcess.h
@@ -39,7 +39,7 @@ public:
 
 
 	bool PrepareForRestart(void);
-	void Restart(void);
+	void Restart(void* restartJob);
 	bool WasSuccessfulRestart(void) const;
 
 
diff --git a/Engine/Source/Developer/Windows/LiveCodingServer/Private/External/LC_ServerCommandThread.cpp b/Engine/Source/Developer/Windows/LiveCodingServer/Private/External/LC_ServerCommandThread.cpp
--- a/Engine/Source/Developer/Windows/LiveCodingServer/Private/External/LC_ServerCommandThread.cpp
+++ b/Engine/Source/Developer/Windows/LiveCodingServer/Private/External/LC_ServerCommandThread.cpp
@@ -206,8 +206,16 @@ void ServerCommandThread::RestartTargets(void)
 	for (size_t i = 0u; i < count; ++i)
 	{
 		LiveProcess* liveProcess = m_liveProcesses[i];
-		liveProcess->Restart();
+		liveProcess->Restart(m_restartJob);
 	}
+
+	// BEGIN EPIC MOD - Prevent orphaned console instances if processes fail to restart. Job object will be duplicated into child process.
+	if (m_restartJob != nullptr)
+	{
+		CloseHandle(m_restartJob);
+		m_restartJob = nullptr;
+	}
+	// END EPIC MOD
 }
 
 
