Live Coding: Restore patched module state if a packaged game is quick restarted.

#rnx

diff --git a/Engine/Source/Developer/Windows/LiveCodingServer/Private/External/LC_LiveProcess.cpp b/Engine/Source/Developer/Windows/LiveCodingServer/Private/External/LC_LiveProcess.cpp
--- a/Engine/Source/Developer/Windows/LiveCodingServer/Private/External/LC_LiveProcess.cpp
+++ b/Engine/Source/Developer/Windows/LiveCodingServer/Private/External/LC_LiveProcess.cpp
@@ -189,8 +189,17 @@ void LiveProcess::Restart(void* restartJob)
 		process::Wait(m_processHandle);
 
 		// restart the target application
+		// BEGIN EPIC MOD - Force LiveCoding to start up for child processes
+		std::wstring commandLine(m_commandLine);
+
+		const std::wstring argument(L" -LiveCoding");
+		if (commandLine.length() >= argument.length() && commandLine.compare(commandLine.length() - argument.length(), argument.length(), argument) != 0)
+		{
+			commandLine += argument;
+		}
+		// END EPIC MOD
 		// BEGIN EPIC MOD - Prevent orphaned console instances if processes fail to restart. Job object will be duplicated into child process.
-		process::Context* context = process::Spawn(m_imagePath.c_str(), m_workingDirectory.c_str(), m_commandLine.c_str(), m_environment.GetData(), process::SpawnFlags::SUSPENDED);
+		process::Context* context = process::Spawn(m_imagePath.c_str(), m_workingDirectory.c_str(), commandLine.c_str(), m_environment.GetData(), process::SpawnFlags::SUSPENDED);
 		void* TargetHandle;
 		DuplicateHandle(GetCurrentProcess(), restartJob, context->pi.hProcess, &TargetHandle, 0, Windows::TRUE, DUPLICATE_SAME_ACCESS);
 		ResumeThread(context->pi.hThread);
diff --git a/Engine/Source/Developer/Windows/LiveCodingServer/Private/External/LC_ServerCommandThread.cpp b/Engine/Source/Developer/Windows/LiveCodingServer/Private/External/LC_ServerCommandThread.cpp
--- a/Engine/Source/Developer/Windows/LiveCodingServer/Private/External/LC_ServerCommandThread.cpp
+++ b/Engine/Source/Developer/Windows/LiveCodingServer/Private/External/LC_ServerCommandThread.cpp
@@ -1622,6 +1622,13 @@ bool ServerCommandThread::actions::RegisterProcess::Execute(const CommandType* c
 						// EPIC REMOVED: g_theApp.GetMainFrame()->ResetStatusBarText();
 						// EPIC REMOVED: g_theApp.GetMainFrame()->SetBusy(false);
 					}
+					// BEGIN EPIC MOD - Prevent orphaned console instances if processes fail to restart. Job object will be duplicated into child process.
+					else
+					{
+						commandThread->m_restartCS.Leave();
+						LC_LOG_USER("---------- Restarting finished ----------");
+					}
+					// END EPIC MOD
 				}
 			}
 		}
