Removing uproperty for windows target settings for audio device selection

diff --git a/Engine/Source/Developer/Windows/WindowsTargetPlatform/Classes/WindowsTargetSettings.h b/Engine/Source/Developer/Windows/WindowsTargetPlatform/Classes/WindowsTargetSettings.h
--- a/Engine/Source/Developer/Windows/WindowsTargetPlatform/Classes/WindowsTargetSettings.h
+++ b/Engine/Source/Developer/Windows/WindowsTargetPlatform/Classes/WindowsTargetSettings.h
@@ -73,10 +73,6 @@ public:
 	UPROPERTY(EditAnywhere, config, Category="OS Info", Meta=(DisplayName = "Minimum OS Version"))
 	EMinimumSupportedOS MinimumOSVersion;
 
-	/** The audio device name to use if not the default windows audio device. Leave blank to use default audio device. */
-	UPROPERTY(config, EditAnywhere, Category = "Audio")
-	FString AudioDevice;
-
 	/** Sample rate to run the audio mixer with. */
 	UPROPERTY(config, EditAnywhere, Category = "Audio", Meta = (DisplayName = "Audio Mixer Sample Rate"))
 	int32 AudioSampleRate;
