Fix whitespace in LC_LiveModule.cpp.

diff --git a/Engine/Source/Developer/Windows/LiveCodingServer/Private/External/LC_LiveModule.cpp b/Engine/Source/Developer/Windows/LiveCodingServer/Private/External/LC_LiveModule.cpp
--- a/Engine/Source/Developer/Windows/LiveCodingServer/Private/External/LC_LiveModule.cpp
+++ b/Engine/Source/Developer/Windows/LiveCodingServer/Private/External/LC_LiveModule.cpp
@@ -2534,33 +2534,33 @@ LiveModule::ErrorType::Enum LiveModule::Update(FileAttributeCache* fileCache, Di
 						const symbols::Compiland* otherCompiland = symbolIt->second.second;
 						if (otherCompiland)
 						{
-						if (symbols::IsCompilandRecompiled(otherCompiland))
-						{
-							// the external symbol comes from one of the *other* recompiled .obj.
-							// in this case, the symbol might have changed, so we are only allowed to strip it
-							// if all relocations to it would be patched anyway.
-							tryStrip = true;
-							LC_LOG_DEV("Symbol comes from recompiled compiland");
-						}
-						else
-						{
-							// the external symbol comes from an .obj that was not recompiled.
-							// in this case, the symbol couldn't have changed, so we strip it directly
-							// in case it exists in our live module already.
-							const ModuleCache::FindSymbolData findData = m_moduleCache->FindSymbolByName(ModuleCache::SEARCH_ALL_MODULES, symbolName);
-							if (findData.symbol)
+							if (symbols::IsCompilandRecompiled(otherCompiland))
 							{
-								doStrip = true;
-								forceStrippedSymbols.insert(symbolName);
+								// the external symbol comes from one of the *other* recompiled .obj.
+								// in this case, the symbol might have changed, so we are only allowed to strip it
+								// if all relocations to it would be patched anyway.
+								tryStrip = true;
+								LC_LOG_DEV("Symbol comes from recompiled compiland");
 							}
 							else
 							{
-								LC_LOG_DEV("Symbol seems to be new (compiland)");
+								// the external symbol comes from an .obj that was not recompiled.
+								// in this case, the symbol couldn't have changed, so we strip it directly
+								// in case it exists in our live module already.
+								const ModuleCache::FindSymbolData findData = m_moduleCache->FindSymbolByName(ModuleCache::SEARCH_ALL_MODULES, symbolName);
+								if (findData.symbol)
+								{
+									doStrip = true;
+									forceStrippedSymbols.insert(symbolName);
+								}
+								else
+								{
+									LC_LOG_DEV("Symbol seems to be new (compiland)");
+								}
 							}
 						}
-					}
-					else
-					{
+						else
+						{
 							// the symbol must come from a new .obj, so we aren't allowed to strip it
 							LC_LOG_DEV("Symbol comes from new compiland");
 						}
@@ -3118,8 +3118,8 @@ LiveModule::ErrorType::Enum LiveModule::Update(FileAttributeCache* fileCache, Di
 		symbols::Compiland* compiland = compilandIt->second;
 		if (compiland)
 		{
-		symbols::ClearCompilandAsRecompiled(compiland);
-	}
+			symbols::ClearCompilandAsRecompiled(compiland);
+		}
 	}
 	++m_patchCounter;
 
