Fix issues related to unnecessary symbol reconstruction in Live Coding modules. Generating separate object files for updated compilands was resulting in symbols being assigned a unique compiland id.

Object files generated for Live Coding compiles are now temporarily renamed to their original versions until the Live Coding patch is generated, after which point they are restored. Would be preferable to incorporate this logic in Live++ directly and just read from the correct files, but the number of changes necessary to support it would make it difficult to integrate new versions.

diff --git a/Engine/Source/Developer/Windows/LiveCodingServer/Private/External/LC_ServerCommandThread.cpp b/Engine/Source/Developer/Windows/LiveCodingServer/Private/External/LC_ServerCommandThread.cpp
--- a/Engine/Source/Developer/Windows/LiveCodingServer/Private/External/LC_ServerCommandThread.cpp
+++ b/Engine/Source/Developer/Windows/LiveCodingServer/Private/External/LC_ServerCommandThread.cpp
@@ -781,8 +781,37 @@ void ServerCommandThread::CompileChanges(bool didAllProcessesMakeProgress)
 			types::vector<LiveModule::ModifiedObjFile> ObjectFiles;
 			for(const FString& ObjectFile : Pair.Value)
 			{
+				std::wstring NormalizedObjectFile = file::NormalizePath(*ObjectFile);
+
+				// If this file has a .lc.obj suffix, temporarily replace the original .obj file while generating the patch.
+				// It'd be nice to track this explicitly inside Live++ and just load the new file, but it requires a lot of changes and would make upgrades difficult.
+				static const TCHAR Suffix[] = TEXT(".lc.obj");
+				static const size_t SuffixLen = ARRAY_COUNT(Suffix) - 1;
+				if (NormalizedObjectFile.length() >= SuffixLen && _wcsicmp(NormalizedObjectFile.c_str() + NormalizedObjectFile.length() - SuffixLen, Suffix) == 0)
+				{
+					// Get the original filename
+					std::wstring OriginalObjectFile(NormalizedObjectFile.c_str(), NormalizedObjectFile.c_str() + NormalizedObjectFile.length() - SuffixLen);
+					OriginalObjectFile += L".obj";
+
+					// Back up the original file, if it exists
+					file::Attributes OriginalFileAttributes = file::GetAttributes(OriginalObjectFile.c_str());
+					if (file::DoesExist(OriginalFileAttributes))
+					{
+						std::wstring OriginalObjectFileBackup = OriginalObjectFile + L".lctmp";
+						m_restoreFiles.push_back(std::make_pair(OriginalObjectFileBackup, OriginalObjectFile));
+						file::DeleteIfExists(OriginalObjectFileBackup.c_str());
+						file::Move(OriginalObjectFile.c_str(), OriginalObjectFileBackup.c_str());
+					}
+
+					// Move the new file into place
+					m_restoreFiles.push_back(std::make_pair(OriginalObjectFile, NormalizedObjectFile));
+					file::Move(NormalizedObjectFile.c_str(), OriginalObjectFile.c_str());
+					NormalizedObjectFile = OriginalObjectFile;
+				}
+
+				// Add the file to the list of modifications
 				LiveModule::ModifiedObjFile ModifiedObjFile;
-				ModifiedObjFile.objPath = file::NormalizePath(*ObjectFile);
+				ModifiedObjFile.objPath = NormalizedObjectFile;
 				ObjectFiles.push_back(std::move(ModifiedObjFile));
 			}
 
@@ -1117,6 +1146,15 @@ unsigned int ServerCommandThread::CompileThread(void)
 
 			CompileChanges(didAllProcessesMakeProgress);
 
+			// BEGIN EPIC MOD - Non-destructive compile
+			for (std::vector<std::pair<std::wstring, std::wstring>>::reverse_iterator it = m_restoreFiles.rbegin(); it != m_restoreFiles.rend(); it++)
+			{
+				file::DeleteIfExists(it->second.c_str());
+				file::Move(it->first.c_str(), it->second.c_str());
+			}
+			m_restoreFiles.clear();
+			// END EPIC MOD
+
 			RemoveVirtualDrive();
 
 			if (!didAllProcessesMakeProgress)
diff --git a/Engine/Source/Developer/Windows/LiveCodingServer/Private/External/LC_ServerCommandThread.h b/Engine/Source/Developer/Windows/LiveCodingServer/Private/External/LC_ServerCommandThread.h
--- a/Engine/Source/Developer/Windows/LiveCodingServer/Private/External/LC_ServerCommandThread.h
+++ b/Engine/Source/Developer/Windows/LiveCodingServer/Private/External/LC_ServerCommandThread.h
@@ -134,6 +134,10 @@ private:
 	CriticalSection m_connectionCS;
 	types::vector<CommandThreadContext*> m_commandThreads;
 
+	// BEGIN EPIC MOD - Non-destructive compile
+	std::vector<std::pair<std::wstring, std::wstring>> m_restoreFiles;
+	// END EPIC MOD
+
 	// BEGIN EPIC MOD - Adding SetActive command
 	bool m_active = true;
 	// END EPIC MOD
