LiveCoding: More fixes for object reconstruction when reapplying Live Coding patches. Compiland IDs are now assigned from a central location, which can be queried at any time. The unity file mapping output from UBT is read on demand and used to populate the cache.

Also fix an issue where symbols in anonymous namespace would not use the correct name if an external build system is enabled.

#rnx

diff --git a/Engine/Source/Developer/Windows/LiveCodingServer/Private/External/LC_Coff.cpp b/Engine/Source/Developer/Windows/LiveCodingServer/Private/External/LC_Coff.cpp
--- a/Engine/Source/Developer/Windows/LiveCodingServer/Private/External/LC_Coff.cpp
+++ b/Engine/Source/Developer/Windows/LiveCodingServer/Private/External/LC_Coff.cpp
@@ -1721,7 +1721,10 @@ namespace coff
 				};
 
 				ObjFile objFile = { objPath, &memoryFile };
-				return GatherDatabase(&objFile, uniqueId::Generate(string::ToWideString(objPath)), ReadFlags::NONE);
+				// BEGIN EPIC MOD - 
+				uint32_t uniqueId = symbols::GetCompilandIdFromPath(string::ToWideString(objPath));
+				return GatherDatabase(&objFile, uniqueId, ReadFlags::NONE);
+				// END EPIC MOD
 			}
 		}
 
diff --git a/Engine/Source/Developer/Windows/LiveCodingServer/Private/External/LC_LiveModule.cpp b/Engine/Source/Developer/Windows/LiveCodingServer/Private/External/LC_LiveModule.cpp
--- a/Engine/Source/Developer/Windows/LiveCodingServer/Private/External/LC_LiveModule.cpp
+++ b/Engine/Source/Developer/Windows/LiveCodingServer/Private/External/LC_LiveModule.cpp
@@ -39,13 +39,6 @@
 #include "Misc/Paths.h"
 // END EPIC MOD
 
-// BEGIN EPIC MOD - Allow mapping from object files to their unity object file
-#include "Misc/FileHelper.h"
-#include "Serialization/JsonReader.h"
-#include "Serialization/JsonSerializer.h"
-#include "Dom/JsonObject.h"
-// END EPIC MOD
-
 namespace
 {
 	// common linker options:
@@ -396,123 +389,11 @@ namespace
 		return LiveModule::CompileResult { exitCode, true };
 	}
 
-	// BEGIN EPIC MOD - Allow mapping from object files to their unity object file
-	static void AddCompilandId(types::StringMap<uint32_t>& objFileToCompilandId, const std::wstring& objFile, uint32_t compilandId)
-	{
-		std::wstring WideObjectFile = file::NormalizePath(objFile.c_str());
-		ImmutableString ObjectFile = string::ToUtf8String(WideObjectFile);
-		if (objFileToCompilandId.find(ObjectFile) == objFileToCompilandId.end())
-		{
-			objFileToCompilandId.insert(std::make_pair(ObjectFile, compilandId));
-		}
-	}
-
-	static bool ReadLiveCodingInfo(const wchar_t* manifestFile, types::StringMap<uint32_t>& objFileToCompilandId)
-	{
-		// Read the file to a string
-		FString FileContents;
-		if (!FFileHelper::LoadFileToString(FileContents, manifestFile))
-		{
-			return false;
-		}
-
-		// Deserialize a JSON object from the string
-		TSharedPtr< FJsonObject > Object;
-		TSharedRef< TJsonReader<> > Reader = TJsonReaderFactory<>::Create(FileContents);
-		if (!FJsonSerializer::Deserialize(Reader, Object) || !Object.IsValid())
-		{
-			return false;
-		}
-
-		const TSharedPtr<FJsonObject>* FilesObject;
-		if (!Object->TryGetObjectField(TEXT("RemapUnityFiles"), FilesObject))
-		{
-			return false;
-		}
-
-		std::wstring BaseDir = file::GetDirectory(manifestFile);
-		for (const TPair<FString, TSharedPtr<FJsonValue>>& Pair : FilesObject->Get()->Values)
-		{
-			std::wstring UnityObjectFile = file::NormalizePath((BaseDir + L"\\" + *Pair.Key).c_str());
-			uint32_t UnityCompilandId = uniqueId::Generate(UnityObjectFile);
-			objFileToCompilandId.insert(std::make_pair(string::ToUtf8String(UnityObjectFile), UnityCompilandId));
-
-			const FJsonValue* Value = Pair.Value.Get();
-			if (Value->Type != EJson::Array)
-			{
-				return false;
-			}
-
-			const TArray<TSharedPtr<FJsonValue>>& SourceFileValues = Value->AsArray();
-			for (const TSharedPtr<FJsonValue>& SourceFileValue : SourceFileValues)
-			{
-				if (SourceFileValue->Type != EJson::String)
-				{
-					return false;
-				}
-
-				std::wstring objFile = BaseDir + L"\\" + *SourceFileValue->AsString();
-				AddCompilandId(objFileToCompilandId, objFile, UnityCompilandId);
-
-				std::wstring patchedObjFile = file::RemoveExtension(objFile) + L".lc.obj";
-				AddCompilandId(objFileToCompilandId, patchedObjFile, UnityCompilandId);
-			}
-		}
-
-		return true;
-	}
-
-	static void UpdateCompilandCache(types::StringMap<symbols::Compiland*>& compilands, types::StringMap<uint32_t>& objFileToCompilandId)
-	{
-		types::unordered_set<std::wstring> Directories;
-		for (std::pair<const ImmutableString, symbols::Compiland*>& Pair : compilands)
-		{
-			symbols::Compiland* Compiland = Pair.second;
-			if (Compiland != nullptr && Compiland->amalgamatedUniqueId == ~(uint32_t)0)
-			{
-				const std::wstring& wideObjPath = string::ToWideString(Pair.first);
-				Directories.insert(file::GetDirectory(wideObjPath));
-			}
-		}
-		for (const std::wstring& Directory : Directories)
-		{
-			std::wstring ManifestFile = Directory + L"\\LiveCodingInfo.json";
-			if (file::DoesExist(file::GetAttributes(ManifestFile.c_str())))
-			{
-				ReadLiveCodingInfo(ManifestFile.c_str(), objFileToCompilandId);
-			}
-		}
-		for (std::pair<const ImmutableString, symbols::Compiland*>& Pair : compilands)
-		{
-			symbols::Compiland* Compiland = Pair.second;
-			if (Compiland != nullptr && Compiland->amalgamatedUniqueId == ~(uint32_t)0)
-			{
-				types::StringMap<uint32_t>::const_iterator Iter = objFileToCompilandId.find(Pair.first);
-				if (Iter == objFileToCompilandId.end())
-				{
-					LC_WARNING_DEV("Unable to get amalgamated id for %s", Pair.first.c_str());
-				}
-				else
-				{
-					Compiland->amalgamatedUniqueId = Iter->second;
-				}
-			}
-		}
-	}
-	// END EPIC MOD
-
 	// helper function that returns or generates the unique ID of an optional compiland.
 	// for files split off from amalgamated files, we need to use the original object path of the amalgamated file here,
 	// otherwise names of symbols would differ, leading to constructors of global instances being called again.
 	static inline uint32_t GetCompilandId(const symbols::Compiland* compiland, const wchar_t* const objPath, const types::vector<LiveModule::ModifiedObjFile>& modifiedObjFiles)
 	{
-		// BEGIN EPIC MOD - Allow mapping from object files to their unity object file
-		if (compiland && compiland->amalgamatedUniqueId != ~(uint32_t)0)
-		{
-			return compiland->amalgamatedUniqueId;
-		}
-		// END EPIC MOD
-
 		// try to find the given .obj path in the array of modified object files to check if there's an original amalgamated object path for it
 		for (size_t i = 0u; i < modifiedObjFiles.size(); ++i)
 		{
@@ -535,7 +416,9 @@ namespace
 		}
 		else
 		{
-			return uniqueId::Generate(file::NormalizePath(objPath));
+			// BEGIN EPIC MOD - Static map of object file to unity blob
+			return symbols::GetCompilandIdFromPath(objPath);
+			// BEGIN EPIC MOD - Static map of object file to unity blob
 		}
 	}
 
@@ -1184,10 +1067,6 @@ void LiveModule::Load(symbols::Provider* provider, symbols::DiaCompilandDB* diaC
 	{
 		LC_LOG_DEV("Caching all .objs on Load() due to external build system being used");
 
-		// BEGIN EPIC MOD - Allow mapping from object files to their unity object file
-		UpdateCompilandCache(m_compilandDB->compilands, m_objFileToCompilandId);
-		// END EPIC MOD
-
 		// the user wants to use an external build system. in this case, we only track .objs for changes and never
 		// compile anything ourselves. we cannot load .objs lazily in this case, so we have to do that right now.
 		struct GatherResult
@@ -1217,20 +1096,13 @@ void LiveModule::Load(symbols::Provider* provider, symbols::DiaCompilandDB* diaC
 
 			// do the loading and gathering concurrently
 			auto task = scheduler::CreateTask(gatherTaskRoot, [objPath, compiland]()
-				{
+			{
 				const std::wstring& wideObjPath = string::ToWideString(objPath);
 				coff::ObjFile* objFile = coff::OpenObj(wideObjPath.c_str());
-				// BEGIN EPIC MOD - Allow mapping from object files to their unity object file
-				uint32_t uniqueId;
-			    if (compiland->amalgamatedUniqueId != ~(uint32_t)0)
-			    {
-				    uniqueId = compiland->amalgamatedUniqueId;
-			    }
-			    else
-			    {
-				    uniqueId = compiland->uniqueId;
-			    }
-				coff::CoffDB* database = coff::GatherDatabase(objFile, uniqueId, coff::ReadFlags::NONE);
+				// BEGIN EPIC MOD - Support for unity files s with external build system
+				const bool splitAmalgamatedFiles = appSettings::g_amalgamationSplitIntoSingleParts->GetValue();
+				const coff::ReadFlags::Enum coffReadFlags = splitAmalgamatedFiles ? coff::ReadFlags::GENERATE_ANS_NAME_FROM_UNIQUE_ID : coff::ReadFlags::NONE;
+				coff::CoffDB* database = coff::GatherDatabase(objFile, compiland->uniqueId, coffReadFlags);
 				// END EPIC MOD
 				coff::CloseObj(objFile);
 
@@ -2297,10 +2169,6 @@ LiveModule::ErrorType::Enum LiveModule::Update(FileAttributeCache* fileCache, Di
 		}
 	}
 
-	// BEGIN EPIC MOD - Allow mapping from object files to their unity object file
-	UpdateCompilandCache(m_compiledCompilands, m_objFileToCompilandId);
-	// END EPIC MOD
-
 	// update the COFF cache for all compiled files
 	UpdateCoffCache(m_compiledCompilands, m_coffCache, CacheUpdate::ALL, coffReadFlags, modifiedOrNewObjFiles);
 
diff --git a/Engine/Source/Developer/Windows/LiveCodingServer/Private/External/LC_ServerCommandThread.cpp b/Engine/Source/Developer/Windows/LiveCodingServer/Private/External/LC_ServerCommandThread.cpp
--- a/Engine/Source/Developer/Windows/LiveCodingServer/Private/External/LC_ServerCommandThread.cpp
+++ b/Engine/Source/Developer/Windows/LiveCodingServer/Private/External/LC_ServerCommandThread.cpp
@@ -743,6 +743,9 @@ void ServerCommandThread::CompileChanges(bool didAllProcessesMakeProgress)
 			}
 		}
 
+		// Reset the unity file cache
+		symbols::ResetCachedUnityManifests();
+
 		// Build up a list of all the modified object files in each module
 		types::unordered_set<std::wstring> ValidModuleFileNames;
 		for (const LiveModule* liveModule : m_liveModules)
diff --git a/Engine/Source/Developer/Windows/LiveCodingServer/Private/External/LC_Symbols.cpp b/Engine/Source/Developer/Windows/LiveCodingServer/Private/External/LC_Symbols.cpp
--- a/Engine/Source/Developer/Windows/LiveCodingServer/Private/External/LC_Symbols.cpp
+++ b/Engine/Source/Developer/Windows/LiveCodingServer/Private/External/LC_Symbols.cpp
@@ -24,6 +24,11 @@
 #include <algorithm>
 #include <process.h>
 
+#include "Misc/FileHelper.h"
+#include "Serialization/JsonReader.h"
+#include "Serialization/JsonSerializer.h"
+#include "Dom/JsonObject.h"
+
 #include "Windows/AllowWindowsPlatformAtomics.h"
 
 namespace
@@ -401,6 +406,112 @@ namespace symbols
 		}
 	}
 
+	// BEGIN EPIC MOD - Static grouping of compilands by unity blobs
+	static CriticalSection g_objFileToCompilandIdCS;
+	static std::unordered_map<std::wstring, uint32_t> g_objFileToCompilandId;
+	static std::unordered_set<std::wstring> g_checkedUnityManifests;
+
+	void ResetCachedUnityManifests()
+	{
+		CriticalSection::ScopedLock lock(&g_objFileToCompilandIdCS);
+		g_checkedUnityManifests.clear();
+	}
+
+	bool TryGetCompilandIdFromUnityManifest(const std::wstring& objFile, uint32_t& compilandId)
+	{
+		std::wstring normalizedObjFile = file::NormalizePath(objFile.c_str());
+		CriticalSection::ScopedLock lock(&g_objFileToCompilandIdCS);
+
+		// Check if it's already cached
+		std::unordered_map<std::wstring, uint32_t>::iterator it = g_objFileToCompilandId.find(normalizedObjFile);
+		if(it == g_objFileToCompilandId.end())
+		{
+			// Read the manifest file
+			std::wstring BaseDir = file::GetDirectory(normalizedObjFile);
+			std::wstring ManifestFile = BaseDir + L"\\LiveCodingInfo.json";
+
+			// If we've already tried to read this string, don't try again
+			if(!g_checkedUnityManifests.insert(ManifestFile).second)
+			{
+				return false;
+			}
+
+			// Read the file to a string
+			FString FileContents;
+			if (!FFileHelper::LoadFileToString(FileContents, ManifestFile.c_str()))
+			{
+				LC_WARNING_USER("%S could not be read", ManifestFile.c_str());
+				return false;
+			}
+
+			// Deserialize a JSON object from the string
+			TSharedPtr< FJsonObject > Object;
+			TSharedRef< TJsonReader<> > Reader = TJsonReaderFactory<>::Create(FileContents);
+			if (!FJsonSerializer::Deserialize(Reader, Object) || !Object.IsValid())
+			{
+				LC_WARNING_USER("%S could not be parsed", ManifestFile.c_str());
+				return false;
+			}
+
+			const TSharedPtr<FJsonObject>* FilesObject;
+			if (!Object->TryGetObjectField(TEXT("RemapUnityFiles"), FilesObject))
+			{
+				LC_WARNING_USER("%S is not a valid manifest file", ManifestFile.c_str());
+				return false;
+			}
+
+			for (const TPair<FString, TSharedPtr<FJsonValue>>& Pair : FilesObject->Get()->Values)
+			{
+				std::wstring UnityObjectFile = file::NormalizePath((BaseDir + L"\\" + *Pair.Key).c_str());
+				uint32_t UnityCompilandId = uniqueId::Generate(UnityObjectFile);
+				g_objFileToCompilandId.insert(std::make_pair(UnityObjectFile, UnityCompilandId));
+
+				const FJsonValue* Value = Pair.Value.Get();
+				if (Value->Type != EJson::Array)
+				{
+					LC_WARNING_USER("%S is not a valid manifest file", ManifestFile.c_str());
+					return false;
+				}
+
+				const TArray<TSharedPtr<FJsonValue>>& SourceFileValues = Value->AsArray();
+				for (const TSharedPtr<FJsonValue>& SourceFileValue : SourceFileValues)
+				{
+					if (SourceFileValue->Type != EJson::String)
+					{
+						LC_WARNING_USER("%S is not a valid manifest file", ManifestFile.c_str());
+						return false;
+					}
+
+					std::wstring MemberObjFile = file::NormalizePath((BaseDir + L"\\" + *SourceFileValue->AsString()).c_str());
+					g_objFileToCompilandId.insert(std::make_pair(MemberObjFile, UnityCompilandId));
+				}
+			}
+
+			// Check again for the object file we're interested in
+			it = g_objFileToCompilandId.find(normalizedObjFile);
+			if(it == g_objFileToCompilandId.end())
+			{
+				return false;
+			}
+		}
+
+		compilandId = it->second;
+		return true;
+	}
+
+	uint32_t GetCompilandIdFromPath(const std::wstring& objPath)
+	{
+		uint32 compilandId;
+		if(TryGetCompilandIdFromUnityManifest(objPath, compilandId))
+		{
+			return compilandId;
+		}
+		else
+		{
+			return uniqueId::Generate(file::NormalizePath(objPath.c_str()));
+		}
+	}
+	// END EPIC MOD
 
 	SymbolDB* GatherSymbols(Provider* provider)
 	{
@@ -887,7 +998,7 @@ namespace symbols
 					envCompilerWorkingDirectory,
 					ImmutableString(""),							// amalgamation .obj path
 					nullptr,										// file indices
-					uniqueId::Generate(normalizedCompilandPath),	// unique ID
+					GetCompilandIdFromPath(normalizedCompilandPath),	// unique ID
 					static_cast<uint32_t>(i),						// dia symbol index
 					Compiland::Type::SINGLE_FILE,					// type of file
 					isPartOfLibrary,								// isPartOfLibrary
@@ -1153,7 +1264,7 @@ namespace symbols
 											// we treat these files as being the amalgamated file.
 											// symbols originally coming from amalgamated files need to have the same
 											// name as symbols from individual files.
-											uniqueId::Generate(normalizedCompilandPath),
+											GetCompilandIdFromPath(normalizedCompilandPath),
 
 											// same for the DIA symbol index, for the same reason
 											static_cast<uint32_t>(i),
diff --git a/Engine/Source/Developer/Windows/LiveCodingServer/Private/External/LC_Symbols.h b/Engine/Source/Developer/Windows/LiveCodingServer/Private/External/LC_Symbols.h
--- a/Engine/Source/Developer/Windows/LiveCodingServer/Private/External/LC_Symbols.h
+++ b/Engine/Source/Developer/Windows/LiveCodingServer/Private/External/LC_Symbols.h
@@ -88,10 +88,6 @@ namespace symbols
 		bool isPartOfLibrary;
 		bool wasRecompiled;
 
-		// BEGIN EPIC MOD - Allow mapping from object files to their unity object file
-		uint32_t amalgamatedUniqueId = ~(uint32_t)0;
-		// END EPIC MOD
-
 		LC_DISABLE_ASSIGNMENT(Compiland);
 		LC_DISABLE_COPY(Compiland);
 	};
@@ -216,6 +212,11 @@ namespace symbols
 	Provider* OpenEXE(const wchar_t* filename, uint32_t openOptions);
 	void Close(Provider* provider);
 
+	// BEGIN EPIC MOD - Static grouping of compilands by unity blobs
+	void ResetCachedUnityManifests();
+	bool TryGetCompilandIdFromUnityManifest(const std::wstring& objPath, uint32_t& compilandId);
+	uint32_t GetCompilandIdFromPath(const std::wstring& objPath);
+	// END EPIC MOD
 
 
 	// UPDATE
