LiveCoding: Live Coding cannot be re-enabled after being disabled from the Compile button.

#rnx

diff --git a/Engine/Source/Developer/Windows/LiveCoding/Private/LiveCodingModule.cpp b/Engine/Source/Developer/Windows/LiveCoding/Private/LiveCodingModule.cpp
--- a/Engine/Source/Developer/Windows/LiveCoding/Private/LiveCodingModule.cpp
+++ b/Engine/Source/Developer/Windows/LiveCoding/Private/LiveCodingModule.cpp
@@ -165,6 +165,11 @@ void FLiveCodingModule::EnableForSession(bool bEnable)
 			StartLiveCoding();
 			ShowConsole();
 		}
+		else
+		{
+			bEnabledForSession = true;
+			ShowConsole();
+		}
 	}
 	else 
 	{
