LiveCoding: Add a delegate that's called when a Live Coding patch is applied.

diff --git a/Engine/Source/Developer/Windows/LiveCoding/Private/External/LC_ClientCommandActions.cpp b/Engine/Source/Developer/Windows/LiveCoding/Private/External/LC_ClientCommandActions.cpp
--- a/Engine/Source/Developer/Windows/LiveCoding/Private/External/LC_ClientCommandActions.cpp
+++ b/Engine/Source/Developer/Windows/LiveCoding/Private/External/LC_ClientCommandActions.cpp
@@ -120,6 +120,10 @@ bool actions::LoadPatch::Execute(const CommandType* command, const DuplexPipe* p
 		}
 	}
 	// END EPIC MOD
+	// BEGIN EPIC MOD - Support for object reinstancing. We need to call a global post-patch handler, rather than just getting individual callbacks for modified modules.
+	extern bool GHasLoadedPatch;
+	GHasLoadedPatch = true;
+	// END EPIC MOD
 
 	pipe->SendAck();
 
diff --git a/Engine/Source/Developer/Windows/LiveCoding/Private/LiveCodingModule.cpp b/Engine/Source/Developer/Windows/LiveCoding/Private/LiveCodingModule.cpp
--- a/Engine/Source/Developer/Windows/LiveCoding/Private/LiveCodingModule.cpp
+++ b/Engine/Source/Developer/Windows/LiveCoding/Private/LiveCodingModule.cpp
@@ -21,6 +21,7 @@ IMPLEMENT_MODULE(FLiveCodingModule, LiveCoding)
 #define LOCTEXT_NAMESPACE "LiveCodingModule"
 
 bool GIsCompileActive = false;
+bool GHasLoadedPatch = false;
 FString GLiveCodingConsolePath;
 FString GLiveCodingConsoleArguments;
 
@@ -236,6 +237,15 @@ bool FLiveCodingModule::IsCompiling() const
 
 void FLiveCodingModule::Tick()
 {
+	extern void LppSyncPoint();
+	LppSyncPoint();
+
+	if (GHasLoadedPatch)
+	{
+		OnPatchCompleteDelegate.Broadcast();
+		GHasLoadedPatch = false;
+	}
+
 	if (LppWantsRestart())
 	{
 		LppRestart(lpp::LPP_RESTART_BEHAVIOR_REQUEST_EXIT, 0);
@@ -254,6 +264,11 @@ void FLiveCodingModule::Tick()
 	}
 }
 
+ILiveCodingModule::FOnPatchCompleteDelegate& FLiveCodingModule::GetOnPatchCompleteDelegate()
+{
+	return OnPatchCompleteDelegate;
+}
+
 bool FLiveCodingModule::StartLiveCoding()
 {
 	if(!bStarted)
diff --git a/Engine/Source/Developer/Windows/LiveCoding/Private/LiveCodingModule.h b/Engine/Source/Developer/Windows/LiveCoding/Private/LiveCodingModule.h
--- a/Engine/Source/Developer/Windows/LiveCoding/Private/LiveCodingModule.h
+++ b/Engine/Source/Developer/Windows/LiveCoding/Private/LiveCodingModule.h
@@ -32,6 +32,7 @@ public:
 	virtual void Compile() override;
 	virtual bool IsCompiling() const override;
 	virtual void Tick() override;
+	virtual FOnPatchCompleteDelegate& GetOnPatchCompleteDelegate() override;
 
 private:
 	ULiveCodingSettings* Settings;
@@ -52,6 +53,7 @@ private:
 	IConsoleVariable* SourceProjectVariable;
 	FDelegateHandle EndFrameDelegateHandle;
 	FDelegateHandle ModulesChangedDelegateHandle;
+	FOnPatchCompleteDelegate OnPatchCompleteDelegate;
 
 	bool StartLiveCoding();
 
diff --git a/Engine/Source/Developer/Windows/LiveCoding/Public/ILiveCodingModule.h b/Engine/Source/Developer/Windows/LiveCoding/Public/ILiveCodingModule.h
--- a/Engine/Source/Developer/Windows/LiveCoding/Public/ILiveCodingModule.h
+++ b/Engine/Source/Developer/Windows/LiveCoding/Public/ILiveCodingModule.h
@@ -3,6 +3,7 @@
 #pragma once
 
 #include "CoreTypes.h"
+#include "Delegates/Delegate.h"
 #include "Modules/ModuleInterface.h"
 
 #define LIVE_CODING_MODULE_NAME "LiveCoding"
@@ -23,5 +24,8 @@ public:
 	virtual void Compile() = 0;
 	virtual bool IsCompiling() const = 0;
 	virtual void Tick() = 0;
+
+	DECLARE_MULTICAST_DELEGATE(FOnPatchCompleteDelegate);
+	virtual FOnPatchCompleteDelegate& GetOnPatchCompleteDelegate() = 0;
 };
 
