Fix compile error on Clang due to different behavior in scoping for extern'd variables.

diff --git a/Engine/Source/Developer/Windows/LiveCoding/Private/External/LC_ClientCommandActions.cpp b/Engine/Source/Developer/Windows/LiveCoding/Private/External/LC_ClientCommandActions.cpp
--- a/Engine/Source/Developer/Windows/LiveCoding/Private/External/LC_ClientCommandActions.cpp
+++ b/Engine/Source/Developer/Windows/LiveCoding/Private/External/LC_ClientCommandActions.cpp
@@ -9,6 +9,9 @@
 #include "LC_Process.h"
 #include "LC_Logging.h"
 
+// BEGIN EPIC MOD - Support for object reinstancing. We need to call a global post-patch handler, rather than just getting individual callbacks for modified modules.
+extern bool GHasLoadedPatch;
+// END EPIC MOD
 
 bool actions::RegisterProcessFinished::Execute(const CommandType* command, const DuplexPipe* pipe, void* context, const void*, size_t)
 {
@@ -121,7 +124,6 @@ bool actions::LoadPatch::Execute(const CommandType* command, const DuplexPipe* p
 	}
 	// END EPIC MOD
 	// BEGIN EPIC MOD - Support for object reinstancing. We need to call a global post-patch handler, rather than just getting individual callbacks for modified modules.
-	extern bool GHasLoadedPatch;
 	GHasLoadedPatch = true;
 	// END EPIC MOD
 
