LiveCoding: Fix deadlock when using the quick restart functionality, where a patch can be loaded before the module has been registered.

#rnx

diff --git a/Engine/Source/Developer/Windows/LiveCoding/Private/LiveCodingModule.cpp b/Engine/Source/Developer/Windows/LiveCoding/Private/LiveCodingModule.cpp
--- a/Engine/Source/Developer/Windows/LiveCoding/Private/LiveCodingModule.cpp
+++ b/Engine/Source/Developer/Windows/LiveCoding/Private/LiveCodingModule.cpp
@@ -237,15 +237,6 @@ bool FLiveCodingModule::IsCompiling() const
 
 void FLiveCodingModule::Tick()
 {
-	extern void LppSyncPoint();
-	LppSyncPoint();
-
-	if (GHasLoadedPatch)
-	{
-		OnPatchCompleteDelegate.Broadcast();
-		GHasLoadedPatch = false;
-	}
-
 	if (LppWantsRestart())
 	{
 		LppRestart(lpp::LPP_RESTART_BEHAVIOR_REQUEST_EXIT, 0);
@@ -262,6 +253,16 @@ void FLiveCodingModule::Tick()
 		UpdateModules();
 		bUpdateModulesInTick = false;
 	}
+
+	// Needs to happen after updating modules, since "Quick Restart" functionality may try to install patch immediately
+	extern void LppSyncPoint();
+	LppSyncPoint();
+
+	if (GHasLoadedPatch)
+	{
+		OnPatchCompleteDelegate.Broadcast();
+		GHasLoadedPatch = false;
+	}
 }
 
 ILiveCodingModule::FOnPatchCompleteDelegate& FLiveCodingModule::GetOnPatchCompleteDelegate()
