Live Coding: Remove warning about LiveCodingInfo.json files not being found. These files do not exist in the binary release.

diff --git a/Engine/Source/Developer/Windows/LiveCodingServer/Private/External/LC_Symbols.cpp b/Engine/Source/Developer/Windows/LiveCodingServer/Private/External/LC_Symbols.cpp
--- a/Engine/Source/Developer/Windows/LiveCodingServer/Private/External/LC_Symbols.cpp
+++ b/Engine/Source/Developer/Windows/LiveCodingServer/Private/External/LC_Symbols.cpp
@@ -440,7 +440,6 @@ namespace symbols
 			FString FileContents;
 			if (!FFileHelper::LoadFileToString(FileContents, ManifestFile.c_str()))
 			{
-				LC_WARNING_USER("%S could not be read", ManifestFile.c_str());
 				return false;
 			}
 
