Updated to latest DXC (github 4167dd46)

* Fixed VertexFactoryInterpolate() texcoord interpolation when USE_PARTICLE_SUBUVS=1 (new compiler generates code that fails DXIL validation if texcoords are uninitialized)
* Use D3D_SIT_RTACCELERATIONSTRUCTURE define instead of hard-coded value in D3DShaderCompiler.cpp (addresses an old dxr_todo)

diff --git a/Engine/Shaders/Private/LocalVertexFactory.ush b/Engine/Shaders/Private/LocalVertexFactory.ush
--- a/Engine/Shaders/Private/LocalVertexFactory.ush
+++ b/Engine/Shaders/Private/LocalVertexFactory.ush
@@ -1071,6 +1071,8 @@ float4 VertexFactoryGetPreviousWorldPosition(FVertexFactoryInput Input, FVertexF
 		{
 			TESSELLATION_INTERPOLATE_MEMBER(InterpolantsVSToPS.TexCoords[tc]);
 		}
+	#elif USE_PARTICLE_SUBUVS
+		TESSELLATION_INTERPOLATE_MEMBER(InterpolantsVSToPS.TexCoords[0]);
 	#endif
 
 		return O;
diff --git a/Engine/Shaders/Shared/RayTracingDefinitions.h b/Engine/Shaders/Shared/RayTracingDefinitions.h
--- a/Engine/Shaders/Shared/RayTracingDefinitions.h
+++ b/Engine/Shaders/Shared/RayTracingDefinitions.h
@@ -9,7 +9,7 @@
 
 // Change this to force recompilation of all ray tracing shaders (use https://www.random.org/cgi-bin/randbyte?nbytes=4&format=h)
 // This avoids changing the global ShaderVersion.ush and forcing recompilation of all shaders in the engine (only RT shaders will be affected)
-#define RAY_TRACING_SHADER_VERSION 0x510cacb6
+#define RAY_TRACING_SHADER_VERSION 0x67f74955
 
 #define RAY_TRACING_REGISTER_SPACE_LOCAL  0 // default register space for hit group (closest hit, any hit, intersection) shader resources
 #define RAY_TRACING_REGISTER_SPACE_GLOBAL 1 // register space for ray generation and miss shaders 
diff --git a/Engine/Source/Developer/Windows/ShaderFormatD3D/Private/D3DShaderCompiler.cpp b/Engine/Source/Developer/Windows/ShaderFormatD3D/Private/D3DShaderCompiler.cpp
--- a/Engine/Source/Developer/Windows/ShaderFormatD3D/Private/D3DShaderCompiler.cpp
+++ b/Engine/Source/Developer/Windows/ShaderFormatD3D/Private/D3DShaderCompiler.cpp
@@ -826,8 +826,7 @@ static void ExtractParameterMapFromD3DShader(
 
 			NumSRVs = FMath::Max(NumSRVs, BindDesc.BindPoint + BindCount);
 		}
-		// #dxr_todo: D3D_SIT_RTACCELERATIONSTRUCTURE is declared in latest version of dxcapi.h. Update this code after upgrading DXC.
-		else if (BindDesc.Type == 12 /*D3D_SIT_RTACCELERATIONSTRUCTURE*/)
+		else if (BindDesc.Type == D3D_SIT_RTACCELERATIONSTRUCTURE)
 		{
 			// Acceleration structure resources are treated as SRVs.
 			check(BindDesc.BindCount == 1);
