Added an option to globally force using DXC for D3D shaders, intended for testing compatibility and correctness when using the new compiler.

diff --git a/Engine/Source/Developer/Windows/ShaderFormatD3D/Private/D3DShaderCompiler.cpp b/Engine/Source/Developer/Windows/ShaderFormatD3D/Private/D3DShaderCompiler.cpp
--- a/Engine/Source/Developer/Windows/ShaderFormatD3D/Private/D3DShaderCompiler.cpp
+++ b/Engine/Source/Developer/Windows/ShaderFormatD3D/Private/D3DShaderCompiler.cpp
@@ -113,7 +113,7 @@ static uint32 GetAutoBindingSpace(const FShaderTarget& Target)
 }
 
 // @return 0 if not recognized
-static const TCHAR* GetShaderProfileName(FShaderTarget Target, bool bUseWaveOperations)
+static const TCHAR* GetShaderProfileName(FShaderTarget Target, bool bForceSM6)
 {
 	if(Target.Platform == SP_PCD3D_SM5)
 	{
@@ -124,17 +124,17 @@ static const TCHAR* GetShaderProfileName(FShaderTarget Target, bool bUseWaveOper
 			checkfSlow(false, TEXT("Unexpected shader frequency"));
 			return nullptr;
 		case SF_Pixel:
-			return bUseWaveOperations ? TEXT("ps_6_0") : TEXT("ps_5_0");
+			return bForceSM6 ? TEXT("ps_6_0") : TEXT("ps_5_0");
 		case SF_Vertex:
-			return bUseWaveOperations ? TEXT("vs_6_0") : TEXT("vs_5_0");
+			return bForceSM6 ? TEXT("vs_6_0") : TEXT("vs_5_0");
 		case SF_Hull:
-			return bUseWaveOperations ? TEXT("hs_6_0") : TEXT("hs_5_0");
+			return bForceSM6 ? TEXT("hs_6_0") : TEXT("hs_5_0");
 		case SF_Domain:
-			return bUseWaveOperations ? TEXT("ds_6_0") : TEXT("ds_5_0");
+			return bForceSM6 ? TEXT("ds_6_0") : TEXT("ds_5_0");
 		case SF_Geometry:
-			return bUseWaveOperations ? TEXT("gs_6_0") : TEXT("gs_5_0");
+			return bForceSM6 ? TEXT("gs_6_0") : TEXT("gs_5_0");
 		case SF_Compute:
-			return bUseWaveOperations ? TEXT("cs_6_0") : TEXT("cs_5_0");
+			return bForceSM6 ? TEXT("cs_6_0") : TEXT("cs_5_0");
 		case SF_RayGen:
 		case SF_RayMiss:
 		case SF_RayHitGroup:
@@ -893,7 +893,9 @@ static bool CompileAndProcessD3DShader(FString& PreprocessedShaderSource, const
 	auto AnsiSourceFile = StringCast<ANSICHAR>(*PreprocessedShaderSource);
 
 	const bool bIsRayTracingShader = IsRayTracingShader(Input.Target);
-	const bool bUseDXC = bIsRayTracingShader || Input.Environment.CompilerFlags.Contains(CFLAG_WaveOperations);
+	const bool bUseDXC = bIsRayTracingShader
+		|| Input.Environment.CompilerFlags.Contains(CFLAG_WaveOperations)
+		|| Input.Environment.CompilerFlags.Contains(CFLAG_ForceDXC);
 
 	const uint32 AutoBindingSpace = GetAutoBindingSpace(Input.Target);
 
@@ -1521,7 +1523,8 @@ void CompileD3DShader(const FShaderCompilerInput& Input,FShaderCompilerOutput& O
 	FString PreprocessedShaderSource;
 	FString CompilerPath;
 	const bool bUseWaveOperations = Input.Environment.CompilerFlags.Contains(CFLAG_WaveOperations); // Forces shader model 6.0 for this shader
-	const TCHAR* ShaderProfile = GetShaderProfileName(Input.Target, bUseWaveOperations);
+	const bool bForceDXC = Input.Environment.CompilerFlags.Contains(CFLAG_ForceDXC);
+	const TCHAR* ShaderProfile = GetShaderProfileName(Input.Target, bUseWaveOperations || bForceDXC);
 
 	if(!ShaderProfile)
 	{
@@ -1532,7 +1535,7 @@ void CompileD3DShader(const FShaderCompilerInput& Input,FShaderCompilerOutput& O
 	// Set additional defines.
 	AdditionalDefines.SetDefine(TEXT("COMPILER_HLSL"), 1);
 
-	if (bUseWaveOperations)
+	if (bUseWaveOperations || bForceDXC)
 	{
 		AdditionalDefines.SetDefine(TEXT("PLATFORM_SUPPORTS_SM6_0_WAVE_OPERATIONS"), 1);
 	}
diff --git a/Engine/Source/Runtime/Engine/Private/ShaderCompiler/ShaderCompiler.cpp b/Engine/Source/Runtime/Engine/Private/ShaderCompiler/ShaderCompiler.cpp
--- a/Engine/Source/Runtime/Engine/Private/ShaderCompiler/ShaderCompiler.cpp
+++ b/Engine/Source/Runtime/Engine/Private/ShaderCompiler/ShaderCompiler.cpp
@@ -236,6 +236,14 @@ static TAutoConsoleVariable<int32> CVarD3DRemoveUnusedInterpolators(
 	ECVF_ReadOnly
 	);
 
+static TAutoConsoleVariable<int32> CVarD3DForceDXC(
+	TEXT("r.D3D.ForceDXC"),
+	0,
+	TEXT("Forces DirectX Shader Compiler (DXC) to be used for all D3D shaders. Shaders compiled with this option are only compatible with D3D12.\n")
+	TEXT(" 0: Disable (default)\n")
+	TEXT(" 1: Force new compiler for all shaders"),
+	ECVF_ReadOnly);
+
 int32 GCreateShadersOnLoad = 0;
 static FAutoConsoleVariableRef CVarCreateShadersOnLoad(
 	TEXT("r.CreateShadersOnLoad"),
@@ -3106,10 +3114,19 @@ void GlobalBeginCompileShader(
 
 	if (IsD3DPlatform((EShaderPlatform)Target.Platform, false))
 	{
-		static const auto CVar = IConsoleManager::Get().FindConsoleVariable(TEXT("r.D3D.RemoveUnusedInterpolators"));
-		if (CVar && CVar->GetInt() != 0)
 		{
-			Input.Environment.CompilerFlags.Add(CFLAG_ForceRemoveUnusedInterpolators);
+			static const auto CVar = IConsoleManager::Get().FindConsoleVariable(TEXT("r.D3D.RemoveUnusedInterpolators"));
+			if (CVar && CVar->GetInt() != 0)
+			{
+				Input.Environment.CompilerFlags.Add(CFLAG_ForceRemoveUnusedInterpolators);
+			}
+		}
+		{
+			static const auto CVar = IConsoleManager::Get().FindConsoleVariable(TEXT("r.D3D.ForceDXC"));
+			if (CVar && CVar->GetInt() != 0)
+			{
+				Input.Environment.CompilerFlags.Add(CFLAG_ForceDXC);
+			}
 		}
 	}
 
diff --git a/Engine/Source/Runtime/RenderCore/Private/Shader.cpp b/Engine/Source/Runtime/RenderCore/Private/Shader.cpp
--- a/Engine/Source/Runtime/RenderCore/Private/Shader.cpp
+++ b/Engine/Source/Runtime/RenderCore/Private/Shader.cpp
@@ -2168,10 +2168,19 @@ void ShaderMapAppendKeyString(EShaderPlatform Platform, FString& KeyString)
 
 	if (IsD3DPlatform(Platform, false))
 	{
-		static const auto CVar = IConsoleManager::Get().FindConsoleVariable(TEXT("r.D3D.RemoveUnusedInterpolators"));
-		if (CVar && CVar->GetInt() != 0)
 		{
-			KeyString += TEXT("_UnInt");
+			static const auto CVar = IConsoleManager::Get().FindConsoleVariable(TEXT("r.D3D.RemoveUnusedInterpolators"));
+			if (CVar && CVar->GetInt() != 0)
+			{
+				KeyString += TEXT("_UnInt");
+			}
+		}
+		{
+			static const auto CVar = IConsoleManager::Get().FindConsoleVariable(TEXT("r.D3D.ForceDXC"));
+			if (CVar && CVar->GetInt() != 0)
+			{
+				KeyString += TEXT("_DXC");
+			}
 		}
 	}
 
diff --git a/Engine/Source/Runtime/RenderCore/Public/ShaderCore.h b/Engine/Source/Runtime/RenderCore/Public/ShaderCore.h
--- a/Engine/Source/Runtime/RenderCore/Public/ShaderCore.h
+++ b/Engine/Source/Runtime/RenderCore/Public/ShaderCore.h
@@ -185,6 +185,8 @@ enum ECompilerFlags
 	// Check GRHISupportsWaveOperations before using shaders compiled with this flag at runtime.
 	// [REDACTED_PROFILE_URL]/DirectXShaderCompiler/wiki/Wave-Intrinsics
 	CFLAG_WaveOperations,
+	// Use DirectX Shader Compiler (DXC) to compile all shaders, intended for compatibility testing.
+	CFLAG_ForceDXC,
 };
 
 enum class EShaderParameterType : uint8
