Backed out accidentally submitted debug change.

diff --git a/Engine/Source/Developer/Windows/ShaderFormatD3D/Private/D3DShaderCompiler.cpp b/Engine/Source/Developer/Windows/ShaderFormatD3D/Private/D3DShaderCompiler.cpp
--- a/Engine/Source/Developer/Windows/ShaderFormatD3D/Private/D3DShaderCompiler.cpp
+++ b/Engine/Source/Developer/Windows/ShaderFormatD3D/Private/D3DShaderCompiler.cpp
@@ -1507,7 +1507,7 @@ static bool CompileAndProcessD3DShader(FString& PreprocessedShaderSource, const
 	{
 		if (Input.Target.Platform == SP_PCD3D_ES2)
 		{
-			if (false && Output.NumTextureSamplers > 8)
+			if (Output.NumTextureSamplers > 8)
 			{
 				FilteredErrors.Add(FString::Printf(TEXT("Shader uses more than 8 texture samplers which is not supported by ES2!  Used: %u"), Output.NumTextureSamplers));
 				Result = E_FAIL;
