Fixed DXC shader diagnostic message reporting

* DXC error buffer is not null-terminated, which previously caused garbage characters at the end of the diagnostic message
* DXC diagnostic messages include few lines of context which are now filtered out to be consistent with how we report FXC messages

diff --git a/Engine/Source/Developer/Windows/ShaderFormatD3D/Private/D3DShaderCompiler.cpp b/Engine/Source/Developer/Windows/ShaderFormatD3D/Private/D3DShaderCompiler.cpp
--- a/Engine/Source/Developer/Windows/ShaderFormatD3D/Private/D3DShaderCompiler.cpp
+++ b/Engine/Source/Developer/Windows/ShaderFormatD3D/Private/D3DShaderCompiler.cpp
@@ -85,6 +85,25 @@ static void D3D11FilterShaderCompileWarnings(const FString& CompileWarnings, TAr
 	}
 }
 
+static void DXCFilterShaderCompileWarnings(const FString& CompileWarnings, TArray<FString>& FilteredWarnings)
+{
+	TArray<FString> WarningArray;
+	FString OutWarningString = TEXT("");
+	CompileWarnings.ParseIntoArray(WarningArray, TEXT("\n"), true);
+
+	//go through each warning line
+	for (int32 WarningIndex = 0; WarningIndex < WarningArray.Num(); WarningIndex++)
+	{
+		// DXC includes few lines of extra context in its diagnostic messages, which is different from FXC.
+		// We only include actual error and warning lines in the filtered output.
+		if (WarningArray[WarningIndex].Contains(TEXT("error:"))
+			|| WarningArray[WarningIndex].Contains(TEXT("warning:")))
+		{
+			FilteredWarnings.AddUnique(WarningArray[WarningIndex]);
+		}
+	}
+}
+
 static bool IsRayTracingShader(const FShaderTarget& Target)
 {
 	switch(Target.Frequency)
@@ -1018,8 +1037,13 @@ static bool CompileAndProcessD3DShader(FString& PreprocessedShaderSource, const
 
 		if (DxcErrorBlob && DxcErrorBlob->GetBufferSize())
 		{
-			void* ErrorBuffer = DxcErrorBlob->GetBufferPointer();
-			D3D11FilterShaderCompileWarnings(ANSI_TO_TCHAR(ErrorBuffer), FilteredErrors);
+			ANSICHAR* ErrorChars = new ANSICHAR[DxcErrorBlob->GetBufferSize() + 1];
+			FMemory::Memcpy(ErrorChars, DxcErrorBlob->GetBufferPointer(), DxcErrorBlob->GetBufferSize());
+			ErrorChars[DxcErrorBlob->GetBufferSize()] = 0;
+			FString ErrorString(ErrorChars);
+			delete[] ErrorChars;
+
+			DXCFilterShaderCompileWarnings(ErrorString, FilteredErrors);
 		}
 
 		if (!SUCCEEDED(Result))
