Added DXBC disassembly check for loads using additional typed formats, which are not supported on all versions of Windows (despite requesting a vanilla DX 11.0 feature level). Additionally this change no longer makes dumping out debug info and checking for double ops mutually exclusive (checks would be skipped if debug dump was enabled). Some variable name typos were fixed, as well.

(ushell-p4-cherrypick of 8067160 by graham.wihlidal)

diff --git a/Engine/Source/Developer/Windows/ShaderFormatD3D/Private/D3DShaderCompiler.cpp b/Engine/Source/Developer/Windows/ShaderFormatD3D/Private/D3DShaderCompiler.cpp
--- a/Engine/Source/Developer/Windows/ShaderFormatD3D/Private/D3DShaderCompiler.cpp
+++ b/Engine/Source/Developer/Windows/ShaderFormatD3D/Private/D3DShaderCompiler.cpp
@@ -44,6 +44,7 @@ static int32 GD3DAllowRemoveUnused = 0;
 
 
 static int32 GD3DCheckForDoubles = 1;
+static int32 GD3DCheckForTypedUAVs = 1;
 static int32 GD3DDumpAMDCodeXLFile = 0;
 
 /**
@@ -1097,30 +1098,45 @@ static bool CompileAndProcessD3DShader(FString& PreprocessedShaderSource, const
 		// Fail the compilation if double operations are being used, since those are not supported on all D3D11 cards
 		if (SUCCEEDED(Result))
 		{
-			if (D3DDisassembleFunc && (GD3DCheckForDoubles || bDumpDebugInfo))
+			if (D3DDisassembleFunc && (GD3DCheckForDoubles || GD3DCheckForTypedUAVs || bDumpDebugInfo))
 			{
-				TRefCountPtr<ID3DBlob> Dissasembly;
-				if (SUCCEEDED(D3DDisassembleFunc(Shader->GetBufferPointer(), Shader->GetBufferSize(), 0, "", Dissasembly.GetInitReference())))
+				TRefCountPtr<ID3DBlob> Disassembly;
+				if (SUCCEEDED(D3DDisassembleFunc(Shader->GetBufferPointer(), Shader->GetBufferSize(), 0, "", Disassembly.GetInitReference())))
 				{
-					ANSICHAR* DissasemblyString = new ANSICHAR[Dissasembly->GetBufferSize() + 1];
-					FMemory::Memcpy(DissasemblyString, Dissasembly->GetBufferPointer(), Dissasembly->GetBufferSize());
-					DissasemblyString[Dissasembly->GetBufferSize()] = 0;
-					FString DissasemblyStringW(DissasemblyString);
-					delete[] DissasemblyString;
+					ANSICHAR* DisassemblyString = new ANSICHAR[Disassembly->GetBufferSize() + 1];
+					FMemory::Memcpy(DisassemblyString, Disassembly->GetBufferPointer(), Disassembly->GetBufferSize());
+					DisassemblyString[Disassembly->GetBufferSize()] = 0;
+					FString DisassemblyStringW(DisassemblyString);
+					delete[] DisassemblyString;
 
 					if (bDumpDebugInfo)
 					{
-						FFileHelper::SaveStringToFile(DissasemblyStringW, *(Input.DumpDebugInfoPath / TEXT("Output.d3dasm")));
+						FFileHelper::SaveStringToFile(DisassemblyStringW, *(Input.DumpDebugInfoPath / TEXT("Output.d3dasm")));
 					}
-					else if (GD3DCheckForDoubles)
+
+					if (GD3DCheckForDoubles)
 					{
 						// dcl_globalFlags will contain enableDoublePrecisionFloatOps when the shader uses doubles, even though the docs on dcl_globalFlags don't say anything about this
-						if (DissasemblyStringW.Contains(TEXT("enableDoublePrecisionFloatOps")))
+						if (DisassemblyStringW.Contains(TEXT("enableDoublePrecisionFloatOps")))
 						{
 							FilteredErrors.Add(TEXT("Shader uses double precision floats, which are not supported on all D3D11 hardware!"));
 							return false;
 						}
 					}
+					
+					if (GD3DCheckForTypedUAVs)
+					{
+						// Disassembly will contain this text with typed loads from UAVs are used where the format and dimension are not fully supported
+						// across all versions of Windows (like Windows 7/8.1).
+						// https://microsoft.github.io/DirectX-Specs/d3d/UAVTypedLoad.html
+						// https://docs.microsoft.com/en-us/windows/win32/direct3d12/typed-unordered-access-view-loads
+						// https://docs.microsoft.com/en-us/windows/win32/direct3ddxgi/format-support-for-direct3d-11-0-feature-level-hardware
+						if (DisassemblyStringW.Contains(TEXT("Typed UAV Load Additional Formats")))
+						{
+							FilteredErrors.Add(TEXT("Shader uses UAV loads from additional typed formats, which are not supported on all D3D11 hardware!"));
+							return false;
+						}
+					}
 				}
 			}
 		}
