UE4R - Support precise on HlslParser
- Fix for not dumping shaders when using shader pipelines on error
- Fix for HlslParser skipping extra lines on #endif

diff --git a/Engine/Source/Developer/ShaderCompilerCommon/Private/HlslAST.cpp b/Engine/Source/Developer/ShaderCompilerCommon/Private/HlslAST.cpp
--- a/Engine/Source/Developer/ShaderCompilerCommon/Private/HlslAST.cpp
+++ b/Engine/Source/Developer/ShaderCompilerCommon/Private/HlslAST.cpp
@@ -706,6 +706,10 @@ namespace CrossCompiler
 			}
 			else
 			{
+				if (bPrecise)
+				{
+					Writer << TEXT("precise ");
+				}
 				Writer << TypeName;
 				if (TextureMSNumSamples > 1)
 				{
diff --git a/Engine/Source/Developer/ShaderCompilerCommon/Private/HlslAST.h b/Engine/Source/Developer/ShaderCompilerCommon/Private/HlslAST.h
--- a/Engine/Source/Developer/ShaderCompilerCommon/Private/HlslAST.h
+++ b/Engine/Source/Developer/ShaderCompilerCommon/Private/HlslAST.h
@@ -562,7 +562,7 @@ namespace CrossCompiler
 			int32 TextureMSNumSamples;
 
 			int32 PatchSize;
-
+			bool bPrecise = false;
 			bool bIsArray;
 			//bool bIsUnsizedArray;
 			FExpression* ArraySize;
diff --git a/Engine/Source/Developer/ShaderCompilerCommon/Private/HlslExpressionParser.inl b/Engine/Source/Developer/ShaderCompilerCommon/Private/HlslExpressionParser.inl
--- a/Engine/Source/Developer/ShaderCompilerCommon/Private/HlslExpressionParser.inl
+++ b/Engine/Source/Developer/ShaderCompilerCommon/Private/HlslExpressionParser.inl
@@ -178,7 +178,7 @@ namespace CrossCompiler
 		EEF_ALLOW_SEQUENCE		= 1 << 1,
 	};
 
-	EParseResult ParseGeneralType(const FHlslToken* Token, int32 TypeFlags, FLinearAllocator* Allocator, AST::FTypeSpecifier** OutSpecifier)
+	EParseResult ParseGeneralType(const FHlslToken* Token, int32 TypeFlags, bool bPrecise, FLinearAllocator* Allocator, AST::FTypeSpecifier** OutSpecifier)
 	{
 		if (!Token)
 		{
@@ -359,9 +359,11 @@ namespace CrossCompiler
 
 		if (bMatched)
 		{
+			//#todo-rco: Don't re-allocate types
 			auto* Type = new(Allocator) AST::FTypeSpecifier(Allocator, Token->SourceInfo);
 			Type->TypeName = Allocator->Strdup(Token->String);
 			Type->InnerType = InnerType;
+			Type->bPrecise = bPrecise;
 			*OutSpecifier = Type;
 			return EParseResult::Matched;
 		}
@@ -369,11 +371,11 @@ namespace CrossCompiler
 		return EParseResult::NotMatched;
 	}
 
-	EParseResult ParseGeneralTypeFromToken(const FHlslToken* Token, int32 TypeFlags, FSymbolScope* SymbolScope, FLinearAllocator* Allocator, AST::FTypeSpecifier** OutSpecifier)
+	EParseResult ParseGeneralTypeFromToken(const FHlslToken* Token, int32 TypeFlags, bool bPrecise, FSymbolScope* SymbolScope, FLinearAllocator* Allocator, AST::FTypeSpecifier** OutSpecifier)
 	{
 		if (Token)
 		{
-			if (ParseGeneralType(Token, TypeFlags, Allocator, OutSpecifier) == EParseResult::Matched)
+			if (ParseGeneralType(Token, TypeFlags, bPrecise, Allocator, OutSpecifier) == EParseResult::Matched)
 			{
 				return EParseResult::Matched;
 			}
@@ -466,7 +468,14 @@ namespace CrossCompiler
 			}
 		}
 
-		auto Result = ParseGeneralTypeFromToken(Token, TypeFlags, SymbolScope, Allocator, OutSpecifier);
+		bool bPrecise = false;
+		if (Scanner.MatchToken(EHlslToken::Precise))
+		{
+			Token = Scanner.GetCurrentToken();
+			bPrecise = true;
+		}
+
+		auto Result = ParseGeneralTypeFromToken(Token, TypeFlags, bPrecise, SymbolScope, Allocator, OutSpecifier);
 		if (Result == EParseResult::Matched)
 		{
 			Scanner.Advance();
@@ -538,21 +547,38 @@ namespace CrossCompiler
 				const auto* Peek2 = Scanner.PeekToken(2);
 
 				bool bFoundConst = false;
+				int PeekN = 0;
+				auto HandleUnaryToken = [&](EHlslToken TokenType)
+				{
+					if (Peek1->Token == TokenType)
+					{
+						++PeekN;
+						Peek1 = Scanner.PeekToken(1 + PeekN);
+						Peek2 = Scanner.PeekToken(2 + PeekN);
+						return true;
+					}
+					return false;
+				};
+
 				//#todo-rco: Workaround for weird const cast in RHS codegen from HLSLTranslator:
 				//		const uint exp2 = ((((const int) ((uRes32>>23)&0xff))-127+15) << 10);
-				if (Peek1->Token == EHlslToken::Const)
+				if (HandleUnaryToken(EHlslToken::Const))
+				{
+					HandleUnaryToken(EHlslToken::Precise);
+				}
+				else if (HandleUnaryToken(EHlslToken::Precise))
 				{
-					bFoundConst = true;
-					Peek1 = Scanner.PeekToken(2);
-					Peek2 = Scanner.PeekToken(3);
+					HandleUnaryToken(EHlslToken::Const);
 				}
 
 				AST::FTypeSpecifier* TypeSpecifier = nullptr;
-				if (Peek1 && ParseGeneralTypeFromToken(Peek1, ETF_BUILTIN_NUMERIC | ETF_USER_TYPES, SymbolScope, Allocator, &TypeSpecifier) == EParseResult::Matched && Peek2 && Peek2->Token == EHlslToken::RightParenthesis)
+				//#todo-rco: Is precise allowed on casts?
+				if (Peek1 && ParseGeneralTypeFromToken(Peek1, ETF_BUILTIN_NUMERIC | ETF_USER_TYPES, false, SymbolScope, Allocator, &TypeSpecifier) == EParseResult::Matched && Peek2 && Peek2->Token == EHlslToken::RightParenthesis)
 				{
-					if (bFoundConst)
+					while (PeekN > 0)
 					{
 						Scanner.Advance();
+						--PeekN;
 					}
 
 					// Cast
diff --git a/Engine/Source/Developer/ShaderCompilerCommon/Private/HlslLexer.h b/Engine/Source/Developer/ShaderCompilerCommon/Private/HlslLexer.h
--- a/Engine/Source/Developer/ShaderCompilerCommon/Private/HlslLexer.h
+++ b/Engine/Source/Developer/ShaderCompilerCommon/Private/HlslLexer.h
@@ -83,6 +83,7 @@ namespace CrossCompiler
 		// Types
 		Void,
 		Const,
+		Precise,
 
 		Bool,
 		Bool1,
diff --git a/Engine/Source/Developer/Windows/ShaderFormatD3D/Private/D3DShaderCompiler.cpp b/Engine/Source/Developer/Windows/ShaderFormatD3D/Private/D3DShaderCompiler.cpp
--- a/Engine/Source/Developer/Windows/ShaderFormatD3D/Private/D3DShaderCompiler.cpp
+++ b/Engine/Source/Developer/Windows/ShaderFormatD3D/Private/D3DShaderCompiler.cpp
@@ -961,6 +961,40 @@ static void ParseRayTracingEntryPoint(const FString& Input, FString& OutMain, FS
 	}
 }
 
+static bool DumpDebugShaderUSF(FString& PreprocessedShaderSource, const FShaderCompilerInput& Input)
+{
+	bool bDumpDebugInfo = false;
+
+	// Write out the preprocessed file and a batch file to compile it if requested (DumpDebugInfoPath is valid)
+	if (Input.DumpDebugInfoPath.Len() > 0 && IFileManager::Get().DirectoryExists(*Input.DumpDebugInfoPath))
+	{
+		bDumpDebugInfo = true;
+		FString Filename = Input.GetSourceFilename();
+		FArchive* FileWriter = IFileManager::Get().CreateFileWriter(*(Input.DumpDebugInfoPath / Filename));
+		if (FileWriter)
+		{
+			auto AnsiSourceFile = StringCast<ANSICHAR>(*PreprocessedShaderSource);
+			FileWriter->Serialize((ANSICHAR*)AnsiSourceFile.Get(), AnsiSourceFile.Length());
+			{
+				FString Line = CrossCompiler::CreateResourceTableFromEnvironment(Input.Environment);
+
+				Line += TEXT("#if 0 /*DIRECT COMPILE*/\n");
+				Line += CreateShaderCompilerWorkerDirectCommandLine(Input);
+				Line += TEXT("\n#endif /*DIRECT COMPILE*/\n");
+				Line += TEXT("//");
+				Line += Input.DebugDescription;
+				Line += TEXT("\n");
+				FileWriter->Serialize(TCHAR_TO_ANSI(*Line), Line.Len());
+			}
+			FileWriter->Close();
+			delete FileWriter;
+		}
+	}
+
+	return bDumpDebugInfo;
+}
+
+
 // Generate the dumped usf file; call the D3D compiler, gather reflection information and generate the output data
 static bool CompileAndProcessD3DShader(FString& PreprocessedShaderSource, const FString& CompilerPath,
 	uint32 CompileFlags, const FShaderCompilerInput& Input, FString& EntryPointName,
@@ -1000,33 +1034,12 @@ static bool CompileAndProcessD3DShader(FString& PreprocessedShaderSource, const
 		}
 	}
 
-	bool bDumpDebugInfo = false;
 	// Write out the preprocessed file and a batch file to compile it if requested (DumpDebugInfoPath is valid)
-	if (Input.DumpDebugInfoPath.Len() > 0 && IFileManager::Get().DirectoryExists(*Input.DumpDebugInfoPath))
+	bool bDumpDebugInfo = DumpDebugShaderUSF(PreprocessedShaderSource, Input);
+	if (bDumpDebugInfo)
 	{
-		bDumpDebugInfo = true;
-		FString Filename = Input.GetSourceFilename();
-		FArchive* FileWriter = IFileManager::Get().CreateFileWriter(*(Input.DumpDebugInfoPath / Filename));
-		if (FileWriter)
-		{
-			FileWriter->Serialize((ANSICHAR*)AnsiSourceFile.Get(), AnsiSourceFile.Length());
-			{
-				FString Line = CrossCompiler::CreateResourceTableFromEnvironment(Input.Environment);
-
-				Line += TEXT("#if 0 /*DIRECT COMPILE*/\n");
-				Line += CreateShaderCompilerWorkerDirectCommandLine(Input);
-				Line += TEXT("\n#endif /*DIRECT COMPILE*/\n");
-				Line += TEXT("//");
-				Line += Input.DebugDescription;
-				Line += TEXT("\n");
-				FileWriter->Serialize(TCHAR_TO_ANSI(*Line), Line.Len());
-			}
-			FileWriter->Close();
-			delete FileWriter;
-		}
-
 		FString BatchFileContents;
-
+		FString Filename = Input.GetSourceFilename();
 		if (bUseDXC)
 		{
 			BatchFileContents = D3DCreateDXCCompileBatchFile(Filename, *EntryPointName, *RayTracingExports, ShaderProfile, CompileFlags, Output, AutoBindingSpace);
@@ -1715,7 +1728,8 @@ void CompileD3DShader(const FShaderCompilerInput& Input,FShaderCompilerOutput& O
 		TArray<FString> Errors;
 		if (!RemoveUnusedOutputs(PreprocessedShaderSource, UsedOutputs, Exceptions, EntryPointName, Errors))
 		{
-			UE_LOG(LogD3D11ShaderCompiler, Warning, TEXT("Failed to Remove unused outputs [%s]!"), *Input.DumpDebugInfoPath);
+			DumpDebugShaderUSF(PreprocessedShaderSource, Input);
+			UE_LOG(LogD3D11ShaderCompiler, Fatal, TEXT("Failed to Remove unused outputs [%s]!"), *Input.DumpDebugInfoPath);
 			for (int32 Index = 0; Index < Errors.Num(); ++Index)
 			{
 				FShaderCompilerError NewError;
