UE4R - Back to Warning, not Fatal
#rnx

diff --git a/Engine/Source/Developer/Windows/ShaderFormatD3D/Private/D3DShaderCompiler.cpp b/Engine/Source/Developer/Windows/ShaderFormatD3D/Private/D3DShaderCompiler.cpp
--- a/Engine/Source/Developer/Windows/ShaderFormatD3D/Private/D3DShaderCompiler.cpp
+++ b/Engine/Source/Developer/Windows/ShaderFormatD3D/Private/D3DShaderCompiler.cpp
@@ -1729,7 +1729,7 @@ void CompileD3DShader(const FShaderCompilerInput& Input,FShaderCompilerOutput& O
 		if (!RemoveUnusedOutputs(PreprocessedShaderSource, UsedOutputs, Exceptions, EntryPointName, Errors))
 		{
 			DumpDebugShaderUSF(PreprocessedShaderSource, Input);
-			UE_LOG(LogD3D11ShaderCompiler, Fatal, TEXT("Failed to Remove unused outputs [%s]!"), *Input.DumpDebugInfoPath);
+			UE_LOG(LogD3D11ShaderCompiler, Warning, TEXT("Failed to Remove unused outputs [%s]!"), *Input.DumpDebugInfoPath);
 			for (int32 Index = 0; Index < Errors.Num(); ++Index)
 			{
 				FShaderCompilerError NewError;
