Enabled DXC multiline error messages in output log.

#rnx

diff --git a/Engine/Source/Developer/Windows/ShaderFormatD3D/Private/D3DShaderCompiler.cpp b/Engine/Source/Developer/Windows/ShaderFormatD3D/Private/D3DShaderCompiler.cpp
--- a/Engine/Source/Developer/Windows/ShaderFormatD3D/Private/D3DShaderCompiler.cpp
+++ b/Engine/Source/Developer/Windows/ShaderFormatD3D/Private/D3DShaderCompiler.cpp
@@ -88,23 +88,9 @@ static void D3D11FilterShaderCompileWarnings(const FString& CompileWarnings, TAr
 	}
 }
 
-static void DXCFilterShaderCompileWarnings(const FString& CompileWarnings, TArray<FString>& FilteredWarnings)
+FORCENOINLINE static void DXCFilterShaderCompileWarnings(const FString& CompileWarnings, TArray<FString>& FilteredWarnings)
 {
-	TArray<FString> WarningArray;
-	FString OutWarningString = TEXT("");
-	CompileWarnings.ParseIntoArray(WarningArray, TEXT("\n"), true);
-
-	//go through each warning line
-	for (int32 WarningIndex = 0; WarningIndex < WarningArray.Num(); WarningIndex++)
-	{
-		// DXC includes few lines of extra context in its diagnostic messages, which is different from FXC.
-		// We only include actual error and warning lines in the filtered output.
-		if (WarningArray[WarningIndex].Contains(TEXT("error:"))
-			|| WarningArray[WarningIndex].Contains(TEXT("warning:")))
-		{
-			FilteredWarnings.AddUnique(WarningArray[WarningIndex]);
-		}
-	}
+	CompileWarnings.ParseIntoArray(FilteredWarnings, TEXT("\n"), true);
 }
 
 static bool IsRayTracingShader(const FShaderTarget& Target)
@@ -1650,13 +1636,15 @@ static bool CompileAndProcessD3DShader(FString& PreprocessedShaderSource, const
 	return SUCCEEDED(Result);
 }
 
-void CompileD3DShader(const FShaderCompilerInput& Input,FShaderCompilerOutput& Output, FShaderCompilerDefinitions& AdditionalDefines, const FString& WorkingDirectory)
+void CompileD3DShader(const FShaderCompilerInput& Input, FShaderCompilerOutput& Output, FShaderCompilerDefinitions& AdditionalDefines, const FString& WorkingDirectory)
 {
 	FString PreprocessedShaderSource;
 	FString CompilerPath;
-	const bool bUseWaveOperations = Input.Environment.CompilerFlags.Contains(CFLAG_WaveOperations); // Forces shader model 6.0 for this shader
-	const bool bForceDXC = Input.Environment.CompilerFlags.Contains(CFLAG_ForceDXC);
-	const TCHAR* ShaderProfile = GetShaderProfileName(Input.Target, bUseWaveOperations || bForceDXC);
+	const bool bIsRayTracingShader = IsRayTracingShader(Input.Target);
+	const bool bUseDXC = bIsRayTracingShader
+		|| Input.Environment.CompilerFlags.Contains(CFLAG_WaveOperations)
+		|| Input.Environment.CompilerFlags.Contains(CFLAG_ForceDXC);
+	const TCHAR* ShaderProfile = GetShaderProfileName(Input.Target, bUseDXC);
 
 	if(!ShaderProfile)
 	{
@@ -1667,7 +1655,7 @@ void CompileD3DShader(const FShaderCompilerInput& Input,FShaderCompilerOutput& O
 	// Set additional defines.
 	AdditionalDefines.SetDefine(TEXT("COMPILER_HLSL"), 1);
 
-	if (bUseWaveOperations || bForceDXC)
+	if (bUseDXC)
 	{
 		AdditionalDefines.SetDefine(TEXT("PLATFORM_SUPPORTS_SM6_0_WAVE_OPERATIONS"), 1);
 	}
@@ -1797,23 +1785,73 @@ void CompileD3DShader(const FShaderCompilerInput& Input,FShaderCompilerOutput& O
 	{
 		const FString& CurrentError = FilteredErrors[ErrorIndex];
 		FShaderCompilerError NewError;
-		// Extract the filename and line number from the shader compiler error message for PC whose format is:
-		// "d:\UE4\Binaries\BasePassPixelShader(30,7): error X3000: invalid target or usage string"
-		int32 FirstParenIndex = CurrentError.Find(TEXT("("));
-		int32 LastParenIndex = CurrentError.Find(TEXT("):"));
-		if (FirstParenIndex != INDEX_NONE 
-			&& LastParenIndex != INDEX_NONE
-			&& LastParenIndex > FirstParenIndex)
+
+		if (bUseDXC)
 		{
-			NewError.ErrorVirtualFilePath = CurrentError.Left(FirstParenIndex);
-			NewError.ErrorLineString = CurrentError.Mid(FirstParenIndex + 1, LastParenIndex - FirstParenIndex - FCString::Strlen(TEXT("(")));
-			NewError.StrippedErrorMessage = CurrentError.Right(CurrentError.Len() - LastParenIndex - FCString::Strlen(TEXT("):")));
+			// Extract filename and line number from DXC output with format:
+			// "d:\UE4\Binaries\BasePassPixelShader:30:7: error: invalid target or usage string"
+			int32 ReportIndex = CurrentError.Find(TEXT(": error: "));
+			if (ReportIndex == INDEX_NONE)
+			{
+				ReportIndex = CurrentError.Find(TEXT(": warning: "));
+			}
+
+			int32 SecondColonIndex = CurrentError.Find(TEXT(":"), ESearchCase::IgnoreCase, ESearchDir::FromEnd, ReportIndex - 1);
+			int32 FirstColonIndex = CurrentError.Find(TEXT(":"), ESearchCase::IgnoreCase, ESearchDir::FromEnd, SecondColonIndex - 1);
+
+			if (ReportIndex != INDEX_NONE &&
+				FirstColonIndex != INDEX_NONE &&
+				SecondColonIndex != INDEX_NONE &&
+				FirstColonIndex < SecondColonIndex &&
+				SecondColonIndex < ReportIndex)
+			{
+				// Extract and store error message with source filename
+				NewError.ErrorVirtualFilePath = CurrentError.Left(FirstColonIndex);
+				NewError.ErrorLineString = CurrentError.Mid(FirstColonIndex + 1, ReportIndex - FirstColonIndex - FCString::Strlen(TEXT(":")));
+				NewError.StrippedErrorMessage = CurrentError.Right(CurrentError.Len() - ReportIndex - FCString::Strlen(TEXT(": ")));
+				Output.Errors.Add(NewError);
+			}
+			else if (Output.Errors.Num() > 0)
+			{
+				// Append highlighted line or marker to the previous error message
+				FShaderCompilerError& PrevError = Output.Errors.Last();
+				if (PrevError.HighlightedLine.IsEmpty())
+				{
+					PrevError.HighlightedLine = CurrentError;
+				}
+				else if (PrevError.HighlightedLineMarker.IsEmpty())
+				{
+					PrevError.HighlightedLineMarker = CurrentError;
+				}
+				else
+				{
+					// Register as new error message
+					NewError.StrippedErrorMessage = CurrentError;
+					Output.Errors.Add(NewError);
+				}
+			}
 		}
 		else
 		{
-			NewError.StrippedErrorMessage = CurrentError;
+			// Extract filename and line number from FXC output with format:
+			// "d:\UE4\Binaries\BasePassPixelShader(30,7): error X3000: invalid target or usage string"
+			int32 FirstParenIndex = CurrentError.Find(TEXT("("));
+			int32 LastParenIndex = CurrentError.Find(TEXT("):"));
+			if (FirstParenIndex != INDEX_NONE &&
+				LastParenIndex != INDEX_NONE &&
+				LastParenIndex > FirstParenIndex)
+			{
+				// Extract and store error message with source filename
+				NewError.ErrorVirtualFilePath = CurrentError.Left(FirstParenIndex);
+				NewError.ErrorLineString = CurrentError.Mid(FirstParenIndex + 1, LastParenIndex - FirstParenIndex - FCString::Strlen(TEXT("(")));
+				NewError.StrippedErrorMessage = CurrentError.Right(CurrentError.Len() - LastParenIndex - FCString::Strlen(TEXT("):")));
+			}
+			else
+			{
+				NewError.StrippedErrorMessage = CurrentError;
+			}
+			Output.Errors.Add(NewError);
 		}
-		Output.Errors.Add(NewError);
 	}
 
 	if (Input.ExtraSettings.bExtractShaderSource)
diff --git a/Engine/Source/Runtime/Engine/Private/ShaderCompiler/ShaderCompiler.cpp b/Engine/Source/Runtime/Engine/Private/ShaderCompiler/ShaderCompiler.cpp
--- a/Engine/Source/Runtime/Engine/Private/ShaderCompiler/ShaderCompiler.cpp
+++ b/Engine/Source/Runtime/Engine/Private/ShaderCompiler/ShaderCompiler.cpp
@@ -847,7 +847,21 @@ static bool CheckSingleJob(FShaderCompileJob* SingleJob, const TArray<FMaterial*
 	{
 		for (int32 ErrorIndex = 0; ErrorIndex < SingleJob->Output.Errors.Num(); ErrorIndex++)
 		{
-			Errors.AddUnique(SingleJob->Output.Errors[ErrorIndex].GetErrorString());
+			const FShaderCompilerError& InError = SingleJob->Output.Errors[ErrorIndex];
+			if (InError.HasLineMarker())
+			{
+				// Append highlighted line and its marker to the same error message with line terminators
+				// to get a similar multiline error output as with DXC
+				Errors.AddUnique(
+					InError.GetErrorString() + LINE_TERMINATOR +
+					TEXT("\t") + InError.HighlightedLine + LINE_TERMINATOR +
+					TEXT("\t") + InError.HighlightedLineMarker
+				);
+			}
+			else
+			{
+				Errors.AddUnique(InError.GetErrorString());
+			}
 		}
 	}
 
@@ -2103,7 +2117,7 @@ void FShaderCompilingManager::ProcessCompiledShaderMaps(
 								FString ErrorMessage = Errors[ErrorIndex];
 								// Work around build machine string matching heuristics that will cause a cook to fail
 								ErrorMessage.ReplaceInline(TEXT("error "), TEXT("err0r "), ESearchCase::CaseSensitive);
-								UE_LOG(LogShaderCompilers, Display, TEXT("	%s"), *ErrorMessage);
+								UE_LOG(LogShaderCompilers, Display, TEXT("%s"), *ErrorMessage);
 							}
 						}
 						else
diff --git a/Engine/Source/Runtime/RenderCore/Public/ShaderCore.h b/Engine/Source/Runtime/RenderCore/Public/ShaderCore.h
--- a/Engine/Source/Runtime/RenderCore/Public/ShaderCore.h
+++ b/Engine/Source/Runtime/RenderCore/Public/ShaderCore.h
@@ -740,24 +740,49 @@ struct FShaderCompilerInput
 struct FShaderCompilerError
 {
 	FShaderCompilerError(const TCHAR* InStrippedErrorMessage = TEXT(""))
-	:	ErrorVirtualFilePath(TEXT(""))
-	,	ErrorLineString(TEXT(""))
-	,	StrippedErrorMessage(InStrippedErrorMessage)
+		: ErrorVirtualFilePath(TEXT(""))
+		, ErrorLineString(TEXT(""))
+		, StrippedErrorMessage(InStrippedErrorMessage)
+		, HighlightedLine(TEXT(""))
+		, HighlightedLineMarker(TEXT(""))
 	{}
 
 	FShaderCompilerError(const TCHAR* InVirtualFilePath, const TCHAR* InLineString, const TCHAR* InStrippedErrorMessage)
 		: ErrorVirtualFilePath(InVirtualFilePath)
 		, ErrorLineString(InLineString)
 		, StrippedErrorMessage(InStrippedErrorMessage)
+		, HighlightedLine(TEXT(""))
+		, HighlightedLineMarker(TEXT(""))
 	{}
 
 	FString ErrorVirtualFilePath;
 	FString ErrorLineString;
 	FString StrippedErrorMessage;
+	FString HighlightedLine;
+	FString HighlightedLineMarker;
 
+	/** Returns the error message with source file and source line (if present). */
 	FString GetErrorString() const
 	{
-		return ErrorVirtualFilePath + TEXT("(") + ErrorLineString + TEXT("): ") + StrippedErrorMessage; // TODO
+		if (ErrorVirtualFilePath.IsEmpty())
+		{
+			return StrippedErrorMessage;
+		}
+		else
+		{
+			return ErrorVirtualFilePath + TEXT("(") + ErrorLineString + TEXT("): ") + StrippedErrorMessage;
+		}
+	}
+
+	/**
+	Returns true if this error message has a marker string for the highlighted source line where the error occurred. Example:
+		/Engine/Private/MySourceFile.usf(120): error: undeclared identifier 'a'
+		float b = a;
+				  ^
+	*/
+	bool HasLineMarker() const
+	{
+		return !HighlightedLine.IsEmpty() && !HighlightedLineMarker.IsEmpty();
 	}
 
 	/** Returns the path of the underlying source file relative to the process base dir. */
