Live Coding: Fix process group names being derived from relative paths, resulting in two workspaces using the same project sharing the same name (causing changes in one project to show up in the other's console window)

#rnx

diff --git a/Engine/Source/Developer/Windows/LiveCoding/Private/LiveCodingModule.cpp b/Engine/Source/Developer/Windows/LiveCoding/Private/LiveCodingModule.cpp
--- a/Engine/Source/Developer/Windows/LiveCoding/Private/LiveCodingModule.cpp
+++ b/Engine/Source/Developer/Windows/LiveCoding/Private/LiveCodingModule.cpp
@@ -303,7 +303,8 @@ bool FLiveCodingModule::StartLiveCoding()
 		LppUseExternalBuildSystem();
 
 		// Enable the server
-		FString ProcessGroup = FString::Printf(TEXT("UE4_%s_0x%08x"), FApp::GetProjectName(), GetTypeHash(FPaths::ProjectDir()));
+		FString ProjectPath = FPaths::ConvertRelativePathToFull(FPaths::ProjectDir()).ToLower();
+		FString ProcessGroup = FString::Printf(TEXT("UE4_%s_0x%08x"), FApp::GetProjectName(), GetTypeHash(ProjectPath));
 		LppRegisterProcessGroup(TCHAR_TO_ANSI(*ProcessGroup));
 
 		// Build the command line
