Fix static analysis warning.

#rnx

diff --git a/Engine/Source/Developer/Windows/LiveCoding/Private/External/LC_Process.cpp b/Engine/Source/Developer/Windows/LiveCoding/Private/External/LC_Process.cpp
--- a/Engine/Source/Developer/Windows/LiveCoding/Private/External/LC_Process.cpp
+++ b/Engine/Source/Developer/Windows/LiveCoding/Private/External/LC_Process.cpp
@@ -641,7 +641,10 @@ namespace process
 		Environment* environment = new Environment;
 		environment->size = environmentData.size();
 		environment->data = ::malloc(environmentData.size() * sizeof(wchar_t));
-		memcpy(environment->data, environmentData.data(), environmentData.size() * sizeof(wchar_t));
+		if (environment->data != nullptr)
+		{
+			memcpy(environment->data, environmentData.data(), environmentData.size() * sizeof(wchar_t));
+		}
 		return environment;
 	}
 	// END EPIC MOD
