Fixed: Reflection capture related crash running DirectX mobile emulation (Contributed by [REDACTED_PERSON])
PR #6585

diff --git a/Engine/Source/Developer/Windows/WindowsTargetPlatform/Private/GenericWindowsTargetPlatform.h b/Engine/Source/Developer/Windows/WindowsTargetPlatform/Private/GenericWindowsTargetPlatform.h
--- a/Engine/Source/Developer/Windows/WindowsTargetPlatform/Private/GenericWindowsTargetPlatform.h
+++ b/Engine/Source/Developer/Windows/WindowsTargetPlatform/Private/GenericWindowsTargetPlatform.h
@@ -83,11 +83,13 @@ public:
 			}
 		}
 
-		// If we are targeting ES 2.0/3.1, we also must cook encoded HDR reflection captures
+		// If we are targeting ES3.1, we also must cook encoded HDR reflection captures
 		static FName NAME_SF_VULKAN_ES31(TEXT("SF_VULKAN_ES31"));
 		static FName NAME_OPENGL_150_ES3_1(TEXT("GLSL_150_ES31"));
+		static FName NAME_PCD3D_ES3_1(TEXT("PCD3D_ES31"));
 		bRequiresEncodedHDRReflectionCaptures =	TargetedShaderFormats.Contains(NAME_SF_VULKAN_ES31)
-												 || TargetedShaderFormats.Contains(NAME_OPENGL_150_ES3_1);
+												|| TargetedShaderFormats.Contains(NAME_OPENGL_150_ES3_1)
+												|| TargetedShaderFormats.Contains(NAME_PCD3D_ES3_1);
 	#endif
 	}
 
