Fix some non-unity issues building LiveCodingConsole.

#rnx

diff --git a/Engine/Source/Developer/Windows/LiveCodingServer/Private/External/LC_COMThread.h b/Engine/Source/Developer/Windows/LiveCodingServer/Private/External/LC_COMThread.h
--- a/Engine/Source/Developer/Windows/LiveCodingServer/Private/External/LC_COMThread.h
+++ b/Engine/Source/Developer/Windows/LiveCodingServer/Private/External/LC_COMThread.h
@@ -2,6 +2,7 @@
 
 #pragma once
 
+#include <type_traits>
 #include "LC_Event.h"
 #include "LC_Thread.h"
 
diff --git a/Engine/Source/Developer/Windows/LiveCodingServer/Private/External/LC_ImmutableString.h b/Engine/Source/Developer/Windows/LiveCodingServer/Private/External/LC_ImmutableString.h
--- a/Engine/Source/Developer/Windows/LiveCodingServer/Private/External/LC_ImmutableString.h
+++ b/Engine/Source/Developer/Windows/LiveCodingServer/Private/External/LC_ImmutableString.h
@@ -3,6 +3,7 @@
 #pragma once
 
 #include "CoreTypes.h"
+#include "LC_Platform.h"
 #include "LC_Memory.h"
 #include <string>
 
diff --git a/Engine/Source/Developer/Windows/LiveCodingServer/Private/External/LC_TimeStamp.h b/Engine/Source/Developer/Windows/LiveCodingServer/Private/External/LC_TimeStamp.h
--- a/Engine/Source/Developer/Windows/LiveCodingServer/Private/External/LC_TimeStamp.h
+++ b/Engine/Source/Developer/Windows/LiveCodingServer/Private/External/LC_TimeStamp.h
@@ -2,6 +2,7 @@
 
 #pragma once
 
+#include "LC_Platform.h"
 
 namespace timeStamp
 {
