Remove shader compilation flag CFLAG_SkipOptimizationsDXC

- this was added to aid DXC bring up and is no longer needed
  - there are no active references

#rnx

diff --git a/Engine/Source/Developer/Windows/ShaderFormatD3D/Private/D3DShaderCompiler.cpp b/Engine/Source/Developer/Windows/ShaderFormatD3D/Private/D3DShaderCompiler.cpp
--- a/Engine/Source/Developer/Windows/ShaderFormatD3D/Private/D3DShaderCompiler.cpp
+++ b/Engine/Source/Developer/Windows/ShaderFormatD3D/Private/D3DShaderCompiler.cpp
@@ -1058,11 +1058,6 @@ static bool CompileAndProcessD3DShader(FString& PreprocessedShaderSource, const
 		FString Filename = Input.GetSourceFilename();
 		if (bUseDXC)
 		{
-			if (Input.Environment.CompilerFlags.Contains(CFLAG_SkipOptimizationsDXC))
-			{
-				CompileFlags |= D3D10_SHADER_SKIP_OPTIMIZATION;
-			}
-
 			BatchFileContents = D3DCreateDXCCompileBatchFile(Filename, *EntryPointName, *RayTracingExports, ShaderProfile, CompileFlags, Output, AutoBindingSpace);
 		}
 		else
@@ -1104,11 +1099,6 @@ static bool CompileAndProcessD3DShader(FString& PreprocessedShaderSource, const
 		// #dxr_todo: this flag should not be even passed into this function from the higher level.
 		uint32 DXCFlags = CompileFlags & (~D3D10_SHADER_ENABLE_BACKWARDS_COMPATIBILITY);
 
-		if (Input.Environment.CompilerFlags.Contains(CFLAG_SkipOptimizationsDXC))
-		{
-			DXCFlags |= D3D10_SHADER_SKIP_OPTIMIZATION;
-		}
-
 		D3DCreateDXCArguments(Args, *RayTracingExports, DXCFlags, Output, AutoBindingSpace);
 
 		TRefCountPtr<IDxcBlobEncoding> DxcErrorBlob;
diff --git a/Engine/Source/Runtime/RenderCore/Public/ShaderCore.h b/Engine/Source/Runtime/RenderCore/Public/ShaderCore.h
--- a/Engine/Source/Runtime/RenderCore/Public/ShaderCore.h
+++ b/Engine/Source/Runtime/RenderCore/Public/ShaderCore.h
@@ -205,8 +205,6 @@ enum ECompilerFlags
 	// Use DirectX Shader Compiler (DXC) to compile all shaders, intended for compatibility testing.
 	CFLAG_ForceDXC,
 	CFLAG_SkipOptimizations,
-	// Temporarily disable optimizations with DXC compiler only, intended to workaround shader compiler bugs until they can be resolved with 1st party
-	CFLAG_SkipOptimizationsDXC
 };
 
 enum class EShaderParameterType : uint8
