UE4.25 - Disambiguates crashes from inside a platform shader compiler
#rnx

diff --git a/Engine/Source/Developer/ShaderCompilerCommon/Private/ShaderCompilerCommon.cpp b/Engine/Source/Developer/ShaderCompilerCommon/Private/ShaderCompilerCommon.cpp
--- a/Engine/Source/Developer/ShaderCompilerCommon/Private/ShaderCompilerCommon.cpp
+++ b/Engine/Source/Developer/ShaderCompilerCommon/Private/ShaderCompilerCommon.cpp
@@ -9,6 +9,8 @@
 IMPLEMENT_MODULE(FDefaultModuleImpl, ShaderCompilerCommon);
 
 
+ESCWErrorCode GSCWErrorCode = ESCWErrorCode::NotSet;
+
 int16 GetNumUniformBuffersUsed(const FShaderCompilerResourceTable& InSRT)
 {
 	auto CountLambda = [&](const TArray<uint32>& In)
diff --git a/Engine/Source/Developer/ShaderCompilerCommon/Public/ShaderCompilerCommon.h b/Engine/Source/Developer/ShaderCompilerCommon/Public/ShaderCompilerCommon.h
--- a/Engine/Source/Developer/ShaderCompilerCommon/Public/ShaderCompilerCommon.h
+++ b/Engine/Source/Developer/ShaderCompilerCommon/Public/ShaderCompilerCommon.h
@@ -356,3 +356,22 @@ namespace CrossCompiler
 	}
 
 }
+
+enum class ESCWErrorCode
+{
+	NotSet = -1,
+	Success,
+	GeneralCrash,
+	BadShaderFormatVersion,
+	BadInputVersion,
+	BadSingleJobHeader,
+	BadPipelineJobHeader,
+	CantDeleteInputFile,
+	CantSaveOutputFile,
+	NoTargetShaderFormatsFound,
+	CantCompileForSpecificFormat,
+	CrashInsidePlatformCompiler,
+};
+
+// Error code for SCW to help track down crashes
+extern SHADERCOMPILERCOMMON_API ESCWErrorCode GSCWErrorCode;
diff --git a/Engine/Source/Developer/Windows/ShaderFormatD3D/Private/D3DShaderCompiler.cpp b/Engine/Source/Developer/Windows/ShaderFormatD3D/Private/D3DShaderCompiler.cpp
--- a/Engine/Source/Developer/Windows/ShaderFormatD3D/Private/D3DShaderCompiler.cpp
+++ b/Engine/Source/Developer/Windows/ShaderFormatD3D/Private/D3DShaderCompiler.cpp
@@ -459,6 +459,7 @@ static HRESULT D3DCompileWrapper(
 #if !PLATFORM_SEH_EXCEPTIONS_DISABLED
 	__except( EXCEPTION_EXECUTE_HANDLER )
 	{
+		GSCWErrorCode = ESCWErrorCode::CrashInsidePlatformCompiler;
 		bException = true;
 		return E_FAIL;
 	}
@@ -497,6 +498,42 @@ static FString DxcBlobEncodingToFString(TRefCountPtr<IDxcBlobEncoding> DxcBlob)
 	return OutString;
 }
 
+
+static HRESULT DXCCompileWrapper(
+	TRefCountPtr<IDxcCompiler> Compiler,
+	TRefCountPtr<IDxcBlobEncoding>& TextBlob,
+	LPCWSTR EntryPoint,
+	LPCWSTR TargetProfile,
+	LPCWSTR* Arguments,
+	uint32 NumArguments,
+	TRefCountPtr<IDxcOperationResult>& OutCompileResult)
+{
+#if !PLATFORM_SEH_EXCEPTIONS_DISABLED
+	__try
+#endif
+	{
+		return Compiler->Compile(
+			TextBlob,							// source text to compile
+			nullptr,							// optional file name for pSource. Used in errors and include handlers.
+			EntryPoint,							// entry point name
+			TargetProfile,						// shader profile to compile
+			Arguments,							// array of pointers to arguments
+			NumArguments,						// number of arguments
+			nullptr,							// array of defines
+			0,									// number of defines
+			nullptr,							// user-provided interface to handle #include directives (optional)
+			OutCompileResult.GetInitReference()	// compiler output status, buffer, and errors
+		);
+	}
+#if !PLATFORM_SEH_EXCEPTIONS_DISABLED
+	__except (EXCEPTION_EXECUTE_HANDLER)
+	{
+		GSCWErrorCode = ESCWErrorCode::CrashInsidePlatformCompiler;
+		return E_FAIL;
+	}
+#endif
+}
+
 static HRESULT D3DCompileToDxil(const char* SourceText, LPCWSTR EntryPoint, LPCWSTR TargetProfile, LPCWSTR* Arguments, uint32 NumArguments, 
 	const FString& DisasmFilename, TRefCountPtr<ID3DBlob>& OutDxilBlob, TRefCountPtr<IDxcBlobEncoding>& OutErrorBlob)
 {
@@ -512,19 +549,7 @@ static HRESULT D3DCompileToDxil(const char* SourceText, LPCWSTR EntryPoint, LPCW
 	VERIFYHRESULT(Library->CreateBlobWithEncodingFromPinned((LPBYTE)SourceText, FCStringAnsi::Strlen(SourceText), CP_UTF8, TextBlob.GetInitReference()));
 
 	TRefCountPtr<IDxcOperationResult> CompileResult;
-
-	VERIFYHRESULT(Compiler->Compile(
-		TextBlob,							// source text to compile
-		nullptr,							// optional file name for pSource. Used in errors and include handlers.
-		EntryPoint,							// entry point name
-		TargetProfile,						// shader profile to compile
-		Arguments,							// array of pointers to arguments
-		NumArguments,						// number of arguments
-		nullptr,							// array of defines
-		0,									// number of defines
-		nullptr,							// user-provided interface to handle #include directives (optional)
-		CompileResult.GetInitReference()	// compiler output status, buffer, and errors
-	));
+	VERIFYHRESULT(DXCCompileWrapper(Compiler, TextBlob, EntryPoint, TargetProfile, Arguments, NumArguments, CompileResult));
 
 	HRESULT CompileResultCode;
 	CompileResult->GetStatus(&CompileResultCode);
diff --git a/Engine/Source/Programs/ShaderCompileWorker/Private/ShaderCompileWorker.cpp b/Engine/Source/Programs/ShaderCompileWorker/Private/ShaderCompileWorker.cpp
--- a/Engine/Source/Programs/ShaderCompileWorker/Private/ShaderCompileWorker.cpp
+++ b/Engine/Source/Programs/ShaderCompileWorker/Private/ShaderCompileWorker.cpp
@@ -12,23 +12,10 @@
 #include "Interfaces/IShaderFormatModule.h"
 #include "Interfaces/ITargetPlatformManagerModule.h"
 #include "RHIShaderFormatDefinitions.inl"
+#include "ShaderCompilerCommon.h"
 
 #define DEBUG_USING_CONSOLE	0
 
-enum class ESCWErrorCode
-{
-	Success,
-	GeneralCrash,
-	BadShaderFormatVersion,
-	BadInputVersion,
-	BadSingleJobHeader,
-	BadPipelineJobHeader,
-	CantDeleteInputFile,
-	CantSaveOutputFile,
-	NoTargetShaderFormatsFound,
-	CantCompileForSpecificFormat,
-};
-
 static double LastCompileTime = 0.0;
 static int32 GNumProcessedJobs = 0;
 
@@ -1017,8 +1004,16 @@ static int32 GuardedMainWrapper(int32 ArgC, TCHAR* ArgV[], const TCHAR* CrashOut
 
 			if (GFailedErrorCode == ESCWErrorCode::Success)
 			{
-				// Something else failed before we could set the error code, so mark it as a General Crash
-				GFailedErrorCode = ESCWErrorCode::GeneralCrash;
+				if (GSCWErrorCode != ESCWErrorCode::NotSet)
+				{
+					// Use the value set inside the shader format
+					GFailedErrorCode = GSCWErrorCode;
+				}
+				else
+				{
+					// Something else failed before we could set the error code, so mark it as a General Crash
+					GFailedErrorCode = ESCWErrorCode::GeneralCrash;
+				}
 			}
 			int64 FileSizePosition = WriteOutputFileHeader(OutputFile, (int32)GFailedErrorCode, FCString::Strlen(GErrorHist), GErrorHist,
 				FCString::Strlen(GErrorExceptionDescription), GErrorExceptionDescription);
diff --git a/Engine/Source/Programs/ShaderCompileWorker/ShaderCompileWorker.Build.cs b/Engine/Source/Programs/ShaderCompileWorker/ShaderCompileWorker.Build.cs
--- a/Engine/Source/Programs/ShaderCompileWorker/ShaderCompileWorker.Build.cs
+++ b/Engine/Source/Programs/ShaderCompileWorker/ShaderCompileWorker.Build.cs
@@ -15,7 +15,8 @@ public class ShaderCompileWorker : ModuleRules
 				"SandboxFile",
 				"TargetPlatform",
 				"ApplicationCore",
-				"TraceLog"
+				"TraceLog",
+				"ShaderCompilerCommon"
 			});
 
 		if (Target.Platform == UnrealTargetPlatform.Linux)
diff --git a/Engine/Source/Runtime/Engine/Private/ShaderCompiler/ShaderCompiler.cpp b/Engine/Source/Runtime/Engine/Private/ShaderCompiler/ShaderCompiler.cpp
--- a/Engine/Source/Runtime/Engine/Private/ShaderCompiler/ShaderCompiler.cpp
+++ b/Engine/Source/Runtime/Engine/Private/ShaderCompiler/ShaderCompiler.cpp
@@ -345,20 +345,6 @@ namespace ShaderCompilerCookStats
 // Make functions so the crash reporter can disambiguate the actual error because of the different callstacks
 namespace SCWErrorCode
 {
-	enum EErrors
-	{
-		Success,
-		GeneralCrash,
-		BadShaderFormatVersion,
-		BadInputVersion,
-		BadSingleJobHeader,
-		BadPipelineJobHeader,
-		CantDeleteInputFile,
-		CantSaveOutputFile,
-		NoTargetShaderFormatsFound,
-		CantCompileForSpecificFormat,
-	};
-
 	void HandleGeneralCrash(const TCHAR* ExceptionInfo, const TCHAR* Callstack)
 	{
 		GLog->PanicFlushThreadedLogs();
@@ -414,6 +400,11 @@ namespace SCWErrorCode
 	{
 		ModalErrorOrLog(FString::Printf(TEXT("Output file corrupted (expected %I64d bytes, but only got %I64d): %s"), ExpectedSize, ActualSize, Filename));
 	}
+
+	void HandleCrashInsidePlatformCompiler(const TCHAR* Filename, int64 ExpectedSize, int64 ActualSize)
+	{
+		ModalErrorOrLog(FString::Printf(TEXT("Crash inside the platform compiler!")));
+	}
 }
 
 static inline void GetFormatVersionMap(TMap<FString, uint32>& OutFormatVersionMap)
@@ -816,6 +807,9 @@ static void HandleWorkerCrash(const TArray<TSharedRef<FShaderCommonCompileJob, E
 	case SCWErrorCode::CantCompileForSpecificFormat:
 		SCWErrorCode::HandleCantCompileForSpecificFormat(ExceptionInfo.GetData());
 		break;
+	case SCWErrorCode::CrashInsidePlatformCompiler:
+		SCWErrorCode::HandleCrashInsidePlatformCompiler(ExceptionInfo.GetData());
+		break;
 	case SCWErrorCode::Success:
 		// Can't get here...
 		break;
