Tweak D3d shader compiler to cleanly report error with too many UAV's rather than issuing a fatal error

diff --git a/Engine/Source/Developer/Windows/ShaderFormatD3D/Private/D3DShaderCompiler.cpp b/Engine/Source/Developer/Windows/ShaderFormatD3D/Private/D3DShaderCompiler.cpp
--- a/Engine/Source/Developer/Windows/ShaderFormatD3D/Private/D3DShaderCompiler.cpp
+++ b/Engine/Source/Developer/Windows/ShaderFormatD3D/Private/D3DShaderCompiler.cpp
@@ -1643,7 +1643,9 @@ static bool CompileAndProcessD3DShader(FString& PreprocessedShaderSource, const
 			// Check for resource limits for feature level 11.0
 			if (NumUAVs > GD3DMaximumNumUAVs)
 			{
-				UE_LOG(LogD3D11ShaderCompiler, Fatal, TEXT("Number of UAVs in \"%s\" exceeded limit: %d slots used, but limit is %d due to maximum feature level 11.0"), *Input.VirtualSourceFilePath, NumUAVs, GD3DMaximumNumUAVs);
+				FilteredErrors.Add(FString::Printf(TEXT("Number of UAVs in \"%s\" exceeded limit: %d slots used, but limit is %d due to maximum feature level 11.0"), *Input.VirtualSourceFilePath, NumUAVs, GD3DMaximumNumUAVs));
+				Result = E_FAIL;
+				Output.bSucceeded = false;
 			}
 
 			// Set the number of instructions.
