S - Change DXC unroll issue from error to warning
#rnx

diff --git a/Engine/Source/Developer/Windows/ShaderFormatD3D/Private/D3DShaderCompiler.inl b/Engine/Source/Developer/Windows/ShaderFormatD3D/Private/D3DShaderCompiler.inl
--- a/Engine/Source/Developer/Windows/ShaderFormatD3D/Private/D3DShaderCompiler.inl
+++ b/Engine/Source/Developer/Windows/ShaderFormatD3D/Private/D3DShaderCompiler.inl
@@ -119,6 +119,9 @@ public:
 
 		checkf(D3DCompileFlags == 0, TEXT("Unhandled shader compiler flags 0x%x!"), D3DCompileFlags);
 
+		ExtraArguments.Add(L"-HV");
+		ExtraArguments.Add(L"2016");
+
 		ExtraArguments.Add(L"/Zss");
 		ExtraArguments.Add(L"/Qembed_debug");
 		ExtraArguments.Add(L"/Zi");
