Updating "D3D10" type usages to the version agnostic "D3D"

diff --git a/Engine/Source/Developer/StandaloneRenderer/Private/Windows/D3D/SlateD3DShaders.cpp b/Engine/Source/Developer/StandaloneRenderer/Private/Windows/D3D/SlateD3DShaders.cpp
--- a/Engine/Source/Developer/StandaloneRenderer/Private/Windows/D3D/SlateD3DShaders.cpp
+++ b/Engine/Source/Developer/StandaloneRenderer/Private/Windows/D3D/SlateD3DShaders.cpp
@@ -76,7 +76,7 @@ static bool CompileShader( const FString& Filename, const FString& EntryPoint, c
 {
 	uint32 ShaderFlags = D3DCOMPILE_ENABLE_STRICTNESS;
 #if UE_BUILD_DEBUG
-	ShaderFlags |= D3D10_SHADER_DEBUG;
+	ShaderFlags |= D3DCOMPILE_DEBUG;
 #else
 	ShaderFlags |= D3DCOMPILE_OPTIMIZATION_LEVEL3;
 #endif
@@ -444,9 +444,10 @@ void FSlateDefaultPS::SetDrawEffects( ESlateDrawEffect InDrawEffects )
 	PerElementConstants.GetBufferData().DisableEffect = (uint32)(InDrawEffects & ESlateDrawEffect::DisabledEffect);
 }
 
-void FSlateDefaultPS::SetShaderParams( const FVector4& InShaderParams )
+void FSlateDefaultPS::SetShaderParams(const FShaderParams& InShaderParams)
 {
-	PerElementConstants.GetBufferData().ShaderParams = InShaderParams;
+	PerElementConstants.GetBufferData().ShaderParams = InShaderParams.PixelParams;
+	PerElementConstants.GetBufferData().ShaderParams2 = InShaderParams.PixelParams2;
 }
 
 void FSlateDefaultPS::SetGammaValues(const FVector2D& InGammaValues)
diff --git a/Engine/Source/Developer/Windows/ShaderFormatD3D/Private/D3DShaderCompiler.cpp b/Engine/Source/Developer/Windows/ShaderFormatD3D/Private/D3DShaderCompiler.cpp
--- a/Engine/Source/Developer/Windows/ShaderFormatD3D/Private/D3DShaderCompiler.cpp
+++ b/Engine/Source/Developer/Windows/ShaderFormatD3D/Private/D3DShaderCompiler.cpp
@@ -19,7 +19,7 @@ DEFINE_LOG_CATEGORY_STATIC(LogD3D11ShaderCompiler, Log, All);
 #define D3D_OVERLOADS 1
 
 // D3D doesn't define a mask for this, so we do so here
-#define SHADER_OPTIMIZATION_LEVEL_MASK (D3D10_SHADER_OPTIMIZATION_LEVEL0 | D3D10_SHADER_OPTIMIZATION_LEVEL1 | D3D10_SHADER_OPTIMIZATION_LEVEL2 | D3D10_SHADER_OPTIMIZATION_LEVEL3)
+#define SHADER_OPTIMIZATION_LEVEL_MASK (D3DCOMPILE_OPTIMIZATION_LEVEL0 | D3DCOMPILE_OPTIMIZATION_LEVEL1 | D3DCOMPILE_OPTIMIZATION_LEVEL2 | D3DCOMPILE_OPTIMIZATION_LEVEL3)
 
 // Disable macro redefinition warning for compatibility with Windows SDK 8+
 #pragma warning(push)
@@ -49,11 +49,10 @@ static int32 GD3DDumpAMDCodeXLFile = 0;
  */
 static uint32 TranslateCompilerFlagD3D11(ECompilerFlags CompilerFlag)
 {
-	// @TODO - currently d3d11 uses d3d10 shader compiler flags... update when this changes in DXSDK
 	switch(CompilerFlag)
 	{
-	case CFLAG_PreferFlowControl: return D3D10_SHADER_PREFER_FLOW_CONTROL;
-	case CFLAG_AvoidFlowControl: return D3D10_SHADER_AVOID_FLOW_CONTROL;
+	case CFLAG_PreferFlowControl: return D3DCOMPILE_PREFER_FLOW_CONTROL;
+	case CFLAG_AvoidFlowControl: return D3DCOMPILE_AVOID_FLOW_CONTROL;
 	default: return 0;
 	};
 }
@@ -167,73 +166,72 @@ static FString D3D11CreateShaderCompileCommandLine(
 	// add the entry point reference
 	FXCCommandline += FString(TEXT(" /E ")) + EntryFunction;
 
-	// @TODO - currently d3d11 uses d3d10 shader compiler flags... update when this changes in DXSDK
 	// go through and add other switches
-	if(CompileFlags & D3D10_SHADER_PREFER_FLOW_CONTROL)
+	if(CompileFlags & D3DCOMPILE_PREFER_FLOW_CONTROL)
 	{
-		CompileFlags &= ~D3D10_SHADER_PREFER_FLOW_CONTROL;
+		CompileFlags &= ~D3DCOMPILE_PREFER_FLOW_CONTROL;
 		FXCCommandline += FString(TEXT(" /Gfp"));
 	}
 
-	if(CompileFlags & D3D10_SHADER_DEBUG)
+	if(CompileFlags & D3DCOMPILE_DEBUG)
 	{
-		CompileFlags &= ~D3D10_SHADER_DEBUG;
+		CompileFlags &= ~D3DCOMPILE_DEBUG;
 		FXCCommandline += FString(TEXT(" /Zi"));
 	}
 
-	if(CompileFlags & D3D10_SHADER_SKIP_OPTIMIZATION)
+	if(CompileFlags & D3DCOMPILE_SKIP_OPTIMIZATION)
 	{
-		CompileFlags &= ~D3D10_SHADER_SKIP_OPTIMIZATION;
+		CompileFlags &= ~D3DCOMPILE_SKIP_OPTIMIZATION;
 		FXCCommandline += FString(TEXT(" /Od"));
 	}
 
-	if (CompileFlags & D3D10_SHADER_SKIP_VALIDATION)
+	if (CompileFlags & D3DCOMPILE_SKIP_VALIDATION)
 	{
-		CompileFlags &= ~D3D10_SHADER_SKIP_VALIDATION;
+		CompileFlags &= ~D3DCOMPILE_SKIP_VALIDATION;
 		FXCCommandline += FString(TEXT(" /Vd"));
 	}
 
-	if(CompileFlags & D3D10_SHADER_AVOID_FLOW_CONTROL)
+	if(CompileFlags & D3DCOMPILE_AVOID_FLOW_CONTROL)
 	{
-		CompileFlags &= ~D3D10_SHADER_AVOID_FLOW_CONTROL;
+		CompileFlags &= ~D3DCOMPILE_AVOID_FLOW_CONTROL;
 		FXCCommandline += FString(TEXT(" /Gfa"));
 	}
 
-	if(CompileFlags & D3D10_SHADER_PACK_MATRIX_ROW_MAJOR)
+	if(CompileFlags & D3DCOMPILE_PACK_MATRIX_ROW_MAJOR)
 	{
-		CompileFlags &= ~D3D10_SHADER_PACK_MATRIX_ROW_MAJOR;
+		CompileFlags &= ~D3DCOMPILE_PACK_MATRIX_ROW_MAJOR;
 		FXCCommandline += FString(TEXT(" /Zpr"));
 	}
 
-	if(CompileFlags & D3D10_SHADER_ENABLE_BACKWARDS_COMPATIBILITY)
+	if(CompileFlags & D3DCOMPILE_ENABLE_BACKWARDS_COMPATIBILITY)
 	{
-		CompileFlags &= ~D3D10_SHADER_ENABLE_BACKWARDS_COMPATIBILITY;
+		CompileFlags &= ~D3DCOMPILE_ENABLE_BACKWARDS_COMPATIBILITY;
 		FXCCommandline += FString(TEXT(" /Gec"));
 	}
 
 	switch (CompileFlags & SHADER_OPTIMIZATION_LEVEL_MASK)
 	{
-		case D3D10_SHADER_OPTIMIZATION_LEVEL2:
-		CompileFlags &= ~D3D10_SHADER_OPTIMIZATION_LEVEL2;
+		case D3DCOMPILE_OPTIMIZATION_LEVEL2:
+		CompileFlags &= ~D3DCOMPILE_OPTIMIZATION_LEVEL2;
 		FXCCommandline += FString(TEXT(" /O2"));
 			break;
 
-		case D3D10_SHADER_OPTIMIZATION_LEVEL3:
-		CompileFlags &= ~D3D10_SHADER_OPTIMIZATION_LEVEL3;
+		case D3DCOMPILE_OPTIMIZATION_LEVEL3:
+		CompileFlags &= ~D3DCOMPILE_OPTIMIZATION_LEVEL3;
 		FXCCommandline += FString(TEXT(" /O3"));
 			break;
 
-		case D3D10_SHADER_OPTIMIZATION_LEVEL1:
-		CompileFlags &= ~D3D10_SHADER_OPTIMIZATION_LEVEL1;
+		case D3DCOMPILE_OPTIMIZATION_LEVEL1:
+		CompileFlags &= ~D3DCOMPILE_OPTIMIZATION_LEVEL1;
 		FXCCommandline += FString(TEXT(" /O1"));
 			break;
 
-		case D3D10_SHADER_OPTIMIZATION_LEVEL0:
-			CompileFlags &= ~D3D10_SHADER_OPTIMIZATION_LEVEL0;
+		case D3DCOMPILE_OPTIMIZATION_LEVEL0:
+			CompileFlags &= ~D3DCOMPILE_OPTIMIZATION_LEVEL0;
 			break;
 
 		default:
-			Output.Errors.Emplace(TEXT("Unknown D3D10 optimization level"));
+			Output.Errors.Emplace(TEXT("Unknown D3DCOMPILE optimization level"));
 			break;
 	}
 
@@ -956,26 +954,25 @@ void CompileD3DShader(const FShaderCompilerInput& Input, FShaderCompilerOutput&
 
 	RemoveUniformBuffersFromSource(Input.Environment, PreprocessedShaderSource);
 
-	// @TODO - currently d3d11 uses d3d10 shader compiler flags... update when this changes in DXSDK
 	// @TODO - implement different material path to allow us to remove backwards compat flag on sm5 shaders
-	uint32 CompileFlags = D3D10_SHADER_ENABLE_BACKWARDS_COMPATIBILITY
+	uint32 CompileFlags = D3DCOMPILE_ENABLE_BACKWARDS_COMPATIBILITY
 		// Unpack uniform matrices as row-major to match the CPU layout.
-		| D3D10_SHADER_PACK_MATRIX_ROW_MAJOR;
+		| D3DCOMPILE_PACK_MATRIX_ROW_MAJOR;
 
 	if (Input.Environment.CompilerFlags.Contains(CFLAG_Debug)) 
 	{
 		//add the debug flags
-		CompileFlags |= D3D10_SHADER_DEBUG | D3D10_SHADER_SKIP_OPTIMIZATION;
+		CompileFlags |= D3DCOMPILE_DEBUG | D3DCOMPILE_SKIP_OPTIMIZATION;
 	}
 	else
 	{
 		if (Input.Environment.CompilerFlags.Contains(CFLAG_StandardOptimization))
 		{
-			CompileFlags |= D3D10_SHADER_OPTIMIZATION_LEVEL1;
+			CompileFlags |= D3DCOMPILE_OPTIMIZATION_LEVEL1;
 		}
 		else
 		{
-			CompileFlags |= D3D10_SHADER_OPTIMIZATION_LEVEL3;
+			CompileFlags |= D3DCOMPILE_OPTIMIZATION_LEVEL3;
 		}
 	}
 
diff --git a/Engine/Source/Developer/Windows/ShaderFormatD3D/Private/D3DShaderCompiler.inl b/Engine/Source/Developer/Windows/ShaderFormatD3D/Private/D3DShaderCompiler.inl
--- a/Engine/Source/Developer/Windows/ShaderFormatD3D/Private/D3DShaderCompiler.inl
+++ b/Engine/Source/Developer/Windows/ShaderFormatD3D/Private/D3DShaderCompiler.inl
@@ -52,61 +52,61 @@ public:
 			ExtraArguments.Add(Exports);
 		}
 
-		if (D3DCompileFlags & D3D10_SHADER_PREFER_FLOW_CONTROL)
+		if (D3DCompileFlags & D3DCOMPILE_PREFER_FLOW_CONTROL)
 		{
-			D3DCompileFlags &= ~D3D10_SHADER_PREFER_FLOW_CONTROL;
+			D3DCompileFlags &= ~D3DCOMPILE_PREFER_FLOW_CONTROL;
 			ExtraArguments.Add(L"/Gfp");
 		}
 
-		if (D3DCompileFlags & D3D10_SHADER_SKIP_OPTIMIZATION)
+		if (D3DCompileFlags & D3DCOMPILE_SKIP_OPTIMIZATION)
 		{
-			D3DCompileFlags &= ~D3D10_SHADER_SKIP_OPTIMIZATION;
+			D3DCompileFlags &= ~D3DCOMPILE_SKIP_OPTIMIZATION;
 			ExtraArguments.Add(L"/Od");
 		}
 
-		if (D3DCompileFlags & D3D10_SHADER_SKIP_VALIDATION)
+		if (D3DCompileFlags & D3DCOMPILE_SKIP_VALIDATION)
 		{
-			D3DCompileFlags &= ~D3D10_SHADER_SKIP_VALIDATION;
+			D3DCompileFlags &= ~D3DCOMPILE_SKIP_VALIDATION;
 			ExtraArguments.Add(L"/Vd");
 		}
 
-		if (D3DCompileFlags & D3D10_SHADER_AVOID_FLOW_CONTROL)
+		if (D3DCompileFlags & D3DCOMPILE_AVOID_FLOW_CONTROL)
 		{
-			D3DCompileFlags &= ~D3D10_SHADER_AVOID_FLOW_CONTROL;
+			D3DCompileFlags &= ~D3DCOMPILE_AVOID_FLOW_CONTROL;
 			ExtraArguments.Add(L"/Gfa");
 		}
 
-		if (D3DCompileFlags & D3D10_SHADER_PACK_MATRIX_ROW_MAJOR)
+		if (D3DCompileFlags & D3DCOMPILE_PACK_MATRIX_ROW_MAJOR)
 		{
-			D3DCompileFlags &= ~D3D10_SHADER_PACK_MATRIX_ROW_MAJOR;
+			D3DCompileFlags &= ~D3DCOMPILE_PACK_MATRIX_ROW_MAJOR;
 			ExtraArguments.Add(L"/Zpr");
 		}
 
-		if (D3DCompileFlags & D3D10_SHADER_ENABLE_BACKWARDS_COMPATIBILITY)
+		if (D3DCompileFlags & D3DCOMPILE_ENABLE_BACKWARDS_COMPATIBILITY)
 		{
-			D3DCompileFlags &= ~D3D10_SHADER_ENABLE_BACKWARDS_COMPATIBILITY;
+			D3DCompileFlags &= ~D3DCOMPILE_ENABLE_BACKWARDS_COMPATIBILITY;
 			ExtraArguments.Add(L"/Gec");
 		}
 
 		switch (D3DCompileFlags & SHADER_OPTIMIZATION_LEVEL_MASK)
 		{
-		case D3D10_SHADER_OPTIMIZATION_LEVEL0:
-			D3DCompileFlags &= ~D3D10_SHADER_OPTIMIZATION_LEVEL0;
+		case D3DCOMPILE_OPTIMIZATION_LEVEL0:
+			D3DCompileFlags &= ~D3DCOMPILE_OPTIMIZATION_LEVEL0;
 			ExtraArguments.Add(L"/O0");
 			break;
 
-		case D3D10_SHADER_OPTIMIZATION_LEVEL1:
-			D3DCompileFlags &= ~D3D10_SHADER_OPTIMIZATION_LEVEL1;
+		case D3DCOMPILE_OPTIMIZATION_LEVEL1:
+			D3DCompileFlags &= ~D3DCOMPILE_OPTIMIZATION_LEVEL1;
 			ExtraArguments.Add(L"/O1");
 			break;
 
-		case D3D10_SHADER_OPTIMIZATION_LEVEL2:
-			D3DCompileFlags &= ~D3D10_SHADER_OPTIMIZATION_LEVEL2;
+		case D3DCOMPILE_OPTIMIZATION_LEVEL2:
+			D3DCompileFlags &= ~D3DCOMPILE_OPTIMIZATION_LEVEL2;
 			ExtraArguments.Add(L"/O2");
 			break;
 
-		case D3D10_SHADER_OPTIMIZATION_LEVEL3:
-			D3DCompileFlags &= ~D3D10_SHADER_OPTIMIZATION_LEVEL3;
+		case D3DCOMPILE_OPTIMIZATION_LEVEL3:
+			D3DCompileFlags &= ~D3DCOMPILE_OPTIMIZATION_LEVEL3;
 			ExtraArguments.Add(L"/O3");
 			break;
 
@@ -114,9 +114,9 @@ public:
 			break;
 		}
 
-		if (D3DCompileFlags & D3D10_SHADER_DEBUG)
+		if (D3DCompileFlags & D3DCOMPILE_DEBUG)
 		{
-			D3DCompileFlags &= ~D3D10_SHADER_DEBUG;
+			D3DCompileFlags &= ~D3DCOMPILE_DEBUG;
 			bKeepDebugInfo = true;
 		}
 
@@ -238,7 +238,7 @@ inline void ExtractParameterMapFromD3DShader(
 			continue;
 		}
 
-		if (BindDesc.Type == D3D10_SIT_CBUFFER || BindDesc.Type == D3D10_SIT_TBUFFER)
+		if (BindDesc.Type == D3D_SIT_CBUFFER || BindDesc.Type == D3D_SIT_TBUFFER)
 		{
 			const uint32 CBIndex = BindDesc.BindPoint;
 			ID3D1xShaderReflectionConstantBuffer* ConstantBuffer = Reflector->GetConstantBufferByName(BindDesc.Name);
@@ -254,7 +254,7 @@ inline void ExtractParameterMapFromD3DShader(
 					ID3D1xShaderReflectionVariable* Variable = ConstantBuffer->GetVariableByIndex(ConstantIndex);
 					D3D1x_SHADER_VARIABLE_DESC VariableDesc;
 					Variable->GetDesc(&VariableDesc);
-					if (VariableDesc.uFlags & D3D10_SVF_USED)
+					if (VariableDesc.uFlags & D3D_SVF_USED)
 					{
 						bGlobalUniformBufferUsed = true;
 
@@ -290,7 +290,7 @@ inline void ExtractParameterMapFromD3DShader(
 
 			NumCBs = FMath::Max(NumCBs, BindDesc.BindPoint + BindDesc.BindCount);
 		}
-		else if (BindDesc.Type == D3D10_SIT_TEXTURE || BindDesc.Type == D3D10_SIT_SAMPLER)
+		else if (BindDesc.Type == D3D_SIT_TEXTURE || BindDesc.Type == D3D_SIT_SAMPLER)
 		{
 			check(BindDesc.BindCount == 1);
 
@@ -304,12 +304,12 @@ inline void ExtractParameterMapFromD3DShader(
 
 			const uint32 BindCount = 1;
 			EShaderParameterType ParameterType = EShaderParameterType::Num;
-			if (BindDesc.Type == D3D10_SIT_SAMPLER)
+			if (BindDesc.Type == D3D_SIT_SAMPLER)
 			{
 				ParameterType = EShaderParameterType::Sampler;
 				NumSamplers = FMath::Max(NumSamplers, BindDesc.BindPoint + BindCount);
 			}
-			else if (BindDesc.Type == D3D10_SIT_TEXTURE)
+			else if (BindDesc.Type == D3D_SIT_TEXTURE)
 			{
 				ParameterType = EShaderParameterType::SRV;
 				NumSRVs = FMath::Max(NumSRVs, BindDesc.BindPoint + BindCount);
@@ -337,9 +337,9 @@ inline void ExtractParameterMapFromD3DShader(
 				);
 			}
 		}
-		else if (BindDesc.Type == D3D11_SIT_UAV_RWTYPED || BindDesc.Type == D3D11_SIT_UAV_RWSTRUCTURED ||
-			BindDesc.Type == D3D11_SIT_UAV_RWBYTEADDRESS || BindDesc.Type == D3D11_SIT_UAV_RWSTRUCTURED_WITH_COUNTER ||
-			BindDesc.Type == D3D11_SIT_UAV_APPEND_STRUCTURED)
+		else if (BindDesc.Type == D3D_SIT_UAV_RWTYPED || BindDesc.Type == D3D_SIT_UAV_RWSTRUCTURED ||
+			BindDesc.Type == D3D_SIT_UAV_RWBYTEADDRESS || BindDesc.Type == D3D_SIT_UAV_RWSTRUCTURED_WITH_COUNTER ||
+			BindDesc.Type == D3D_SIT_UAV_APPEND_STRUCTURED)
 		{
 			check(BindDesc.BindCount == 1);
 
@@ -395,7 +395,7 @@ inline void ExtractParameterMapFromD3DShader(
 
 			NumUAVs = FMath::Max(NumUAVs, BindDesc.BindPoint + BindCount);
 		}
-		else if (BindDesc.Type == D3D11_SIT_STRUCTURED || BindDesc.Type == D3D11_SIT_BYTEADDRESS)
+		else if (BindDesc.Type == D3D_SIT_STRUCTURED || BindDesc.Type == D3D_SIT_BYTEADDRESS)
 		{
 			check(BindDesc.BindCount == 1);
 			TCHAR OfficialName[1024];
diff --git a/Engine/Source/Developer/Windows/ShaderFormatD3D/Private/D3DShaderCompilerDXC.cpp b/Engine/Source/Developer/Windows/ShaderFormatD3D/Private/D3DShaderCompilerDXC.cpp
--- a/Engine/Source/Developer/Windows/ShaderFormatD3D/Private/D3DShaderCompilerDXC.cpp
+++ b/Engine/Source/Developer/Windows/ShaderFormatD3D/Private/D3DShaderCompilerDXC.cpp
@@ -14,7 +14,7 @@
 DEFINE_LOG_CATEGORY_STATIC(LogD3D12ShaderCompiler, Log, All);
 
 // D3D doesn't define a mask for this, so we do so here
-#define SHADER_OPTIMIZATION_LEVEL_MASK (D3D10_SHADER_OPTIMIZATION_LEVEL0 | D3D10_SHADER_OPTIMIZATION_LEVEL1 | D3D10_SHADER_OPTIMIZATION_LEVEL2 | D3D10_SHADER_OPTIMIZATION_LEVEL3)
+#define SHADER_OPTIMIZATION_LEVEL_MASK (D3DCOMPILE_OPTIMIZATION_LEVEL0 | D3DCOMPILE_OPTIMIZATION_LEVEL1 | D3DCOMPILE_OPTIMIZATION_LEVEL2 | D3DCOMPILE_OPTIMIZATION_LEVEL3)
 
 // Disable macro redefinition warning for compatibility with Windows SDK 8+
 #pragma warning(push)
@@ -654,10 +654,10 @@ bool CompileAndProcessD3DShaderDXC(FString& PreprocessedShaderSource,
 
 	// Ignore backwards compatibility flag (/Gec) as it is deprecated.
 	// #dxr_todo: this flag should not be even passed into this function from the higher level.
-	uint32 DXCFlags = CompileFlags & (~D3D10_SHADER_ENABLE_BACKWARDS_COMPATIBILITY);
+	uint32 DXCFlags = CompileFlags & (~D3DCOMPILE_ENABLE_BACKWARDS_COMPATIBILITY);
 	if (Input.Environment.CompilerFlags.Contains(CFLAG_SkipOptimizationsDXC))
 	{
-		DXCFlags |= D3D10_SHADER_SKIP_OPTIMIZATION;
+		DXCFlags |= D3DCOMPILE_SKIP_OPTIMIZATION;
 	}
 
 	const bool bKeepDebugInfo = Input.Environment.CompilerFlags.Contains(CFLAG_KeepDebugInfo);
