Fixing static analysis warnings.
Each case is a "if (BindDesc.Type == D3D_SIT_TEXTURE || BindDesc.Type == D3D_SIT_SAMPLER)" block that has "if (BindDesc.Type == D3D_SIT_SAMPLER) {} else if (BindDesc.Type == D3D_SIT_TEXTURE) {}", the second condition should just be an else.

diff --git a/Engine/Source/Developer/Windows/ShaderFormatD3D/Private/D3DShaderCompiler.inl b/Engine/Source/Developer/Windows/ShaderFormatD3D/Private/D3DShaderCompiler.inl
--- a/Engine/Source/Developer/Windows/ShaderFormatD3D/Private/D3DShaderCompiler.inl
+++ b/Engine/Source/Developer/Windows/ShaderFormatD3D/Private/D3DShaderCompiler.inl
@@ -309,7 +309,7 @@ inline void ExtractParameterMapFromD3DShader(
 				ParameterType = EShaderParameterType::Sampler;
 				NumSamplers = FMath::Max(NumSamplers, BindDesc.BindPoint + BindCount);
 			}
-			else if (BindDesc.Type == D3D_SIT_TEXTURE)
+			else // D3D_SIT_TEXTURE
 			{
 				ParameterType = EShaderParameterType::SRV;
 				NumSRVs = FMath::Max(NumSRVs, BindDesc.BindPoint + BindCount);
