Fix for RGS not finding global uniform buffers.

diff --git a/Engine/Source/Developer/Windows/ShaderFormatD3D/Private/D3DShaderCompiler.inl b/Engine/Source/Developer/Windows/ShaderFormatD3D/Private/D3DShaderCompiler.inl
--- a/Engine/Source/Developer/Windows/ShaderFormatD3D/Private/D3DShaderCompiler.inl
+++ b/Engine/Source/Developer/Windows/ShaderFormatD3D/Private/D3DShaderCompiler.inl
@@ -453,7 +453,12 @@ inline void GenerateFinalOutput(TRefCountPtr<TBlob>& CompressedData,
 		// Build the generic SRT for this shader.
 		FShaderCompilerResourceTable GenericSRT;
 		BuildResourceTableMapping(Input.Environment.ResourceTableMap, Input.Environment.ResourceTableLayoutHashes, UsedUniformBufferSlots, Output.ParameterMap, GenericSRT);
-		CullGlobalUniformBuffers(Input.Environment.ResourceTableLayoutSlots, Output.ParameterMap);
+
+		// Ray generation shaders rely on a different binding model that aren't compatible with global uniform buffers.
+		if (Input.Target.Frequency != SF_RayGen)
+		{
+			CullGlobalUniformBuffers(Input.Environment.ResourceTableLayoutSlots, Output.ParameterMap);
+		}
 
 		if (UniformBufferNames.Num() < GenericSRT.ResourceTableLayoutHashes.Num())
 		{
