Make sure we use the checked in version of d3dcompiler_47.dll when loading the d3d shader format. Also make sure that version goes along in SN-DBS.

diff --git a/Engine/Source/Developer/Windows/ShaderFormatD3D/ShaderFormatD3D.Build.cs b/Engine/Source/Developer/Windows/ShaderFormatD3D/ShaderFormatD3D.Build.cs
--- a/Engine/Source/Developer/Windows/ShaderFormatD3D/ShaderFormatD3D.Build.cs
+++ b/Engine/Source/Developer/Windows/ShaderFormatD3D/ShaderFormatD3D.Build.cs
@@ -29,7 +29,8 @@ public class ShaderFormatD3D : ModuleRules
 
             string DxDllsPath = "$(EngineDir)/Binaries/ThirdParty/Windows/DirectX/x64/";
 
-            RuntimeDependencies.Add("$(TargetOutputDir)/dxil.dll", DxDllsPath + "dxil.dll");
+			RuntimeDependencies.Add("$(TargetOutputDir)/d3dcompiler_47.dll", DxDllsPath + "d3dcompiler_47.dll");
+			RuntimeDependencies.Add("$(TargetOutputDir)/dxil.dll", DxDllsPath + "dxil.dll");
 
 			string BinaryFolder = Target.UEThirdPartyBinariesDirectory + "ShaderConductor/Win64";
 			RuntimeDependencies.Add(BinaryFolder + "/dxcompiler.dll");
diff --git a/Engine/Source/Developer/Windows/WindowsTargetPlatform/Private/GenericWindowsTargetPlatform.h b/Engine/Source/Developer/Windows/WindowsTargetPlatform/Private/GenericWindowsTargetPlatform.h
--- a/Engine/Source/Developer/Windows/WindowsTargetPlatform/Private/GenericWindowsTargetPlatform.h
+++ b/Engine/Source/Developer/Windows/WindowsTargetPlatform/Private/GenericWindowsTargetPlatform.h
@@ -144,6 +144,7 @@ public:
 		FTargetPlatformBase::AddDependencySCArrayHelper(OutDependencies, TEXT("Binaries/ThirdParty/Windows/DirectX/x64/dxil.dll"));
 		FTargetPlatformBase::AddDependencySCArrayHelper(OutDependencies, TEXT("Binaries/ThirdParty/ShaderConductor/Win64/ShaderConductor.dll"));
 		FTargetPlatformBase::AddDependencySCArrayHelper(OutDependencies, TEXT("Binaries/ThirdParty/ShaderConductor/Win64/dxcompiler.dll"));
+		FTargetPlatformBase::AddDependencySCArrayHelper(OutDependencies, TEXT("Binaries/Win64/d3dcompiler_47.dll"));
 		FTargetPlatformBase::AddDependencySCArrayHelper(OutDependencies, TEXT("Binaries/Win64/dxcompiler.dll"));
 		FTargetPlatformBase::AddDependencySCArrayHelper(OutDependencies, TEXT("Binaries/Win64/dxil.dll"));
 	}
