Integration of Live++ 1.6.3

#rnx

diff --git a/Engine/Source/Developer/Windows/LiveCoding/LiveCoding.Build.cs b/Engine/Source/Developer/Windows/LiveCoding/LiveCoding.Build.cs
--- a/Engine/Source/Developer/Windows/LiveCoding/LiveCoding.Build.cs
+++ b/Engine/Source/Developer/Windows/LiveCoding/LiveCoding.Build.cs
@@ -20,6 +20,15 @@ public class LiveCoding : ModuleRules
             PrivateDefinitions.Add("USE_DEBUG_LIVE_CODING_CONSOLE=0");
         }
 
+        if (Target.Configuration == UnrealTargetConfiguration.Debug)
+        {
+        	PrivateDefinitions.Add("LC_DEBUG=1");
+        }
+        else
+        {
+        	PrivateDefinitions.Add("LC_DEBUG=0");
+        }
+
 		PublicAdditionalLibraries.Add(Path.Combine(ModuleDirectory, "Private", "External", "LC_JumpToSelf.lib"));
 	}
 }
diff --git a/Engine/Source/Developer/Windows/LiveCoding/Private/External/LC_API.cpp b/Engine/Source/Developer/Windows/LiveCoding/Private/External/LC_API.cpp
--- a/Engine/Source/Developer/Windows/LiveCoding/Private/External/LC_API.cpp
+++ b/Engine/Source/Developer/Windows/LiveCoding/Private/External/LC_API.cpp
@@ -1,14 +1,19 @@
-// Copyright 2011-2019 Molecular Matters GmbH, all rights reserved.
+// Copyright 2011-2020 Molecular Matters GmbH, all rights reserved.
 
+// EPIC BEGIN MOD
+//#include PCH_INCLUDE
+// EPIC END MOD
 #include "LC_API.h"
 #include "LC_ClientStartupThread.h"
 #include "LC_CriticalSection.h"
 #include "LC_SyncPoint.h"
 #include "LC_Restart.h"
 #include "LC_RunMode.h"
-#include "LC_Logging.h"
 #include "LC_Memory.h"
 #include "LPP_API.h"
+// BEGIN EPIC MOD
+#include "LC_Logging.h"
+// END EPIC MOD
 
 namespace
 {
@@ -35,22 +40,12 @@ namespace
 }
 
 
-
-namespace
-{
-	// startup thread
-	static ClientStartupThread* g_mainStartupThread = nullptr;
-}
-
-
-// BEGIN EPIC MOD - Adding instance to startup function
-LPP_DLL_API(void) LppStartup(void* thisInstance)
+LPP_DLL_API(void) LppStartup(void)
 {
-	g_startupThread = new ClientStartupThread((HINSTANCE)thisInstance);
+	g_startupThread = new ClientStartupThread();
 
 	restart::Startup();
 }
-// END EPIC MOD
 
 
 LPP_DLL_API(void) LppShutdown(void)
@@ -361,4 +356,4 @@ LPP_DLL_API(void) LppApplySettingString(const char* settingName, const wchar_t*
 	}
 
 	g_startupThread->ApplySettingString(settingName, value);
-}
+}
\ No newline at end of file
diff --git a/Engine/Source/Developer/Windows/LiveCoding/Private/External/LC_API.h b/Engine/Source/Developer/Windows/LiveCoding/Private/External/LC_API.h
--- a/Engine/Source/Developer/Windows/LiveCoding/Private/External/LC_API.h
+++ b/Engine/Source/Developer/Windows/LiveCoding/Private/External/LC_API.h
@@ -1,4 +1,4 @@
-// Copyright 2011-2019 Molecular Matters GmbH, all rights reserved.
+// Copyright 2011-2020 Molecular Matters GmbH, all rights reserved.
 
 #pragma once
 
@@ -11,9 +11,7 @@
 //#define LPP_DLL_API(_rv) extern "C" __declspec(dllexport) _rv __cdecl
 // END EPIC MOD - Internalizing API
 
-// BEGIN EPIC MOD - Adding instance to startup function
-LPP_DLL_API(void) LppStartup(void* thisInstance);
-// END EPIC MOD
+LPP_DLL_API(void) LppStartup(void);
 LPP_DLL_API(void) LppShutdown(void);
 LPP_DLL_API(const char*) LppGetVersion(void);
 LPP_DLL_API(int) LppCheckVersion(const char* apiVersion);
@@ -44,4 +42,3 @@ LPP_DLL_API(void*) LppEnableLazyLoadedModule(const wchar_t* nameOfExeOrDll);
 LPP_DLL_API(void) LppApplySettingBool(const char* settingName, int value);
 LPP_DLL_API(void) LppApplySettingInt(const char* settingName, int value);
 LPP_DLL_API(void) LppApplySettingString(const char* settingName, const wchar_t* value);
-
diff --git a/Engine/Source/Developer/Windows/LiveCoding/Private/External/LC_BitUtil.h b/Engine/Source/Developer/Windows/LiveCoding/Private/External/LC_BitUtil.h
--- a/Engine/Source/Developer/Windows/LiveCoding/Private/External/LC_BitUtil.h
+++ b/Engine/Source/Developer/Windows/LiveCoding/Private/External/LC_BitUtil.h
@@ -1,4 +1,4 @@
-// Copyright 2011-2019 Molecular Matters GmbH, all rights reserved.
+// Copyright 2011-2020 Molecular Matters GmbH, all rights reserved.
 
 #pragma once
 
diff --git a/Engine/Source/Developer/Windows/LiveCoding/Private/External/LC_ClientCommandActions.cpp b/Engine/Source/Developer/Windows/LiveCoding/Private/External/LC_ClientCommandActions.cpp
--- a/Engine/Source/Developer/Windows/LiveCoding/Private/External/LC_ClientCommandActions.cpp
+++ b/Engine/Source/Developer/Windows/LiveCoding/Private/External/LC_ClientCommandActions.cpp
@@ -1,5 +1,8 @@
-// Copyright 2011-2019 Molecular Matters GmbH, all rights reserved.
+// Copyright 2011-2020 Molecular Matters GmbH, all rights reserved.
 
+// BEGIN EPIC MOD
+//#include PCH_INCLUDE
+// END EPIC MOD
 #include "LC_ClientCommandActions.h"
 #include "LC_ClientUserCommandThread.h"
 #include "LC_DuplexPipe.h"
@@ -7,7 +10,9 @@
 #include "LC_Executable.h"
 #include "LC_Event.h"
 #include "LC_Process.h"
+// BEGIN EPIC MOD
 #include "LC_Logging.h"
+// END EPIC MOD
 
 // BEGIN EPIC MOD - Support for object reinstancing. We need to call a global post-patch handler, rather than just getting individual callbacks for modified modules.
 extern bool GHasLoadedPatch;
@@ -167,7 +172,9 @@ bool actions::CallEntryPoint::Execute(const CommandType* command, const DuplexPi
 
 bool actions::LogOutput::Execute(const CommandType*, const DuplexPipe* pipe, void*, const void* payload, size_t)
 {
-	logging::LogNoFormat<logging::Channel::USER>(static_cast<const wchar_t*>(payload));
+	// BEGIN EPIC MOD
+	Logging::LogNoFormat<Logging::Channel::USER>(static_cast<const wchar_t*>(payload));
+	// END EPIC MOD
 	pipe->SendAck();
 
 	return true;
diff --git a/Engine/Source/Developer/Windows/LiveCoding/Private/External/LC_ClientCommandActions.h b/Engine/Source/Developer/Windows/LiveCoding/Private/External/LC_ClientCommandActions.h
--- a/Engine/Source/Developer/Windows/LiveCoding/Private/External/LC_ClientCommandActions.h
+++ b/Engine/Source/Developer/Windows/LiveCoding/Private/External/LC_ClientCommandActions.h
@@ -1,8 +1,10 @@
-// Copyright 2011-2019 Molecular Matters GmbH, all rights reserved.
+// Copyright 2011-2020 Molecular Matters GmbH, all rights reserved.
 
 #pragma once
 
+// BEGIN EPIC MOD
 #include "CoreTypes.h"
+// END EPIC MOD
 #include "LC_Commands.h"
 
 
diff --git a/Engine/Source/Developer/Windows/LiveCoding/Private/External/LC_ClientCommandThread.cpp b/Engine/Source/Developer/Windows/LiveCoding/Private/External/LC_ClientCommandThread.cpp
--- a/Engine/Source/Developer/Windows/LiveCoding/Private/External/LC_ClientCommandThread.cpp
+++ b/Engine/Source/Developer/Windows/LiveCoding/Private/External/LC_ClientCommandThread.cpp
@@ -1,5 +1,8 @@
-// Copyright 2011-2019 Molecular Matters GmbH, all rights reserved.
+// Copyright 2011-2020 Molecular Matters GmbH, all rights reserved.
 
+// BEGIN EPIC MOD
+//#include PCH_INCLUDE
+// END EPIC MOD
 #include "LC_ClientCommandThread.h"
 #include "LC_Event.h"
 #include "LC_CommandMap.h"
@@ -8,10 +11,11 @@
 #include "LC_DuplexPipeClient.h"
 #include "LC_Process.h"
 #include "LC_HeartBeat.h"
+#include "LC_Thread.h"
 
 
 ClientCommandThread::ClientCommandThread(DuplexPipeClient* pipeClient)
-	: m_thread(INVALID_HANDLE_VALUE)
+	: m_thread(Thread::INVALID_HANDLE)
 	, m_pipe(pipeClient)
 {
 }
@@ -22,26 +26,28 @@ ClientCommandThread::~ClientCommandThread(void)
 }
 
 
-unsigned int ClientCommandThread::Start(const std::wstring& processGroupName, Event* compilationEvent, Event* waitForStartEvent, CriticalSection* pipeAccessCS)
+Thread::Id ClientCommandThread::Start(const std::wstring& processGroupName, Event* compilationEvent, Event* waitForStartEvent, CriticalSection* pipeAccessCS)
 {
 	// spawn a thread that communicates with the server
-	m_thread = thread::Create("Live coding commands", 128u * 1024u, &ClientCommandThread::ThreadFunction, this, processGroupName, compilationEvent, waitForStartEvent, pipeAccessCS);
+	// BEGIN EPIC MOD
+	m_thread = Thread::CreateFromMemberFunction("Live coding commands", 128u * 1024u, this, &ClientCommandThread::ThreadFunction, processGroupName, compilationEvent, waitForStartEvent, pipeAccessCS);
+	// END EPIC MOD
 
-	return thread::GetId(m_thread);
+	return Thread::GetId(m_thread);
 }
 
 
 void ClientCommandThread::Join(void)
 {
-	if (m_thread != INVALID_HANDLE_VALUE)
+	if (m_thread != Thread::INVALID_HANDLE)
 	{
-		thread::Join(m_thread);
-		thread::Close(m_thread);
+		Thread::Join(m_thread);
+		Thread::Close(m_thread);
 	}
 }
 
 
-unsigned int ClientCommandThread::ThreadFunction(const std::wstring& processGroupName, Event* compilationEvent, Event* waitForStartEvent, CriticalSection* pipeAccessCS)
+Thread::ReturnValue ClientCommandThread::ThreadFunction(const std::wstring& processGroupName, Event* compilationEvent, Event* waitForStartEvent, CriticalSection* pipeAccessCS)
 {
 	waitForStartEvent->Wait();
 
@@ -55,7 +61,7 @@ unsigned int ClientCommandThread::ThreadFunction(const std::wstring& processGrou
 	commandMap.RegisterAction<actions::LogOutput>();
 	commandMap.RegisterAction<actions::CompilationFinished>();
 
-	HeartBeat heartBeat(processGroupName.c_str(), process::GetId());
+	HeartBeat heartBeat(processGroupName.c_str(), Process::Current::GetId());
 
 	for (;;)
 	{
@@ -65,7 +71,7 @@ unsigned int ClientCommandThread::ThreadFunction(const std::wstring& processGrou
 			if (!m_pipe->IsValid())
 			{
 				// pipe was closed or is broken, bail out
-				return 1u;
+				return Thread::ReturnValue(1u);
 			}
 
 			heartBeat.Store();
@@ -74,7 +80,7 @@ unsigned int ClientCommandThread::ThreadFunction(const std::wstring& processGrou
 		if (!m_pipe->IsValid())
 		{
 			// pipe was closed or is broken, bail out
-			return 1u;
+			return Thread::ReturnValue(1u);
 		}
 
 		// lock critical section for accessing the pipe.
@@ -86,5 +92,7 @@ unsigned int ClientCommandThread::ThreadFunction(const std::wstring& processGrou
 		commandMap.HandleCommands(m_pipe, nullptr);
 	}
 
-	return 0u;
+	// BEGIN EPIC MOD
+	return Thread::ReturnValue(0u);
+	// END EPIC MOD
 }
diff --git a/Engine/Source/Developer/Windows/LiveCoding/Private/External/LC_ClientCommandThread.h b/Engine/Source/Developer/Windows/LiveCoding/Private/External/LC_ClientCommandThread.h
--- a/Engine/Source/Developer/Windows/LiveCoding/Private/External/LC_ClientCommandThread.h
+++ b/Engine/Source/Developer/Windows/LiveCoding/Private/External/LC_ClientCommandThread.h
@@ -1,10 +1,14 @@
-// Copyright 2011-2019 Molecular Matters GmbH, all rights reserved.
+// Copyright 2011-2020 Molecular Matters GmbH, all rights reserved.
 
 #pragma once
 
+// BEGIN EPIC MOD
 #include "CoreTypes.h"
-#include "LC_Thread.h"
+// END EPIC MOD
+#include "LC_ThreadTypes.h"
+// BEGIN EPIC MOD
 #include <string>
+// END EPIC MOD
 
 class DuplexPipeClient;
 class Event;
@@ -20,14 +24,14 @@ public:
 
 	// Starts the thread that takes care of handling incoming commands on the pipe.
 	// Returns the thread ID.
-	unsigned int Start(const std::wstring& processGroupName, Event* compilationEvent, Event* waitForStartEvent, CriticalSection* pipeAccessCS);
+	Thread::Id Start(const std::wstring& processGroupName, Event* compilationEvent, Event* waitForStartEvent, CriticalSection* pipeAccessCS);
 
 	// Joins this thread.
 	void Join(void);
 
 private:
-	unsigned int ThreadFunction(const std::wstring& processGroupName, Event* compilationEvent, Event* waitForStartEvent, CriticalSection* pipeAccessCS);
+	Thread::ReturnValue ThreadFunction(const std::wstring& processGroupName, Event* compilationEvent, Event* waitForStartEvent, CriticalSection* pipeAccessCS);
 
-	thread::Handle m_thread;
+	Thread::Handle m_thread;
 	DuplexPipeClient* m_pipe;
 };
diff --git a/Engine/Source/Developer/Windows/LiveCoding/Private/External/LC_ClientStartupThread.cpp b/Engine/Source/Developer/Windows/LiveCoding/Private/External/LC_ClientStartupThread.cpp
--- a/Engine/Source/Developer/Windows/LiveCoding/Private/External/LC_ClientStartupThread.cpp
+++ b/Engine/Source/Developer/Windows/LiveCoding/Private/External/LC_ClientStartupThread.cpp
@@ -1,5 +1,8 @@
-// Copyright 2011-2019 Molecular Matters GmbH, all rights reserved.
+// Copyright 2011-2020 Molecular Matters GmbH, all rights reserved.
 
+// BEGIN EPIC MOD
+//#include PCH_INCLUDE
+// END EPIC MOD
 #include "LC_ClientStartupThread.h"
 #include "LC_StringUtil.h"
 #include "LC_NamedSharedMemory.h"
@@ -14,10 +17,15 @@
 #include "LC_PrimitiveNames.h"
 #include "LC_Environment.h"
 #include "LC_MemoryStream.h"
+#include "LC_Thread.h"
+#include "LC_Process.h"
+#include "LPP_API.h"
+// BEGIN EPIC MOD
 #include "LC_Logging.h"
 #include "Misc/Paths.h"
 #include "Misc/ConfigCacheIni.h"
 #include "Misc/App.h"
+// END EPIC MOD
 
 // JumpToSelf is an extern function coming from assembler source
 extern void JumpToSelf(void);
@@ -33,9 +41,8 @@ namespace
 }
 
 
-ClientStartupThread::ClientStartupThread(HINSTANCE instance)
-	: m_instance(instance)
-	, m_thread(INVALID_HANDLE_VALUE)
+ClientStartupThread::ClientStartupThread(void)
+	: m_thread(Thread::INVALID_HANDLE)
 	, m_job(nullptr)
 	, m_sharedMemory(nullptr)
 	, m_mainProcessContext(nullptr)
@@ -99,7 +106,7 @@ ClientStartupThread::~ClientStartupThread(void)
 
 	if (m_mainProcessContext)
 	{
-		process::Destroy(m_mainProcessContext);
+		Process::Destroy(m_mainProcessContext);
 	}
 
 	// close job object to make child processes close as well.
@@ -107,7 +114,7 @@ ClientStartupThread::~ClientStartupThread(void)
 	::CloseHandle(m_job);
 
 	// clean up interprocess objects
-	DeleteAndNull(m_sharedMemory);
+	Process::DestroyNamedSharedMemory(m_sharedMemory);
 }
 
 
@@ -117,16 +124,18 @@ void ClientStartupThread::Start(const char* const groupName, RunMode::Enum runMo
 	// in the context of mutexes, jobs, named shared memory, etc. object names behave similar to
 	// file names and are not allowed to contain certain characters.
 	std::wstring safeProcessGroupName = string::MakeSafeName(string::ToWideString(groupName));
-	m_thread = thread::Create("Live coding startup", 128u * 1024u, &ClientStartupThread::ThreadFunction, this, safeProcessGroupName, runMode);
+	// BEGIN EPIC MOD
+	m_thread = Thread::CreateFromMemberFunction("Live coding startup", 128u * 1024u, this, &ClientStartupThread::ThreadFunction, safeProcessGroupName, runMode);
+	// END EPIC MOD
 }
 
 
 void ClientStartupThread::Join(void)
 {
-	if (m_thread != INVALID_HANDLE_VALUE)
+	if (m_thread != Thread::INVALID_HANDLE)
 	{
-		thread::Join(m_thread);
-		thread::Close(m_thread);
+		Thread::Join(m_thread);
+		Thread::Close(m_thread);
 	}
 }
 
@@ -394,7 +403,7 @@ void ClientStartupThread::ApplySettingString(const char* settingName, const wcha
 }
 
 
-unsigned int ClientStartupThread::ThreadFunction(const std::wstring& processGroupName, RunMode::Enum runMode)
+Thread::ReturnValue ClientStartupThread::ThreadFunction(const std::wstring& processGroupName, RunMode::Enum runMode)
 {
 	// configure all child processes associated with the job to terminate when the parent terminates.
 	// we create (or open) a process-wide job per process group and register the spawned process with that job.
@@ -418,8 +427,8 @@ unsigned int ClientStartupThread::ThreadFunction(const std::wstring& processGrou
 		InterprocessMutex initProcessMutex(primitiveNames::StartupMutex(processGroupName).c_str());
 		InterprocessMutex::ScopedLock mutexLock(&initProcessMutex);
 
-		m_sharedMemory = new NamedSharedMemory(primitiveNames::StartupNamedSharedMemory(processGroupName).c_str());
-		if (m_sharedMemory->IsOwnedByCallingProcess())
+		m_sharedMemory = Process::CreateNamedSharedMemory(primitiveNames::StartupNamedSharedMemory(processGroupName).c_str(), 4096u);
+		if (Process::Current::DoesOwnNamedSharedMemory(m_sharedMemory))
 		{
 			// BEGIN EPIC MOD - Using LiveCodeConsole
 			// we are the first DLL. spawn the console.
@@ -433,7 +442,6 @@ unsigned int ClientStartupThread::ThreadFunction(const std::wstring& processGrou
 			commandLine += L"-Group=";
 			commandLine += processGroupName;
 
-			// BEGIN EPIC MOD - Additional arguments for console app
 			extern FString GLiveCodingConsoleArguments;
 			if(GLiveCodingConsoleArguments.Len() > 0)
 			{
@@ -446,40 +454,43 @@ unsigned int ClientStartupThread::ThreadFunction(const std::wstring& processGrou
 				commandLine += FApp::GetProjectName();
 				commandLine += L"\"";
 			}
-			// END EPIC MOD
 
-			m_mainProcessContext = process::Spawn(exePath.c_str(), nullptr, commandLine.c_str(), nullptr, process::SpawnFlags::NONE);
-			if (m_mainProcessContext->pi.dwProcessId != 0)
+			m_mainProcessContext = Process::Spawn(exePath.c_str(), nullptr, commandLine.c_str(), nullptr, Process::SpawnFlags::NONE);
+			if (Process::GetId(m_mainProcessContext) != Process::Id(0u))
 			{
-				m_processHandle = m_mainProcessContext->pi.hProcess;
-				::AssignProcessToJobObject(m_job, m_processHandle);
+				m_processHandle = Process::GetHandle(m_mainProcessContext);
+				::AssignProcessToJobObject(m_job, +m_processHandle);
 
-				// share Live++ process ID with other processes
-				m_sharedMemory->Write(m_mainProcessContext->pi.dwProcessId);
+				// share Live++ process Id with other processes
+				Process::WriteNamedSharedMemory(m_sharedMemory, +Process::GetId(m_mainProcessContext));
 			}
 			// END EPIC MOD - Using LiveCodeConsole
 		}
 		else
 		{
 			// the Live++ process is already running. fetch the process ID from shared memory.
-			const DWORD processId = m_sharedMemory->Read<DWORD>();
+			const Process::Id::Type processId = Process::ReadNamedSharedMemory<Process::Id::Type>(m_sharedMemory);
+			// BEGIN EPIC MOD
 			LC_LOG_USER("Detected running instance in process group \"%S\", connecting to console process (PID: %d)", processGroupName.c_str(), processId);
+			// END EPIC MOD
 
 			if (processId != 0u)
 			{
-				m_processHandle = process::Open(processId);
-				::AssignProcessToJobObject(m_job, m_processHandle);
+				m_processHandle = Process::Open(Process::Id(processId));
+				::AssignProcessToJobObject(m_job, +m_processHandle);
 			}
 		}
 	}
 
-	if (!m_processHandle)
+	if (+m_processHandle == nullptr)
 	{
 		// we were unable to open the process, bail out
+		// BEGIN EPIC MOD
 		LC_ERROR_USER("Unable to attach to console process");
-		DeleteAndNull(m_sharedMemory);
+		// END EPIC MOD
+		Process::DestroyNamedSharedMemory(m_sharedMemory);
 
-		return 1u;
+		return Thread::ReturnValue(1u);
 	}
 
 	// wait for server to become ready
@@ -494,18 +505,22 @@ unsigned int ClientStartupThread::ThreadFunction(const std::wstring& processGrou
 	if (!m_pipeClient->Connect(primitiveNames::Pipe(processGroupName).c_str()))
 	{
 		// could not connect to Live++ process
+		// BEGIN EPIC MOD
 		LC_ERROR_USER("Could not connect named pipe to console process");
+		// END EPIC MOD
 
-		return 2u;
+		return Thread::ReturnValue(2u);
 	}
 
 	// create a named duplex pipe for communicating exceptions between DLL and Live++ process
 	if (!m_exceptionPipeClient->Connect(primitiveNames::ExceptionPipe(processGroupName).c_str()))
 	{
 		// could not connect to Live++ process
+		// BEGIN EPIC MOD
 		LC_ERROR_USER("Could not connect exception pipe to console process");
+		// END EPIC MOD
 
-		return 3u;
+		return Thread::ReturnValue(3u);
 	}
 
 	m_pipeClientCS = new CriticalSection;
@@ -518,13 +533,13 @@ unsigned int ClientStartupThread::ThreadFunction(const std::wstring& processGrou
 	// pipe for communicating as long as we aren't finished with it.
 	m_startEvent = new Event(nullptr, Event::Type::MANUAL_RESET);
 
-	const unsigned int commandThreadId = m_commandThread->Start(processGroupName, m_compilationEvent, m_startEvent, m_pipeClientCS);
+	const Thread::Id commandThreadId = m_commandThread->Start(processGroupName, m_compilationEvent, m_startEvent, m_pipeClientCS);
 	m_userCommandThread->Start(processGroupName, m_startEvent, m_pipeClientCS);
 
 	// register this process with Live++
 	{
 		// try getting the previous process ID from the environment in case the process was restarted
-		unsigned int restartedProcessId = 0u;
+		Process::Id restartedProcessId(0u);
 		const std::wstring& processIdStr = environment::GetVariable(L"LPP_PROCESS_RESTART_ID", nullptr);
 		if (processIdStr.length() != 0u)
 		{
@@ -536,33 +551,33 @@ unsigned int ClientStartupThread::ThreadFunction(const std::wstring& processGrou
 		// upon restart, the environment block is inherited by the new process and can be used to map the process IDs of
 		// restarted processes to their previous IDs.
 		{
-			const unsigned int processID = process::GetId();
-			environment::SetVariable(L"LPP_PROCESS_RESTART_ID", std::to_wstring(processID).c_str());
+			const Process::Id processID = Process::Current::GetId();
+			environment::SetVariable(L"LPP_PROCESS_RESTART_ID", std::to_wstring(+processID).c_str());
 		}
 
-		const std::wstring& imagePath = process::GetImagePath();
-		const std::wstring& commandLine = process::GetCommandLine();
-		const std::wstring& workingDirectory = process::GetWorkingDirectory();
-		process::Environment* environment = process::CreateEnvironment(::GetCurrentProcess());
+		const std::wstring imagePath = Process::Current::GetImagePath().GetString();
+		const std::wstring& commandLine = Process::Current::GetCommandLine();
+		const std::wstring& workingDirectory = Process::Current::GetWorkingDirectory().GetString();
+		Process::Environment environment = Process::CreateEnvironment(Process::Current::GetHandle());
 
 		const commands::RegisterProcess command =
 		{
-			process::GetBase(), process::GetId(), restartedProcessId, commandThreadId, reinterpret_cast<void*>(&JumpToSelf),
+			Process::Current::GetBase(), Process::Current::GetId(), restartedProcessId, commandThreadId, reinterpret_cast<void*>(&JumpToSelf),
 			(imagePath.size() + 1u) * sizeof(wchar_t), 
 			(commandLine.size() + 1u) * sizeof(wchar_t),
 			(workingDirectory.size() + 1u) * sizeof(wchar_t),
-			environment->size
+			environment.size
 		};
 
 		memoryStream::Writer payload(command.imagePathSize + command.commandLineSize + command.workingDirectorySize + command.environmentSize);
 		payload.Write(imagePath.data(), command.imagePathSize);
 		payload.Write(commandLine.data(), command.commandLineSize);
 		payload.Write(workingDirectory.data(), command.workingDirectorySize);
-		payload.Write(environment->data, environment->size);
+		payload.Write(environment.data, environment.size);
 
 		m_pipeClient->SendCommandAndWaitForAck(command, payload.GetData(), payload.GetSize());
 
-		process::DestroyEnvironment(environment);
+		Process::DestroyEnvironment(environment);
 	}
 
 	// handle commands until registration is finished
@@ -575,7 +590,9 @@ unsigned int ClientStartupThread::ThreadFunction(const std::wstring& processGrou
 	if (!m_successfulInit)
 	{
 		// process could not be registered, bail out
+		// BEGIN EPIC MOD
 		LC_ERROR_USER("Could not register live coding process");
+		// END EPIC MOD
 
 		// close the pipe and then wait for the helper threads to finish.
 		// closing the pipe bails out the helper threads.
@@ -603,7 +620,7 @@ unsigned int ClientStartupThread::ThreadFunction(const std::wstring& processGrou
 		DeleteAndNull(m_compilationEvent);
 		DeleteAndNull(m_pipeClientCS);
 
-		return 3u;
+		return Thread::ReturnValue(3u);
 	}
 
 	LC_LOG_USER("Successfully initialized, removing startup thread");
@@ -611,5 +628,5 @@ unsigned int ClientStartupThread::ThreadFunction(const std::wstring& processGrou
 	// helper threads are now allowed to run, we're finished with the pipe
 	m_startEvent->Signal();
 
-	return 0u;
+	return Thread::ReturnValue(0u);
 }
diff --git a/Engine/Source/Developer/Windows/LiveCoding/Private/External/LC_ClientStartupThread.h b/Engine/Source/Developer/Windows/LiveCoding/Private/External/LC_ClientStartupThread.h
--- a/Engine/Source/Developer/Windows/LiveCoding/Private/External/LC_ClientStartupThread.h
+++ b/Engine/Source/Developer/Windows/LiveCoding/Private/External/LC_ClientStartupThread.h
@@ -1,13 +1,18 @@
-// Copyright 2011-2019 Molecular Matters GmbH, all rights reserved.
+// Copyright 2011-2020 Molecular Matters GmbH, all rights reserved.
 
 #pragma once
 
+// BEGIN EPIC MOD
 #include "CoreTypes.h"
-#include "LC_Thread.h"
-#include "LC_Process.h"
+// END EPIC MOD
+#include "LC_ThreadTypes.h"
+#include "LC_ProcessTypes.h"
+#include "LC_NamedSharedMemoryTypes.h"
 #include "LC_RunMode.h"
+// BEGIN EPIC MOD
+#include <string>
+// END EPIC MOD
 
-class NamedSharedMemory;
 class ClientCommandThread;
 class ClientUserCommandThread;
 class DuplexPipeClient;
@@ -18,7 +23,7 @@ class Event;
 class ClientStartupThread
 {
 public:
-	explicit ClientStartupThread(HINSTANCE instance);
+	ClientStartupThread(void);
 	~ClientStartupThread(void);
 
 	// Spawns a thread that runs client initialization
@@ -69,21 +74,19 @@ public:
 	// END EPIC MOD
 
 private:
-	unsigned int ThreadFunction(const std::wstring& groupName, RunMode::Enum runMode);
+	Thread::ReturnValue ThreadFunction(const std::wstring& groupName, RunMode::Enum runMode);
 
-	HINSTANCE m_instance;
-
-	thread::Handle m_thread;
+	Thread::Handle m_thread;
 
 	// job object for associating spawned processes with main process the DLL is loaded into
 	HANDLE m_job;
 
 	// named shared memory for sharing the Live++ process ID between processes
-	NamedSharedMemory* m_sharedMemory;
+	Process::NamedSharedMemory* m_sharedMemory;
 
 	// main Live++ process. context may be empty in case we connected to an existing Live++ process
-	process::Context* m_mainProcessContext;
-	process::Handle m_processHandle;
+	Process::Context* m_mainProcessContext;
+	Process::Handle m_processHandle;
 
 	bool m_successfulInit;
 
diff --git a/Engine/Source/Developer/Windows/LiveCoding/Private/External/LC_ClientUserCommandThread.cpp b/Engine/Source/Developer/Windows/LiveCoding/Private/External/LC_ClientUserCommandThread.cpp
--- a/Engine/Source/Developer/Windows/LiveCoding/Private/External/LC_ClientUserCommandThread.cpp
+++ b/Engine/Source/Developer/Windows/LiveCoding/Private/External/LC_ClientUserCommandThread.cpp
@@ -1,5 +1,8 @@
-// Copyright 2011-2019 Molecular Matters GmbH, all rights reserved.
+// Copyright 2011-2020 Molecular Matters GmbH, all rights reserved.
 
+// BEGIN EPIC MOD
+//#include PCH_INCLUDE
+// END EPIC MOD
 #include "LC_ClientUserCommandThread.h"
 #include "LC_CommandMap.h"
 #include "LC_DuplexPipeClient.h"
@@ -7,12 +10,20 @@
 #include "LC_Event.h"
 #include "LC_Process.h"
 #include "LC_CriticalSection.h"
+// BEGIN EPIC MOD
+//#include "LC_ClientExceptionHandler.h"
+// END EPIC MOD
 #include "LC_StringUtil.h"
 #include "LC_Executable.h"
 #include "LC_MemoryStream.h"
+// BEGIN EPIC MOD
+#include "LC_Platform.h"
+#include LC_PLATFORM_INCLUDE(LC_Foundation)
+#include "LC_Thread.h"
 #include "LC_Logging.h"
 #include <deque>
 #include <unordered_set>
+// END EPIC MOD
 
 
 namespace
@@ -42,22 +53,31 @@ namespace
 		LC_DISABLE_MOVE_ASSIGNMENT(ProxyCommand);
 	};
 
+
 	// gathers module data for the given module, its import modules, the import's import modules, and so forth
-	static std::vector<commands::ModuleData> GatherImportModuleData(HMODULE mainModule)
+	// BEGIN EPIC MOD
+	static std::vector<commands::ModuleData> GatherImportModuleData(Windows::HMODULE mainModule)
+	// END EPIC MOD
 	{
+		// BEGIN EPIC MOD
 		std::vector<commands::ModuleData> moduleDatas;
+		// END EPIC MOD
 		moduleDatas.reserve(1024u);
 
+		// BEGIN EPIC MOD
 		std::unordered_set<std::wstring> loadedModules;
+		// END EPIC MOD
 		loadedModules.reserve(1024u);
 
-		std::vector<HMODULE> modules;
+		// BEGIN EPIC MOD
+		std::vector<Windows::HMODULE> modules;
+		// END EPIC MOD
 		modules.reserve(1024u);
 
 		modules.push_back(mainModule);
 		while (!modules.empty())
 		{
-			const HMODULE module = modules.back();
+			const Windows::HMODULE module = modules.back();
 			modules.pop_back();
 
 			// get the absolute path of the module.
@@ -81,7 +101,7 @@ namespace
 					moduleDatas.emplace_back(moduleData);
 				}
 
-				executable::Image* image = executable::OpenImage(fullPath, file::OpenMode::READ_ONLY);
+				executable::Image* image = executable::OpenImage(fullPath, Filesystem::OpenMode::READ);
 				if (image)
 				{
 					executable::ImageSectionDB* imageSections = executable::GatherImageSectionDB(image);
@@ -95,7 +115,9 @@ namespace
 								const char* importModulePath = importModules->modules[i].path;
 
 								// only add the import module if it is loaded into the process
-								HMODULE importModule = ::GetModuleHandleA(importModulePath);
+								// BEGIN EPIC MOD
+								Windows::HMODULE importModule = ::GetModuleHandleA(importModulePath);
+								// END EPIC MOD
 								if (importModule)
 								{
 									modules.push_back(importModule);
@@ -138,8 +160,9 @@ bool ClientUserCommandThread::BaseCommand::ExpectsResponse(void) const
 	return m_expectResponse;
 }
 
+
 ClientUserCommandThread::ClientUserCommandThread(DuplexPipeClient* pipeClient, DuplexPipeClient* exceptionPipeClient)
-	: m_thread(INVALID_HANDLE_VALUE)
+	: m_thread(Thread::INVALID_HANDLE)
 	, m_processGroupName()
 	, m_pipe(pipeClient)
 	, m_exceptionPipe(exceptionPipeClient)
@@ -155,23 +178,25 @@ ClientUserCommandThread::~ClientUserCommandThread(void)
 }
 
 
-unsigned int ClientUserCommandThread::Start(const std::wstring& processGroupName, Event* waitForStartEvent, CriticalSection* pipeAccessCS)
+Thread::Id ClientUserCommandThread::Start(const std::wstring& processGroupName, Event* waitForStartEvent, CriticalSection* pipeAccessCS)
 {
 	m_processGroupName = processGroupName;
 
 	// spawn a thread that does the work
-	m_thread = thread::Create("Live coding user commands", 128u * 1024u, &ClientUserCommandThread::ThreadFunction, this, waitForStartEvent, pipeAccessCS);
+	// BEGIN EPIC MOD
+	m_thread = Thread::CreateFromMemberFunction("Live coding user commands", 128u * 1024u, this, &ClientUserCommandThread::ThreadFunction, waitForStartEvent, pipeAccessCS);
+	// END EPIC MOD
 
-	return thread::GetId(m_thread);
+	return Thread::GetId(m_thread);
 }
 
 
 void ClientUserCommandThread::Join(void)
 {
-	if (m_thread != INVALID_HANDLE_VALUE)
+	if (m_thread != Thread::INVALID_HANDLE)
 	{
-		thread::Join(m_thread);
-		thread::Close(m_thread);
+		Thread::Join(m_thread);
+		Thread::Close(m_thread);
 	}
 }
 
@@ -184,12 +209,16 @@ void* ClientUserCommandThread::EnableModule(const wchar_t* nameOfExeOrDll)
 
 void* ClientUserCommandThread::EnableModules(const wchar_t* namesOfExeOrDll[], unsigned int count)
 {
+	// BEGIN EPIC MOD
 	std::vector<Windows::HMODULE> loadedModules;
+	// END EPIC MOD
 	loadedModules.reserve(count);
 
 	for (unsigned int i = 0u; i < count; ++i)
 	{
+		// BEGIN EPIC MOD
 		Windows::HMODULE module = ::GetModuleHandleW(namesOfExeOrDll[i]);
+		// END EPIC MOD
 		if (module)
 		{
 			loadedModules.push_back(module);
@@ -208,13 +237,15 @@ void* ClientUserCommandThread::EnableModules(const wchar_t* namesOfExeOrDll[], u
 	}
 
 	ProxyCommand<commands::EnableModules>* proxy = new ProxyCommand<commands::EnableModules>(true, sizeof(commands::ModuleData) * loadedModuleCount);
-	proxy->m_command.processId = process::GetId();
+	proxy->m_command.processId = Process::Current::GetId();
 	proxy->m_command.moduleCount = static_cast<unsigned int>(loadedModuleCount);
 	proxy->m_command.token = new Event(nullptr, Event::Type::AUTO_RESET);
 
 	for (size_t i = 0u; i < loadedModuleCount; ++i)
 	{
+		// BEGIN EPIC MOD
 		Windows::HMODULE module = loadedModules[i];
+		// END EPIC MOD
 
 		commands::ModuleData moduleData = {};
 		moduleData.base = module;
@@ -228,20 +259,25 @@ void* ClientUserCommandThread::EnableModules(const wchar_t* namesOfExeOrDll[], u
 	return proxy->m_command.token;
 }
 
+
 void* ClientUserCommandThread::EnableAllModules(const wchar_t* nameOfExeOrDll)
 {
-	HMODULE module = ::GetModuleHandleW(nameOfExeOrDll);
+	// BEGIN EPIC MOD
+	Windows::HMODULE module = ::GetModuleHandleW(nameOfExeOrDll);
+	// END EPIC MOD
 	if (!module)
 	{
 		LC_ERROR_USER("Cannot enable module %S because it is not loaded by this process.", nameOfExeOrDll);
 		return nullptr;
 	}
 
+	// BEGIN EPIC MOD
 	const std::vector<commands::ModuleData>& allModuleData = GatherImportModuleData(module);
+	// END EPIC MOD
 	const size_t moduleCount = allModuleData.size();
 
 	ProxyCommand<commands::EnableModules>* proxy = new ProxyCommand<commands::EnableModules>(true, sizeof(commands::ModuleData) * moduleCount);
-	proxy->m_command.processId = process::GetId();
+	proxy->m_command.processId = Process::Current::GetId();
 	proxy->m_command.moduleCount = static_cast<unsigned int>(moduleCount);
 	proxy->m_command.token = new Event(nullptr, Event::Type::AUTO_RESET);
 
@@ -265,12 +301,16 @@ void* ClientUserCommandThread::DisableModule(const wchar_t* nameOfExeOrDll)
 
 void* ClientUserCommandThread::DisableModules(const wchar_t* namesOfExeOrDll[], unsigned int count)
 {
+	// BEGIN EPIC MOD
 	std::vector<Windows::HMODULE> loadedModules;
+	// END EPIC MOD
 	loadedModules.reserve(count);
 
 	for (unsigned int i = 0u; i < count; ++i)
 	{
+		// BEGIN EPIC MOD
 		Windows::HMODULE module = ::GetModuleHandleW(namesOfExeOrDll[i]);
+		// END EPIC MOD
 		if (module)
 		{
 			loadedModules.push_back(module);
@@ -289,13 +329,15 @@ void* ClientUserCommandThread::DisableModules(const wchar_t* namesOfExeOrDll[],
 	}
 
 	ProxyCommand<commands::DisableModules>* proxy = new ProxyCommand<commands::DisableModules>(true, sizeof(commands::ModuleData) * loadedModuleCount);
-	proxy->m_command.processId = process::GetId();
+	proxy->m_command.processId = Process::Current::GetId();
 	proxy->m_command.moduleCount = static_cast<unsigned int>(loadedModuleCount);
 	proxy->m_command.token = new Event(nullptr, Event::Type::AUTO_RESET);
 
 	for (size_t i = 0u; i < loadedModuleCount; ++i)
 	{
-		HMODULE module = loadedModules[i];
+		// BEGIN EPIC MOD
+		Windows::HMODULE module = loadedModules[i];
+		// END EPIC MOD
 
 		commands::ModuleData moduleData = {};
 		moduleData.base = module;
@@ -312,18 +354,22 @@ void* ClientUserCommandThread::DisableModules(const wchar_t* namesOfExeOrDll[],
 
 void* ClientUserCommandThread::DisableAllModules(const wchar_t* nameOfExeOrDll)
 {
-	HMODULE module = ::GetModuleHandleW(nameOfExeOrDll);
+	// BEGIN EPIC MOD
+	Windows::HMODULE module = ::GetModuleHandleW(nameOfExeOrDll);
+	// END EPIC MOD
 	if (!module)
 	{
 		LC_ERROR_USER("Cannot disable module %S because it is not loaded by this process.", nameOfExeOrDll);
 		return nullptr;
 	}
 
+	// BEGIN EPIC MOD
 	const std::vector<commands::ModuleData>& allModuleData = GatherImportModuleData(module);
+	// END EPIC MOD
 	const size_t moduleCount = allModuleData.size();
 
 	ProxyCommand<commands::DisableModules>* proxy = new ProxyCommand<commands::DisableModules>(true, sizeof(commands::ModuleData) * moduleCount);
-	proxy->m_command.processId = process::GetId();
+	proxy->m_command.processId = Process::Current::GetId();
 	proxy->m_command.moduleCount = static_cast<unsigned int>(moduleCount);
 	proxy->m_command.token = new Event(nullptr, Event::Type::AUTO_RESET);
 
@@ -343,7 +389,7 @@ void ClientUserCommandThread::WaitForToken(void* token)
 {
 	Event* event = static_cast<Event*>(token);
 
-	if (m_thread != INVALID_HANDLE_VALUE)
+	if (m_thread != Thread::INVALID_HANDLE)
 	{
 		// thread was successfully initialized, wait until the command has been executed in the queue
 		event->Wait();
@@ -468,7 +514,7 @@ void ClientUserCommandThread::SetActive(bool active)
 void ClientUserCommandThread::SetBuildArguments(const wchar_t* arguments)
 {
 	ProxyCommand<commands::SetBuildArguments>* proxy = new ProxyCommand<commands::SetBuildArguments>(false, 0u);
-	proxy->m_command.processId = process::GetId();
+	proxy->m_command.processId = Process::Current::GetId();
 	wcscpy_s(proxy->m_command.arguments, arguments);
 
 	PushUserCommand(proxy);
@@ -479,7 +525,7 @@ void ClientUserCommandThread::SetBuildArguments(const wchar_t* arguments)
 void* ClientUserCommandThread::EnableLazyLoadedModule(const wchar_t* fileName, Windows::HMODULE moduleBase)
 {
 	ProxyCommand<commands::EnableLazyLoadedModule>* proxy = new ProxyCommand<commands::EnableLazyLoadedModule>(true, 0u);
-	proxy->m_command.processId = process::GetId();
+	proxy->m_command.processId = Process::Current::GetId();
 	wcscpy_s(proxy->m_command.fileName, fileName);
 	proxy->m_command.moduleBase = moduleBase;
 	proxy->m_command.token = new Event(nullptr, Event::Type::AUTO_RESET);
@@ -499,10 +545,10 @@ void ClientUserCommandThread::InstallExceptionHandler(void)
 }
 
 
-ClientUserCommandThread::ExceptionResult ClientUserCommandThread::HandleException(EXCEPTION_RECORD* exception, CONTEXT* context, unsigned int threadId)
+ClientUserCommandThread::ExceptionResult ClientUserCommandThread::HandleException(EXCEPTION_RECORD* exception, CONTEXT* context, Thread::Id threadId)
 {
 	commands::HandleException serverCommand;
-	serverCommand.processId = process::GetId();
+	serverCommand.processId = Process::Current::GetId();
 	serverCommand.threadId = threadId;
 	serverCommand.exception = *exception;
 	serverCommand.context = *context;
@@ -551,49 +597,7 @@ ClientUserCommandThread::BaseCommand* ClientUserCommandThread::PopUserCommand(vo
 }
 
 
-// BEGIN EPIC MOD - Temporarily release lock to prevent hangs
-class LeaveableScopedLock
-{
-public:
-	explicit LeaveableScopedLock(CriticalSection* cs)
-		: m_cs(cs)
-		, m_hasLock(true)
-	{
-		cs->Enter();
-	}
-
-	void Enter()
-	{
-		if (!m_hasLock)
-		{
-			m_cs->Enter();
-			m_hasLock = true;
-		}
-	}
-
-	void Leave()
-	{
-		if (m_hasLock)
-		{
-			m_cs->Leave();
-			m_hasLock = false;
-		}
-	}
-
-	~LeaveableScopedLock(void)
-	{
-		Leave();
-	}
-
-private:
-	CriticalSection* m_cs;
-	bool m_hasLock;
-};
-// END EPIC MOD
-
-
-
-unsigned int ClientUserCommandThread::ThreadFunction(Event* waitForStartEvent, CriticalSection* pipeAccessCS)
+Thread::ReturnValue ClientUserCommandThread::ThreadFunction(Event* waitForStartEvent, CriticalSection* pipeAccessCS)
 {
 	// wait until we get the signal that the thread can start
 	waitForStartEvent->Wait();
@@ -620,7 +624,7 @@ unsigned int ClientUserCommandThread::ThreadFunction(Event* waitForStartEvent, C
 			// // no new item available, bail out
 			// exceptionHandler::Unregister();
 			// END EPIC MOD - Using internal CrashReporter
-			return 2u;
+			return Thread::ReturnValue(2u);
 		}
 
 		if (!m_pipe->IsValid())
@@ -629,7 +633,7 @@ unsigned int ClientUserCommandThread::ThreadFunction(Event* waitForStartEvent, C
 			// // pipe was closed or is broken, bail out
 			// exceptionHandler::Unregister();
 			// END EPIC MOD - Using internal CrashReporter
-			return 1u;
+			return Thread::ReturnValue(1u);
 		}
 
 		// lock critical section for accessing the pipe.
@@ -647,5 +651,7 @@ unsigned int ClientUserCommandThread::ThreadFunction(Event* waitForStartEvent, C
 		delete command;
 	}
 
-	return 0u;
+	// BEGIN EPIC MOD
+	return Thread::ReturnValue(0u);
+	// END EPIC MOD
 }
diff --git a/Engine/Source/Developer/Windows/LiveCoding/Private/External/LC_ClientUserCommandThread.h b/Engine/Source/Developer/Windows/LiveCoding/Private/External/LC_ClientUserCommandThread.h
--- a/Engine/Source/Developer/Windows/LiveCoding/Private/External/LC_ClientUserCommandThread.h
+++ b/Engine/Source/Developer/Windows/LiveCoding/Private/External/LC_ClientUserCommandThread.h
@@ -1,13 +1,17 @@
-// Copyright 2011-2019 Molecular Matters GmbH, all rights reserved.
+// Copyright 2011-2020 Molecular Matters GmbH, all rights reserved.
 
 #pragma once
 
+// BEGIN EPIC MOD
 #include "CoreTypes.h"
-#include "LC_Thread.h"
+// END EPIC MOD
+#include "LC_ThreadTypes.h"
 #include "LC_CriticalSection.h"
 #include "LC_Semaphore.h"
+// BEGIN EPIC MOD
 #include <string>
 #include <deque>
+// END EPIC MOD
 
 class DuplexPipe;
 class DuplexPipeClient;
@@ -44,8 +48,7 @@ public:
 	~ClientUserCommandThread(void);
 
 	// Starts the thread that takes care of handling incoming commands on the pipe.
-	// Returns the thread ID.
-	unsigned int Start(const std::wstring& processGroupName, Event* waitForStartEvent, CriticalSection* pipeAccessCS);
+	Thread::Id Start(const std::wstring& processGroupName, Event* waitForStartEvent, CriticalSection* pipeAccessCS);
 
 	// Joins this thread.
 	void Join(void);
@@ -90,7 +93,7 @@ public:
 	void ApplySettingString(const char* settingName, const wchar_t* value);
 
 	void InstallExceptionHandler(void);
-	ExceptionResult HandleException(EXCEPTION_RECORD* exception, CONTEXT* context, unsigned int threadId);
+	ExceptionResult HandleException(EXCEPTION_RECORD* exception, CONTEXT* context, Thread::Id threadId);
 	void End(void);
 
 private:
@@ -101,9 +104,9 @@ private:
 	// blocks until a command becomes available.
 	BaseCommand* PopUserCommand(void);
 
-	unsigned int ThreadFunction(Event* waitForStartEvent, CriticalSection* pipeAccessCS);
+	Thread::ReturnValue ThreadFunction(Event* waitForStartEvent, CriticalSection* pipeAccessCS);
 
-	thread::Handle m_thread;
+	Thread::Handle m_thread;
 	std::wstring m_processGroupName;
 	DuplexPipeClient* m_pipe;
 	DuplexPipeClient* m_exceptionPipe;
diff --git a/Engine/Source/Developer/Windows/LiveCoding/Private/External/LC_CommandMap.cpp b/Engine/Source/Developer/Windows/LiveCoding/Private/External/LC_CommandMap.cpp
--- a/Engine/Source/Developer/Windows/LiveCoding/Private/External/LC_CommandMap.cpp
+++ b/Engine/Source/Developer/Windows/LiveCoding/Private/External/LC_CommandMap.cpp
@@ -1,5 +1,8 @@
-// Copyright 2011-2019 Molecular Matters GmbH, all rights reserved.
+// Copyright 2011-2020 Molecular Matters GmbH, all rights reserved.
 
+// BEGIN EPIC MOD
+//#include PCH_INCLUDE
+// END EPIC MOD
 #include "LC_CommandMap.h"
 #include "LC_DuplexPipe.h"
 
diff --git a/Engine/Source/Developer/Windows/LiveCoding/Private/External/LC_CommandMap.h b/Engine/Source/Developer/Windows/LiveCoding/Private/External/LC_CommandMap.h
--- a/Engine/Source/Developer/Windows/LiveCoding/Private/External/LC_CommandMap.h
+++ b/Engine/Source/Developer/Windows/LiveCoding/Private/External/LC_CommandMap.h
@@ -1,11 +1,14 @@
-// Copyright 2011-2019 Molecular Matters GmbH, all rights reserved.
+// Copyright 2011-2020 Molecular Matters GmbH, all rights reserved.
 
 #pragma once
 
+// BEGIN EPIC MOD
 #include "CoreTypes.h"
+// END EPIC MOD
 #include "LC_Commands.h"
+// BEGIN EPIC MOD
 #include "LC_DuplexPipe.h"
-
+// END EPIC MOD
 
 class DuplexPipe;
 
diff --git a/Engine/Source/Developer/Windows/LiveCoding/Private/External/LC_Commands.h b/Engine/Source/Developer/Windows/LiveCoding/Private/External/LC_Commands.h
--- a/Engine/Source/Developer/Windows/LiveCoding/Private/External/LC_Commands.h
+++ b/Engine/Source/Developer/Windows/LiveCoding/Private/External/LC_Commands.h
@@ -1,15 +1,15 @@
-// Copyright 2011-2019 Molecular Matters GmbH, all rights reserved.
+// Copyright 2011-2020 Molecular Matters GmbH, all rights reserved.
 
 #pragma once
 
+// BEGIN EPIC MOD
 #include "CoreTypes.h"
 #include "../../../LiveCodingServer/Private/External/LC_Hook.h"
 #include "Windows/WindowsHWrapper.h"
+// END EPIC MOD
+#include "LC_ProcessTypes.h"
+#include "LC_ThreadTypes.h"
 
-namespace hook
-{
-	typedef void (*Function)(void);
-}
 
 namespace commands
 {
@@ -22,13 +22,17 @@ namespace commands
 	struct ModuleData
 	{
 		void* base;
+		// BEGIN EPIC MOD
 		wchar_t path[WINDOWS_MAX_PATH];
+		// END EPIC MOD
 	};
 
 	// acknowledge that a command has been received
 	struct Acknowledge
 	{
+		// BEGIN EPIC MOD
 		static const uint32_t ID = 100;
+		// END EPIC MOD
 	};
 
 	// register a process with Live++
@@ -37,9 +41,9 @@ namespace commands
 		static const uint32_t ID = Acknowledge::ID + 1u;
 
 		void* processBase;
-		unsigned int processId;				// current process ID
-		unsigned int restartedProcessId;	// process ID of the previous, restarted process. 0 if non-existent
-		unsigned int threadId;				// thread ID of Live++ thread running in host
+		Process::Id processId;				// current process ID
+		Process::Id restartedProcessId;		// process ID of the previous, restarted process. 0 if non-existent
+		Thread::Id threadId;				// thread ID of Live++ thread running in host
 		const void* jumpToSelf;				// address of jump-to-self instruction in host
 
 		size_t imagePathSize;
@@ -63,7 +67,7 @@ namespace commands
 	{
 		static const uint32_t ID = RegisterProcessFinished::ID + 1u;
 
-		unsigned int processId;
+		Process::Id processId;
 		unsigned int moduleCount;
 		void* token;
 
@@ -83,7 +87,7 @@ namespace commands
 	{
 		static const uint32_t ID = EnableModulesFinished::ID + 1u;
 
-		unsigned int processId;
+		Process::Id processId;
 		unsigned int moduleCount;
 		void* token;
 
@@ -125,7 +129,9 @@ namespace commands
 	{
 		static const uint32_t ID = CallHooks::ID + 1u;
 
+		// BEGIN EPIC MOD
 		wchar_t path[WINDOWS_MAX_PATH];
+		// END EPIC MOD
 	};
 
 	// returns info about a loaded DLL to Live++
@@ -133,7 +139,9 @@ namespace commands
 	{
 		static const uint32_t ID = LoadPatch::ID + 1u;
 
+		// BEGIN EPIC MOD
 		Windows::HMODULE module;
+		// END EPIC MOD
 	};
 
 	// tell the DLL to unload a DLL
@@ -141,7 +149,9 @@ namespace commands
 	{
 		static const uint32_t ID = LoadPatchInfo::ID + 1u;
 
+		// BEGIN EPIC MOD
 		Windows::HMODULE module;
+		// END EPIC MOD
 	};
 
 	// tell the DLL to call the entry point of a DLL
@@ -216,8 +226,8 @@ namespace commands
 	{
 		static const uint32_t ID = BuildPatch::ID + 1u;
 
-		unsigned int processId;
-		unsigned int threadId;
+		Process::Id processId;
+		Thread::Id threadId;
 		EXCEPTION_RECORD exception;
 		CONTEXT context;
 		CONTEXT* clientContextPtr;
@@ -264,7 +274,7 @@ namespace commands
 	{
 		static const uint32_t ID = SetActive::ID + 1u;
 
-		unsigned int processId;
+		Process::Id processId;
 		wchar_t arguments[1024];
 	};
 	// END EPIC MOD
@@ -274,7 +284,7 @@ namespace commands
 	{
 		static const uint32_t ID = SetBuildArguments::ID + 1u;
 
-		unsigned int processId;
+		Process::Id processId;
 		wchar_t fileName[260];
 		Windows::HMODULE moduleBase;
 		void* token;
diff --git a/Engine/Source/Developer/Windows/LiveCoding/Private/External/LC_CriticalSection.cpp b/Engine/Source/Developer/Windows/LiveCoding/Private/External/LC_CriticalSection.cpp
--- a/Engine/Source/Developer/Windows/LiveCoding/Private/External/LC_CriticalSection.cpp
+++ b/Engine/Source/Developer/Windows/LiveCoding/Private/External/LC_CriticalSection.cpp
@@ -1,27 +1,38 @@
-// Copyright 2011-2019 Molecular Matters GmbH, all rights reserved.
+// Copyright 2011-2020 Molecular Matters GmbH, all rights reserved.
 
+// BEGIN EPIC MOD
+//#include PCH_INCLUDE
+// END EPIC MOD
 #include "LC_CriticalSection.h"
 
 
 CriticalSection::CriticalSection(void)
 {
+	// BEGIN EPIC MOD
 	Windows::InitializeCriticalSection(&m_cs);
+	// END EPIC MOD
 }
 
 
 CriticalSection::~CriticalSection(void)
 {
+	// BEGIN EPIC MOD
 	Windows::DeleteCriticalSection(&m_cs);
+	// END EPIC MOD
 }
 
 
 void CriticalSection::Enter(void)
 {
+	// BEGIN EPIC MOD
 	Windows::EnterCriticalSection(&m_cs);
+	// END EPIC MOD
 }
 
 
 void CriticalSection::Leave(void)
 {
+	// BEGIN EPIC MOD
 	Windows::LeaveCriticalSection(&m_cs);
+	// END EPIC MOD
 }
diff --git a/Engine/Source/Developer/Windows/LiveCoding/Private/External/LC_CriticalSection.h b/Engine/Source/Developer/Windows/LiveCoding/Private/External/LC_CriticalSection.h
--- a/Engine/Source/Developer/Windows/LiveCoding/Private/External/LC_CriticalSection.h
+++ b/Engine/Source/Developer/Windows/LiveCoding/Private/External/LC_CriticalSection.h
@@ -1,10 +1,13 @@
-// Copyright 2011-2019 Molecular Matters GmbH, all rights reserved.
+// Copyright 2011-2020 Molecular Matters GmbH, all rights reserved.
 
 #pragma once
 
+// BEGIN EPIC MOD
 #include "CoreTypes.h"
 #include "Windows/MinimalWindowsApi.h"
 #include "LC_Platform.h"
+#include LC_PLATFORM_INCLUDE(LC_Foundation)
+// END EPIC MOD
 
 class CriticalSection
 {
@@ -39,5 +42,7 @@ public:
 	};
 
 private:
+	// BEGIN EPIC MOD
 	Windows::CRITICAL_SECTION m_cs;
+	// END EPIC MOD
 };
diff --git a/Engine/Source/Developer/Windows/LiveCoding/Private/External/LC_DuplexPipe.cpp b/Engine/Source/Developer/Windows/LiveCoding/Private/External/LC_DuplexPipe.cpp
--- a/Engine/Source/Developer/Windows/LiveCoding/Private/External/LC_DuplexPipe.cpp
+++ b/Engine/Source/Developer/Windows/LiveCoding/Private/External/LC_DuplexPipe.cpp
@@ -1,8 +1,12 @@
-// Copyright 2011-2019 Molecular Matters GmbH, all rights reserved.
+// Copyright 2011-2020 Molecular Matters GmbH, all rights reserved.
 
+// BEGIN EPIC MOD
+//#include PCH_INCLUDE
+// END EPIC MOD
 #include "LC_DuplexPipe.h"
+// BEGIN EPIC MOD
 #include "LC_Logging.h"
-
+// END EPIC MOD
 
 DuplexPipe::DuplexPipe(void)
 	: m_pipe(INVALID_HANDLE_VALUE)
diff --git a/Engine/Source/Developer/Windows/LiveCoding/Private/External/LC_DuplexPipe.h b/Engine/Source/Developer/Windows/LiveCoding/Private/External/LC_DuplexPipe.h
--- a/Engine/Source/Developer/Windows/LiveCoding/Private/External/LC_DuplexPipe.h
+++ b/Engine/Source/Developer/Windows/LiveCoding/Private/External/LC_DuplexPipe.h
@@ -1,10 +1,15 @@
-// Copyright 2011-2019 Molecular Matters GmbH, all rights reserved.
+// Copyright 2011-2020 Molecular Matters GmbH, all rights reserved.
 
 #pragma once
 
+// BEGIN EPIC MOD
 #include "CoreTypes.h"
+// END EPIC MOD
 #include "LC_Commands.h"
+// BEGIN EPIC MOD
+#include "LC_Assert.h"
 #include "LC_Platform.h"
+// END EPIC MOD
 
 class DuplexPipe
 {
diff --git a/Engine/Source/Developer/Windows/LiveCoding/Private/External/LC_DuplexPipeClient.cpp b/Engine/Source/Developer/Windows/LiveCoding/Private/External/LC_DuplexPipeClient.cpp
--- a/Engine/Source/Developer/Windows/LiveCoding/Private/External/LC_DuplexPipeClient.cpp
+++ b/Engine/Source/Developer/Windows/LiveCoding/Private/External/LC_DuplexPipeClient.cpp
@@ -1,8 +1,12 @@
 // Copyright 2011-2019 Molecular Matters GmbH, all rights reserved.
 
+// BEGIN EPIC MOD
+//#include PCH_INCLUDE
+// END EPIC MOD
 #include "LC_DuplexPipeClient.h"
+// BEGIN EPIC MOD
 #include "LC_Logging.h"
-
+// END EPIC MOD
 
 bool DuplexPipeClient::Connect(const wchar_t* name)
 {
diff --git a/Engine/Source/Developer/Windows/LiveCoding/Private/External/LC_DuplexPipeClient.h b/Engine/Source/Developer/Windows/LiveCoding/Private/External/LC_DuplexPipeClient.h
--- a/Engine/Source/Developer/Windows/LiveCoding/Private/External/LC_DuplexPipeClient.h
+++ b/Engine/Source/Developer/Windows/LiveCoding/Private/External/LC_DuplexPipeClient.h
@@ -1,8 +1,10 @@
-// Copyright 2011-2019 Molecular Matters GmbH, all rights reserved.
+// Copyright 2011-2020 Molecular Matters GmbH, all rights reserved.
 
 #pragma once
 
+// BEGIN EPIC MOD
 #include "CoreTypes.h"
+// END EPIC MOD
 #include "LC_DuplexPipe.h"
 
 
diff --git a/Engine/Source/Developer/Windows/LiveCoding/Private/External/LC_Environment.cpp b/Engine/Source/Developer/Windows/LiveCoding/Private/External/LC_Environment.cpp
--- a/Engine/Source/Developer/Windows/LiveCoding/Private/External/LC_Environment.cpp
+++ b/Engine/Source/Developer/Windows/LiveCoding/Private/External/LC_Environment.cpp
@@ -1,13 +1,19 @@
-// Copyright 2011-2019 Molecular Matters GmbH, all rights reserved.
+// Copyright 2011-2020 Molecular Matters GmbH, all rights reserved.
 
+// BEGIN EPIC MOD
+//#include PCH_INCLUDE
+// END EPIC MOD
 #include "LC_Environment.h"
+// BEGIN EPIC MOD
 #include "LC_Logging.h"
 #include "Windows/WindowsHWrapper.h"
-
+// END EPIC MOD
 
 namespace
 {
+	// BEGIN EPIC MOD
 	static const DWORD MAX_ENVIRONMENT_VARIABLE_SIZE = 1024u;
+	// END EPIC MOD
 }
 
 
@@ -34,9 +40,13 @@ void environment::SetVariable(const wchar_t* variable, const wchar_t* value)
 std::wstring environment::GetVariable(const wchar_t* variable, const wchar_t* defaultValue)
 {
 	std::wstring value;
+	// BEGIN EPIC MOD
 	value.resize(MAX_ENVIRONMENT_VARIABLE_SIZE);
+	// END EPIC MOD
 
+	// BEGIN EPIC MOD
 	const DWORD result = ::GetEnvironmentVariableW(variable, &value[0], MAX_ENVIRONMENT_VARIABLE_SIZE);
+	// END EPIC _MOD
 	if (result == 0)
 	{
 		// if a default value was provided, return that instead
diff --git a/Engine/Source/Developer/Windows/LiveCoding/Private/External/LC_Event.cpp b/Engine/Source/Developer/Windows/LiveCoding/Private/External/LC_Event.cpp
--- a/Engine/Source/Developer/Windows/LiveCoding/Private/External/LC_Event.cpp
+++ b/Engine/Source/Developer/Windows/LiveCoding/Private/External/LC_Event.cpp
@@ -1,12 +1,19 @@
-// Copyright 2011-2019 Molecular Matters GmbH, all rights reserved.
+// Copyright 2011-2020 Molecular Matters GmbH, all rights reserved.
 
+// BEGIN EPIC MOD
+//#include PCH_INCLUDE
+// END EPIC MOD
 #include "LC_Event.h"
+// BEGIN EPIC MOD
 #include "LC_Logging.h"
 #include "Windows/WindowsHWrapper.h"
+// END EPIC MOD
 
 
 Event::Event(const wchar_t* name, Type::Enum type)
+	// BEGIN EPIC MOD
 	: m_event(::CreateEventW(NULL, (type == Type::MANUAL_RESET) ? Windows::TRUE : Windows::FALSE, Windows::FALSE, name))
+	// END EPIC MOD
 {
 	const DWORD error = ::GetLastError();
 	if (m_event == NULL)
diff --git a/Engine/Source/Developer/Windows/LiveCoding/Private/External/LC_Event.h b/Engine/Source/Developer/Windows/LiveCoding/Private/External/LC_Event.h
--- a/Engine/Source/Developer/Windows/LiveCoding/Private/External/LC_Event.h
+++ b/Engine/Source/Developer/Windows/LiveCoding/Private/External/LC_Event.h
@@ -1,9 +1,11 @@
-// Copyright 2011-2019 Molecular Matters GmbH, all rights reserved.
+// Copyright 2011-2020 Molecular Matters GmbH, all rights reserved.
 
 #pragma once
 
+// BEGIN EPIC MOD
 #include "CoreTypes.h"
 #include "Windows/MinimalWindowsApi.h"
+// END EPIC MOD
 
 // named/unnamed event.
 // acts process-wide if given a name.
@@ -39,5 +41,7 @@ public:
 	bool TryWait(void);
 
 private:
+	// BEGIN EPIC MOD
 	Windows::HANDLE m_event;
+	// END EPIC MOD
 };
diff --git a/Engine/Source/Developer/Windows/LiveCoding/Private/External/LC_Executable.cpp b/Engine/Source/Developer/Windows/LiveCoding/Private/External/LC_Executable.cpp
--- a/Engine/Source/Developer/Windows/LiveCoding/Private/External/LC_Executable.cpp
+++ b/Engine/Source/Developer/Windows/LiveCoding/Private/External/LC_Executable.cpp
@@ -1,9 +1,15 @@
-// Copyright 2011-2019 Molecular Matters GmbH, all rights reserved.
+// Copyright 2011-2020 Molecular Matters GmbH, all rights reserved.
 
+// BEGIN EPIC MOD
+//#include PCH_INCLUDE
+// END EPIC MOD
 #include "LC_Executable.h"
+#include "LC_MemoryMappedFile.h"
+// BEGIN EPIC MOD
+#include "LC_Assert.h"
 #include "LC_Logging.h"
 #include <algorithm>
-
+// END EPIC MOD
 
 // this is the default DLL entry point, taken from the CRT. we don't need it, but can extract its signature.
 extern "C" BOOL WINAPI _DllMainCRTStartup(
@@ -23,10 +29,10 @@ namespace detail
 
 	static const IMAGE_NT_HEADERS* GetNtHeader(const executable::Image* image)
 	{
-		void* base = image->base;
+		const void* base = Filesystem::GetMemoryMappedFileData(image);
 
 		// PE image start with a DOS header
-		const IMAGE_DOS_HEADER* dosHeader = static_cast<IMAGE_DOS_HEADER*>(base);
+		const IMAGE_DOS_HEADER* dosHeader = static_cast<const IMAGE_DOS_HEADER*>(base);
 		if (dosHeader->e_magic != IMAGE_DOS_SIGNATURE)
 		{
 			LC_ERROR_USER("Image has unknown file format");
@@ -89,7 +95,7 @@ executable::Header executable::GetHeader(const Image* image)
 		return executable::Header {};
 	}
 
-	const uint64_t sizeOnDisk = file::GetSizeOnDisk(image);
+	const uint64_t sizeOnDisk = Filesystem::GetMemoryMappedFileSizeOnDisk(image);
 	return executable::Header { sizeOnDisk, ntHeader->FileHeader, ntHeader->OptionalHeader };
 }
 
@@ -142,33 +148,33 @@ uint32_t executable::RvaToFileOffset(const ImageSectionDB* database, uint32_t rv
 
 void executable::ReadFromFileOffset(const Image* image, uint32_t offset, void* destination, size_t byteCount)
 {
-	const void* address = pointer::Offset<const void*>(image->base, offset);
+	const void* address = pointer::Offset<const void*>(Filesystem::GetMemoryMappedFileData(image), offset);
 	memcpy(destination, address, byteCount);
 }
 
 
 void executable::WriteToFileOffset(Image* image, uint32_t offset, const void* source, size_t byteCount)
 {
-	void* address = pointer::Offset<void*>(image->base, offset);
+	void* address = pointer::Offset<void*>(Filesystem::GetMemoryMappedFileData(image), offset);
 	memcpy(address, source, byteCount);
 }
 
 
-executable::Image* executable::OpenImage(const wchar_t* filename, file::OpenMode::Enum openMode)
+executable::Image* executable::OpenImage(const wchar_t* filename, Filesystem::OpenMode::Enum openMode)
 {
-	return file::Open(filename, openMode);
+	return Filesystem::OpenMemoryMappedFile(filename, openMode);
 }
 
 
 void executable::CloseImage(Image*& image)
 {
-	file::Close(image);
+	Filesystem::CloseMemoryMappedFile(image);
 }
 
 
 void executable::RebaseImage(Image* image, PreferredBase preferredBase)
 {
-	void* base = image->base;
+	void* base = Filesystem::GetMemoryMappedFileData(image);
 
 	IMAGE_NT_HEADERS* ntHeader = detail::GetNtHeader(image);
 	if (!ntHeader)
@@ -327,11 +333,11 @@ executable::ImportModuleDB* executable::GatherImportModuleDB(const Image* image,
 
 	database->modules.reserve(count);
 
-	const IMAGE_IMPORT_DESCRIPTOR* importModules = pointer::Offset<const IMAGE_IMPORT_DESCRIPTOR*>(image->base, baseImportModule);
+	const IMAGE_IMPORT_DESCRIPTOR* importModules = pointer::Offset<const IMAGE_IMPORT_DESCRIPTOR*>(Filesystem::GetMemoryMappedFileData(image), baseImportModule);
 	while (importModules->Name != 0)
 	{
 		const DWORD fileOffset = RvaToFileOffset(imageSections, importModules->Name);
-		const char* importModuleName = pointer::Offset<const char*>(image->base, fileOffset);
+		const char* importModuleName = pointer::Offset<const char*>(Filesystem::GetMemoryMappedFileData(image), fileOffset);
 
 		ImportModule module;
 		strcpy_s(module.path, importModuleName);
@@ -379,7 +385,7 @@ executable::PdbInfo* executable::GatherPdbInfo(const Image* image, const ImageSe
 		return nullptr;
 	}
 
-	const IMAGE_DEBUG_DIRECTORY* debugDirectory = pointer::Offset<const IMAGE_DEBUG_DIRECTORY*>(image->base, baseDebugDirectory);
+	const IMAGE_DEBUG_DIRECTORY* debugDirectory = pointer::Offset<const IMAGE_DEBUG_DIRECTORY*>(Filesystem::GetMemoryMappedFileData(image), baseDebugDirectory);
 	for (size_t i=0u; i < count; ++i)
 	{
 		// we are only interested in PDB files
@@ -395,7 +401,7 @@ executable::PdbInfo* executable::GatherPdbInfo(const Image* image, const ImageSe
 			};
 
 			const DWORD RSDS = 0x53445352u;
-			const PDB70Header* pdbHeader = pointer::Offset<const PDB70Header*>(image->base, debugDirectory->PointerToRawData);
+			const PDB70Header* pdbHeader = pointer::Offset<const PDB70Header*>(Filesystem::GetMemoryMappedFileData(image), debugDirectory->PointerToRawData);
 			if (pdbHeader->signature == RSDS)
 			{
 				// PDB filename follows right after the header data
diff --git a/Engine/Source/Developer/Windows/LiveCoding/Private/External/LC_Executable.h b/Engine/Source/Developer/Windows/LiveCoding/Private/External/LC_Executable.h
--- a/Engine/Source/Developer/Windows/LiveCoding/Private/External/LC_Executable.h
+++ b/Engine/Source/Developer/Windows/LiveCoding/Private/External/LC_Executable.h
@@ -1,14 +1,17 @@
-// Copyright 2011-2019 Molecular Matters GmbH, all rights reserved.
+// Copyright 2011-2020 Molecular Matters GmbH, all rights reserved.
 
 #pragma once
 
+// BEGIN EPIC MOD
 #include "CoreTypes.h"
+// END EPIC MOD
 #include "LC_PointerUtil.h"
-#include "LC_MemoryFile.h"
-#include "xxhash.h"
+#include "LC_MemoryMappedFile.h"
+#include "LC_Hashing.h"
+// BEGIN EPIC MOD
 #include <vector>
 #include "Windows/WindowsHWrapper.h"
-
+// END EPIC MOD
 
 namespace executable
 {
@@ -25,7 +28,7 @@ namespace executable
 		IMAGE_OPTIONAL_HEADER optionalHeader;
 	};
 
-	typedef file::MemoryFile Image;
+	typedef Filesystem::MemoryMappedFile Image;
 
 	struct ImageSection
 	{
@@ -58,7 +61,7 @@ namespace executable
 	};
 
 
-	Image* OpenImage(const wchar_t* filename, file::OpenMode::Enum openMode);
+	Image* OpenImage(const wchar_t* filename, Filesystem::OpenMode::Enum openMode);
 	void CloseImage(Image*& image);
 
 	void RebaseImage(Image* image, PreferredBase preferredBase);
@@ -109,7 +112,7 @@ namespace executable
 			return T(0);
 		}
 
-		const T* address = pointer::Offset<const T*>(image->base, fileOffset);
+		const T* address = pointer::Offset<const T*>(Filesystem::GetMemoryMappedFileData(image), fileOffset);
 		return *address;
 	}
 
@@ -125,8 +128,8 @@ namespace std
 	{
 		inline std::size_t operator()(const executable::Header& header) const
 		{
-			const uint32_t hash1 = XXH32(&header.imageHeader, sizeof(executable::Header::imageHeader), 0u);
-			const uint32_t hash2 = XXH32(&header.optionalHeader, sizeof(executable::Header::optionalHeader), hash1);
+			const uint32_t hash1 = Hashing::Hash32(&header.imageHeader, sizeof(executable::Header::imageHeader), 0u);
+			const uint32_t hash2 = Hashing::Hash32(&header.optionalHeader, sizeof(executable::Header::optionalHeader), hash1);
 
 			return hash2;
 		}
diff --git a/Engine/Source/Developer/Windows/LiveCoding/Private/External/LC_HeartBeat.cpp b/Engine/Source/Developer/Windows/LiveCoding/Private/External/LC_HeartBeat.cpp
--- a/Engine/Source/Developer/Windows/LiveCoding/Private/External/LC_HeartBeat.cpp
+++ b/Engine/Source/Developer/Windows/LiveCoding/Private/External/LC_HeartBeat.cpp
@@ -1,5 +1,8 @@
-// Copyright 2011-2019 Molecular Matters GmbH, all rights reserved.
+// Copyright 2011-2020 Molecular Matters GmbH, all rights reserved.
 
+// BEGIN EPIC MOD
+//#include PCH_INCLUDE
+// END EPIC MOD
 #include "LC_HeartBeat.h"
 #include "LC_InterprocessMutex.h"
 #include "LC_NamedSharedMemory.h"
@@ -7,18 +10,18 @@
 #include "LC_PrimitiveNames.h"
 
 
-HeartBeat::HeartBeat(const wchar_t* const processGroupName, unsigned int processId)
+HeartBeat::HeartBeat(const wchar_t* const processGroupName, Process::Id processId)
 	: m_mutex(nullptr)
 	, m_memory(nullptr)
 {
 	m_mutex = new InterprocessMutex(primitiveNames::HeartBeatMutex(processGroupName, processId).c_str());
-	m_memory = new NamedSharedMemory(primitiveNames::HeartBeatNamedSharedMemory(processGroupName, processId).c_str());
+	m_memory = Process::CreateNamedSharedMemory(primitiveNames::HeartBeatNamedSharedMemory(processGroupName, processId).c_str(), 4096u);
 }
 
 
 HeartBeat::~HeartBeat(void)
 {
-	delete m_memory;
+	Process::DestroyNamedSharedMemory(m_memory);
 	delete m_mutex;
 }
 
@@ -28,7 +31,7 @@ void HeartBeat::Store(void)
 	const uint64_t currentTime = utcTime::GetCurrent();
 
 	InterprocessMutex::ScopedLock lock(m_mutex);
-	m_memory->Write(currentTime);
+	Process::WriteNamedSharedMemory(m_memory, currentTime);
 }
 
 
@@ -51,5 +54,5 @@ uint64_t HeartBeat::ReadBeatDelta(void) const
 uint64_t HeartBeat::ReadBeat(void) const
 {
 	InterprocessMutex::ScopedLock lock(m_mutex);
-	return m_memory->Read<uint64_t>();
+	return Process::ReadNamedSharedMemory<uint64_t>(m_memory);
 }
diff --git a/Engine/Source/Developer/Windows/LiveCoding/Private/External/LC_HeartBeat.h b/Engine/Source/Developer/Windows/LiveCoding/Private/External/LC_HeartBeat.h
--- a/Engine/Source/Developer/Windows/LiveCoding/Private/External/LC_HeartBeat.h
+++ b/Engine/Source/Developer/Windows/LiveCoding/Private/External/LC_HeartBeat.h
@@ -1,19 +1,22 @@
-// Copyright 2011-2019 Molecular Matters GmbH, all rights reserved.
+// Copyright 2011-2020 Molecular Matters GmbH, all rights reserved.
 
 #pragma once
 
+// BEGIN EPIC MOD
 #include "CoreTypes.h"
-#include <inttypes.h>
+// END EPIC MOD
+#include "LC_ProcessTypes.h"
+#include "LC_NamedSharedMemoryTypes.h"
+
 
 class InterprocessMutex;
-class NamedSharedMemory;
 
 // used for communicating heart beats between client and server.
 // can be used by several processes.
 class HeartBeat
 {
 public:
-	HeartBeat(const wchar_t* const processGroupName, unsigned int processId);
+	HeartBeat(const wchar_t* const processGroupName, Process::Id processId);
 	~HeartBeat(void);
 
 	// stores the current UTC time as heart beat
@@ -26,5 +29,5 @@ private:
 	uint64_t ReadBeat(void) const;
 
 	InterprocessMutex* m_mutex;
-	NamedSharedMemory* m_memory;
+	Process::NamedSharedMemory* m_memory;
 };
diff --git a/Engine/Source/Developer/Windows/LiveCoding/Private/External/LC_InterprocessMutex.cpp b/Engine/Source/Developer/Windows/LiveCoding/Private/External/LC_InterprocessMutex.cpp
--- a/Engine/Source/Developer/Windows/LiveCoding/Private/External/LC_InterprocessMutex.cpp
+++ b/Engine/Source/Developer/Windows/LiveCoding/Private/External/LC_InterprocessMutex.cpp
@@ -1,12 +1,18 @@
-// Copyright 2011-2019 Molecular Matters GmbH, all rights reserved.
+// Copyright 2011-2020 Molecular Matters GmbH, all rights reserved.
 
+// BEGIN EPIC MOD
+//#include PCH_INCLUDE
+// END EPIC MOD
 #include "LC_InterprocessMutex.h"
+// BEGIN EPIC MOD
 #include "LC_Logging.h"
 #include "Windows/WindowsHWrapper.h"
-
+// END EPIC MOD
 
 InterprocessMutex::InterprocessMutex(const wchar_t* name)
+	// BEGIN EPIC MOD
 	: m_mutex(::CreateMutexW(NULL, Windows::FALSE, name))
+	// END EPIC MOD
 {
 }
 
diff --git a/Engine/Source/Developer/Windows/LiveCoding/Private/External/LC_InterprocessMutex.h b/Engine/Source/Developer/Windows/LiveCoding/Private/External/LC_InterprocessMutex.h
--- a/Engine/Source/Developer/Windows/LiveCoding/Private/External/LC_InterprocessMutex.h
+++ b/Engine/Source/Developer/Windows/LiveCoding/Private/External/LC_InterprocessMutex.h
@@ -1,9 +1,12 @@
-// Copyright 2011-2019 Molecular Matters GmbH, all rights reserved.
+// Copyright 2011-2020 Molecular Matters GmbH, all rights reserved.
 
 #pragma once
 
+// BEGIN EPIC MOD
 #include "Windows/MinimalWindowsAPI.h"
 #include "LC_Platform.h"
+#include LC_PLATFORM_INCLUDE(LC_Foundation)
+// END EPIC MOD
 
 class InterprocessMutex
 {
@@ -38,5 +41,7 @@ public:
 	};
 
 private:
+	// BEGIN EPIC MOD
 	Windows::HANDLE m_mutex;
+	// END EPIC MOD
 };
diff --git a/Engine/Source/Developer/Windows/LiveCoding/Private/External/LC_Logging.cpp b/Engine/Source/Developer/Windows/LiveCoding/Private/External/LC_Logging.cpp
--- a/Engine/Source/Developer/Windows/LiveCoding/Private/External/LC_Logging.cpp
+++ b/Engine/Source/Developer/Windows/LiveCoding/Private/External/LC_Logging.cpp
@@ -1,10 +1,18 @@
-// Copyright 2011-2019 Molecular Matters GmbH, all rights reserved.
+// Copyright 2011-2020 Molecular Matters GmbH, all rights reserved.
 
+// BEGIN EPIC MOD
+//#include PCH_INCLUDE
+// END EPIC MOD
 #include "LC_Logging.h"
+// BEGIN EPIC MOD
+//#include "LC_App.h"
+//#include "LC_TimeDate.h"
+// END EPIC MOD
+#include "LC_Atomic_Windows.h"
+// BEGIN EPIC MOD
 #include "Logging/LogMacros.h"
 #include "LiveCodingLog.h"
-#include <atomic>
-
+// END EPIC MOD
 
 namespace
 {
@@ -13,7 +21,7 @@ namespace
 	static thread_local char gtls_logBuffer[PER_THREAD_BUFFER_SIZE] = {};
 
 	// keeps track of the current indentation
-	static std::atomic<int> g_indentationLevel[logging::Channel::COUNT] = {};
+	static volatile int32_t g_indentationLevel[Logging::Channel::COUNT] = {};
 
 	// small lookup table for different levels of indentation
 	static const int MAX_INDENTATION_LEVELS = 7;
@@ -28,16 +36,17 @@ namespace
 		"            * "		// level 6
 	};
 
+
 	// BEGIN EPIC MOD - Redirecting log output
-	static void DefaultOutputHandler(logging::Channel::Enum channel, logging::Type::Enum type, const wchar_t* const Message)
+	static void DefaultOutputHandler(Logging::Channel::Enum channel, Logging::Type::Enum type, const wchar_t* const Message)
 	{
 		FString MessageWithoutNewline = FString(Message).TrimEnd();
 		switch (type)
 		{
-		case logging::Type::LOG_WARNING:
+		case Logging::Type::LOG_WARNING:
 			UE_LOG(LogLiveCoding, Warning, TEXT("%s"), *MessageWithoutNewline);
 			break;
-		case logging::Type::LOG_ERROR:
+		case Logging::Type::LOG_ERROR:
 			UE_LOG(LogLiveCoding, Error, TEXT("%s"), *MessageWithoutNewline);
 			break;
 		default:
@@ -46,17 +55,17 @@ namespace
 		}
 	}
 
-	static logging::OutputHandlerType OutputHandler = &DefaultOutputHandler;
+	static Logging::OutputHandlerType OutputHandler = &DefaultOutputHandler;
 	// END EPIC MOD - Redirecting log output
 
 	// BEGIN EPIC MOD - Explicit API for enabling log channels
-	static bool ChannelEnabled[logging::Channel::COUNT] = { };
+	static bool ChannelEnabled[Logging::Channel::COUNT] = { };
 	// END EPIC MOD - Explicit API for enabling log channels
 
-	static bool IsChannelEnabled(logging::Channel::Enum channel, logging::Type::Enum type)
+	static bool IsChannelEnabled(Logging::Channel::Enum channel, Logging::Type::Enum type)
 	{
 		// warnings, errors and success logs should always be output to all channels
-		if (type != logging::Type::LOG_INFO)
+		if (type != Logging::Type::LOG_INFO)
 		{
 			return true;
 		}
@@ -64,7 +73,7 @@ namespace
 		// BEGIN EPIC MOD - Explicit API for enabling log channels
 		switch (channel)
 		{
-			case logging::Channel::USER:
+			case Logging::Channel::USER:
 				// user-visible logs are *always* logged
 				return true;
 
@@ -77,32 +86,32 @@ namespace
 }
 
 
-logging::Indent::Indent(Channel::Enum channel)
+Logging::Indent::Indent(Channel::Enum channel)
 	: m_channel(channel)
 {
 	IncrementIndentation(channel);
 }
 
 
-logging::Indent::~Indent(void)
+Logging::Indent::~Indent(void)
 {
 	DecrementIndentation(m_channel);
 }
 
 
-void logging::IncrementIndentation(Channel::Enum channel)
+void Logging::IncrementIndentation(Channel::Enum channel)
 {
-	std::atomic_fetch_add(&g_indentationLevel[channel], 1);
+	Atomic::IncrementConsistent(&g_indentationLevel[channel]);
 }
 
 
-void logging::DecrementIndentation(Channel::Enum channel)
+void Logging::DecrementIndentation(Channel::Enum channel)
 {
-	std::atomic_fetch_sub(&g_indentationLevel[channel], 1);
+	Atomic::DecrementConsistent(&g_indentationLevel[channel]);
 }
 
 
-const char* logging::GetIndentation(logging::Channel::Enum channel)
+const char* Logging::GetIndentation(Logging::Channel::Enum channel)
 {
 	int indent = g_indentationLevel[channel];
 	if (indent >= MAX_INDENTATION_LEVELS)
@@ -114,22 +123,23 @@ const char* logging::GetIndentation(logging::Channel::Enum channel)
 }
 
 
-template <> void logging::LogNoFormat<logging::Channel::Enum::USER>(const wchar_t* const buffer)
+template <> void Logging::LogNoFormat<Logging::Channel::Enum::USER>(const wchar_t* const buffer)
 {
 	// BEGIN EPIC MOD - Redirecting log output
-	OutputHandler(logging::Channel::Enum::USER, Type::LOG_INFO, buffer);
+	OutputHandler(Logging::Channel::Enum::USER, Type::LOG_INFO, buffer);
 	// END EPIC MOD - Redirecting log output
 }
 
 
-template <> void logging::LogNoFormat<logging::Channel::Enum::DEV>(const char* const buffer)
+template <> void Logging::LogNoFormat<Logging::Channel::Enum::DEV>(const char* const buffer)
 {
 	// BEGIN EPIC MOD - Redirecting log output
-	OutputHandler(logging::Channel::Enum::DEV, Type::LOG_INFO, ANSI_TO_TCHAR(buffer));
+	OutputHandler(Logging::Channel::Enum::DEV, Type::LOG_INFO, ANSI_TO_TCHAR(buffer));
 	// END EPIC MOD - Redirecting log output
 }
 
-void logging::Log(Channel::Enum channel, Type::Enum type, const char* const format, ...)
+
+void Logging::Log(Channel::Enum channel, Type::Enum type, const char* const format, ...)
 {
 	if (!IsChannelEnabled(channel, type))
 	{
@@ -150,13 +160,13 @@ void logging::Log(Channel::Enum channel, Type::Enum type, const char* const form
 }
 
 // BEGIN EPIC MOD - Redirecting log output
-void logging::EnableChannel(Channel::Enum channel, bool enabled)
+void Logging::EnableChannel(Channel::Enum channel, bool enabled)
 {
-	check(channel != logging::Channel::USER || enabled);
+	check(channel != Logging::Channel::USER || enabled);
 	ChannelEnabled[(int)channel] = enabled;
 }
 
-void logging::SetOutputHandler(logging::OutputHandlerType handler)
+void Logging::SetOutputHandler(Logging::OutputHandlerType handler)
 {
 	if (handler == nullptr)
 	{
diff --git a/Engine/Source/Developer/Windows/LiveCoding/Private/External/LC_Logging.h b/Engine/Source/Developer/Windows/LiveCoding/Private/External/LC_Logging.h
--- a/Engine/Source/Developer/Windows/LiveCoding/Private/External/LC_Logging.h
+++ b/Engine/Source/Developer/Windows/LiveCoding/Private/External/LC_Logging.h
@@ -1,12 +1,14 @@
-// Copyright 2011-2019 Molecular Matters GmbH, all rights reserved.
+// Copyright 2011-2020 Molecular Matters GmbH, all rights reserved.
 
 #pragma once
 
+// BEGIN EPIC MOD
 #include "CoreTypes.h"
+// END EPIC MOD
 #include "LC_Preprocessor.h"
 
 
-namespace logging
+namespace Logging
 {
 	struct Channel
 	{
@@ -67,20 +69,20 @@ namespace logging
 	// END EPIC MOD - Explicit API for log channels
 }
 
-#define LC_LOG_INDENT_USER						const logging::Indent LC_PP_UNIQUE_NAME(autoIndentUser)(logging::Channel::USER)
-#define LC_LOG_INDENT_DEV						const logging::Indent LC_PP_UNIQUE_NAME(autoIndentDev)(logging::Channel::DEV)
-#define LC_LOG_INDENT_TELEMETRY					const logging::Indent LC_PP_UNIQUE_NAME(autoIndentTelemetry)(logging::Channel::TELEMETRY)
+#define LC_LOG_INDENT_USER						const Logging::Indent LC_PP_UNIQUE_NAME(autoIndentUser)(Logging::Channel::USER)
+#define LC_LOG_INDENT_DEV						const Logging::Indent LC_PP_UNIQUE_NAME(autoIndentDev)(Logging::Channel::DEV)
+#define LC_LOG_INDENT_TELEMETRY					const Logging::Indent LC_PP_UNIQUE_NAME(autoIndentTelemetry)(Logging::Channel::TELEMETRY)
 
-#define LC_LOG_USER(_format, ...)				logging::Log(logging::Channel::USER, logging::Type::LOG_INFO, "%s" _format "\n", logging::GetIndentation(logging::Channel::USER), __VA_ARGS__)
-#define LC_LOG_DEV(_format, ...)				logging::Log(logging::Channel::DEV, logging::Type::LOG_INFO, "%s" _format "\n", logging::GetIndentation(logging::Channel::DEV), __VA_ARGS__)
-#define LC_LOG_TELEMETRY(_format, ...)			logging::Log(logging::Channel::TELEMETRY, logging::Type::LOG_INFO, "%s" _format "\n", logging::GetIndentation(logging::Channel::TELEMETRY), __VA_ARGS__)
+#define LC_LOG_USER(_format, ...)				Logging::Log(Logging::Channel::USER, Logging::Type::LOG_INFO, "%s" _format "\n", Logging::GetIndentation(Logging::Channel::USER), __VA_ARGS__)
+#define LC_LOG_DEV(_format, ...)				Logging::Log(Logging::Channel::DEV, Logging::Type::LOG_INFO, "%s" _format "\n", Logging::GetIndentation(Logging::Channel::DEV), __VA_ARGS__)
+#define LC_LOG_TELEMETRY(_format, ...)			Logging::Log(Logging::Channel::TELEMETRY, Logging::Type::LOG_INFO, "%s" _format "\n", Logging::GetIndentation(Logging::Channel::TELEMETRY), __VA_ARGS__)
 
-#define LC_WARNING_USER(_format, ...)			logging::Log(logging::Channel::USER, logging::Type::LOG_WARNING, "%s" _format "\n", logging::GetIndentation(logging::Channel::USER), __VA_ARGS__)
-#define LC_WARNING_DEV(_format, ...)			logging::Log(logging::Channel::DEV, logging::Type::LOG_WARNING, "WARNING: %s" _format "\n", logging::GetIndentation(logging::Channel::DEV), __VA_ARGS__)
-#define LC_WARNING_TELEMETRY(_format, ...)		logging::Log(logging::Channel::TELEMETRY, logging::Type::LOG_WARNING, "WARNING: %s" _format "\n", logging::GetIndentation(logging::Channel::TELEMETRY), __VA_ARGS__)
+#define LC_WARNING_USER(_format, ...)			Logging::Log(Logging::Channel::USER, Logging::Type::LOG_WARNING, "%s" _format "\n", Logging::GetIndentation(Logging::Channel::USER), __VA_ARGS__)
+#define LC_WARNING_DEV(_format, ...)			Logging::Log(Logging::Channel::DEV, Logging::Type::LOG_WARNING, "WARNING: %s" _format "\n", Logging::GetIndentation(Logging::Channel::DEV), __VA_ARGS__)
+#define LC_WARNING_TELEMETRY(_format, ...)		Logging::Log(Logging::Channel::TELEMETRY, Logging::Type::LOG_WARNING, "WARNING: %s" _format "\n", Logging::GetIndentation(Logging::Channel::TELEMETRY), __VA_ARGS__)
 
-#define LC_ERROR_USER(_format, ...)				logging::Log(logging::Channel::USER, logging::Type::LOG_ERROR, "%s" _format "\n", logging::GetIndentation(logging::Channel::USER), __VA_ARGS__)
-#define LC_ERROR_DEV(_format, ...)				logging::Log(logging::Channel::DEV, logging::Type::LOG_ERROR, "ERROR: %s" _format "\n", logging::GetIndentation(logging::Channel::DEV), __VA_ARGS__)
-#define LC_ERROR_TELEMETRY(_format, ...)		logging::Log(logging::Channel::TELEMETRY, logging::Type::LOG_ERROR, "ERROR: %s" _format "\n", logging::GetIndentation(logging::Channel::TELEMETRY), __VA_ARGS__)
+#define LC_ERROR_USER(_format, ...)				Logging::Log(Logging::Channel::USER, Logging::Type::LOG_ERROR, "%s" _format "\n", Logging::GetIndentation(Logging::Channel::USER), __VA_ARGS__)
+#define LC_ERROR_DEV(_format, ...)				Logging::Log(Logging::Channel::DEV, Logging::Type::LOG_ERROR, "ERROR: %s" _format "\n", Logging::GetIndentation(Logging::Channel::DEV), __VA_ARGS__)
+#define LC_ERROR_TELEMETRY(_format, ...)		Logging::Log(Logging::Channel::TELEMETRY, Logging::Type::LOG_ERROR, "ERROR: %s" _format "\n", Logging::GetIndentation(Logging::Channel::TELEMETRY), __VA_ARGS__)
 
-#define LC_SUCCESS_USER(_format, ...)			logging::Log(logging::Channel::USER, logging::Type::LOG_SUCCESS, "%s" _format "\n", logging::GetIndentation(logging::Channel::USER), __VA_ARGS__)
+#define LC_SUCCESS_USER(_format, ...)			Logging::Log(Logging::Channel::USER, Logging::Type::LOG_SUCCESS, "%s" _format "\n", Logging::GetIndentation(Logging::Channel::USER), __VA_ARGS__)
diff --git a/Engine/Source/Developer/Windows/LiveCoding/Private/External/LC_Memory.cpp b/Engine/Source/Developer/Windows/LiveCoding/Private/External/LC_Memory.cpp
--- a/Engine/Source/Developer/Windows/LiveCoding/Private/External/LC_Memory.cpp
+++ b/Engine/Source/Developer/Windows/LiveCoding/Private/External/LC_Memory.cpp
@@ -1,3 +1,6 @@
 // Copyright 2011-2019 Molecular Matters GmbH, all rights reserved.
 
+// BEGIN EPIC MOD
+//#include PCH_INCLUDE
+// END EPIC MOD
 #include "LC_Memory.h"
diff --git a/Engine/Source/Developer/Windows/LiveCoding/Private/External/LC_Memory.h b/Engine/Source/Developer/Windows/LiveCoding/Private/External/LC_Memory.h
--- a/Engine/Source/Developer/Windows/LiveCoding/Private/External/LC_Memory.h
+++ b/Engine/Source/Developer/Windows/LiveCoding/Private/External/LC_Memory.h
@@ -1,8 +1,10 @@
-// Copyright 2011-2019 Molecular Matters GmbH, all rights reserved.
+// Copyright 2011-2020 Molecular Matters GmbH, all rights reserved.
 
 #pragma once
 
+// BEGIN EPIC MOD
 #include "CoreTypes.h"
+// END EPIC MOD
 
 namespace memory
 {
@@ -31,4 +33,19 @@ namespace memory
 		delete[] instance;
 		instance = nullptr;
 	}
+
+
+	template <typename T, size_t N>
+	bool Matches(const T(&mem1)[N], const T(&mem2)[N])
+	{
+		for (size_t i = 0u; i < N; ++i)
+		{
+			if (mem1[i] != mem2[i])
+			{
+				return false;
+			}
+		}
+
+		return true;
+	}
 }
diff --git a/Engine/Source/Developer/Windows/LiveCoding/Private/External/LC_MemoryStream.cpp b/Engine/Source/Developer/Windows/LiveCoding/Private/External/LC_MemoryStream.cpp
--- a/Engine/Source/Developer/Windows/LiveCoding/Private/External/LC_MemoryStream.cpp
+++ b/Engine/Source/Developer/Windows/LiveCoding/Private/External/LC_MemoryStream.cpp
@@ -1,7 +1,13 @@
-// Copyright 2011-2019 Molecular Matters GmbH, all rights reserved.
+// Copyright 2011-2020 Molecular Matters GmbH, all rights reserved.
 
+// BEGIN EPIC MOD
+//#include PCH_INCLUDE
+// END EPIC MOD
 #include "LC_MemoryStream.h"
 #include "LC_PointerUtil.h"
+// BEGIN EPIC MOD
+#include "LC_Assert.h"
+// END EPIC MOD
 
 
 memoryStream::Reader::Reader(const void* data, size_t capacity)
diff --git a/Engine/Source/Developer/Windows/LiveCoding/Private/External/LC_MemoryStream.h b/Engine/Source/Developer/Windows/LiveCoding/Private/External/LC_MemoryStream.h
--- a/Engine/Source/Developer/Windows/LiveCoding/Private/External/LC_MemoryStream.h
+++ b/Engine/Source/Developer/Windows/LiveCoding/Private/External/LC_MemoryStream.h
@@ -1,9 +1,12 @@
-// Copyright 2011-2019 Molecular Matters GmbH, all rights reserved.
+// Copyright 2011-2020 Molecular Matters GmbH, all rights reserved.
 
 #pragma once
 
+// BEGIN EPIC MOD
 #include "CoreTypes.h"
 #include "LC_Platform.h"
+#include LC_PLATFORM_INCLUDE(LC_Foundation)
+// END EPIC MOD
 
 namespace memoryStream
 {
diff --git a/Engine/Source/Developer/Windows/LiveCoding/Private/External/LC_NamedSharedMemory.cpp b/Engine/Source/Developer/Windows/LiveCoding/Private/External/LC_NamedSharedMemory.cpp
--- a/Engine/Source/Developer/Windows/LiveCoding/Private/External/LC_NamedSharedMemory.cpp
+++ b/Engine/Source/Developer/Windows/LiveCoding/Private/External/LC_NamedSharedMemory.cpp
@@ -1,73 +1,81 @@
-// Copyright 2011-2019 Molecular Matters GmbH, all rights reserved.
+// Copyright 2011-2020 Molecular Matters GmbH, all rights reserved.
 
+// BEGIN EPIC MOD
+//#include PCH_INCLUDE
+// END EPIC MOD
 #include "LC_NamedSharedMemory.h"
+// BEGIN EPIC MOD
 #include "LC_Logging.h"
 #include "Windows/WindowsHWrapper.h"
+//END EPIC MOD
+
+namespace Process
+{
+	struct NamedSharedMemory
+	{
+		// BEGIN EPIC MOD
+		Windows::HANDLE memoryMapping;
+		// END EPIC MOD
+		void* memoryView;
+		bool isOwned;
+	};
+}
 
 
-namespace
+bool Process::Current::DoesOwnNamedSharedMemory(const NamedSharedMemory* memory)
 {
-	static const DWORD MEMORY_SIZE = 4096u;
+	return memory->isOwned;
 }
 
 
-NamedSharedMemory::NamedSharedMemory(const wchar_t* name)
-	: m_file(nullptr)
-	, m_view(nullptr)
-	, m_isOwned(true)
+Process::NamedSharedMemory* Process::CreateNamedSharedMemory(const wchar_t* name, size_t size)
 {
-	m_file = ::CreateFileMappingW(INVALID_HANDLE_VALUE, NULL, PAGE_READWRITE, 0, MEMORY_SIZE, name);
+	::ULARGE_INTEGER integer = {};
+	integer.QuadPart = size;
+
+	// BEGIN EPIC MOD
+	Windows::HANDLE memoryMapping = ::CreateFileMappingW(INVALID_HANDLE_VALUE, NULL, PAGE_READWRITE, integer.HighPart, integer.LowPart, name);
+	// END EPIC MOD
 	const DWORD error = ::GetLastError();
-	if (m_file == NULL)
+	if (memoryMapping == NULL)
 	{
 		LC_ERROR_USER("Cannot create named shared memory. Error: 0x%X", error);
+		return nullptr;
 	}
-	else if (error == ERROR_ALREADY_EXISTS)
-	{
-		// another process already created this file mapping. the handle points to the already existing object.
-		m_isOwned = false;
-	}
 
-	if (m_file)
+	// check if another process already created this file mapping.
+	// the handle then points to the already existing object.
+	const bool isOwned = (error != ERROR_ALREADY_EXISTS);
+
+	void* memoryView = ::MapViewOfFile(memoryMapping, FILE_MAP_ALL_ACCESS, 0, 0, size);
+	if (memoryView == NULL)
 	{
-		m_view = ::MapViewOfFile(m_file, FILE_MAP_ALL_ACCESS, 0, 0, MEMORY_SIZE);
-		if (m_view == NULL)
-		{
-			LC_ERROR_USER("Cannot map view for named shared memory. Error: 0x%X", ::GetLastError());
-
-			::CloseHandle(m_file);
-			m_file = nullptr;
-		}
-	}
-}
+		LC_ERROR_USER("Cannot map view for named shared memory. Error: 0x%X", ::GetLastError());
 
+		::CloseHandle(memoryMapping);
+		return nullptr;
+	}
 
-NamedSharedMemory::~NamedSharedMemory(void)
-{
-	::UnmapViewOfFile(m_view);
-	::CloseHandle(m_file);
+	return new NamedSharedMemory { memoryMapping, memoryView, isOwned };
 }
 
 
-bool NamedSharedMemory::IsOwnedByCallingProcess(void) const
+void Process::DestroyNamedSharedMemory(NamedSharedMemory*& memory)
 {
-	return m_isOwned;
+	::UnmapViewOfFile(memory->memoryView);
+	::CloseHandle(memory->memoryMapping);
+	delete memory;
+	memory = nullptr;
 }
 
 
-void NamedSharedMemory::Read(void* buffer, size_t size)
+void Process::ReadNamedSharedMemory(const NamedSharedMemory* memory, void* buffer, size_t size)
 {
-	if (m_view)
-	{
-		memcpy(buffer, m_view, size);
-	}
+	memcpy(buffer, memory->memoryView, size);
 }
 
 
-void NamedSharedMemory::Write(const void* buffer, size_t size)
+void Process::WriteNamedSharedMemory(NamedSharedMemory* memory, const void* buffer, size_t size)
 {
-	if (m_view)
-	{
-		memcpy(m_view, buffer, size);
-	}
+	memcpy(memory->memoryView, buffer, size);
 }
diff --git a/Engine/Source/Developer/Windows/LiveCoding/Private/External/LC_NamedSharedMemory.h b/Engine/Source/Developer/Windows/LiveCoding/Private/External/LC_NamedSharedMemory.h
--- a/Engine/Source/Developer/Windows/LiveCoding/Private/External/LC_NamedSharedMemory.h
+++ b/Engine/Source/Developer/Windows/LiveCoding/Private/External/LC_NamedSharedMemory.h
@@ -1,36 +1,49 @@
-// Copyright 2011-2019 Molecular Matters GmbH, all rights reserved.
+// Copyright 2011-2020 Molecular Matters GmbH, all rights reserved.
 
 #pragma once
 
+// BEGIN EPIC MOD
 #include "Windows/MinimalWindowsApi.h"
+// END EPIC MOD
 
-class NamedSharedMemory
+namespace Process
 {
-public:
-	explicit NamedSharedMemory(const wchar_t* name);
-	~NamedSharedMemory(void);
+	// opaque type
+	struct NamedSharedMemory;
 
-	bool IsOwnedByCallingProcess(void) const;
+	// Current/calling process.
+	namespace Current
+	{
+		// Returns whether the named shared memory is owned by the calling process.
+		bool DoesOwnNamedSharedMemory(const NamedSharedMemory* memory);
+	}
+
+	// Creates named shared memory.
+	NamedSharedMemory* CreateNamedSharedMemory(const wchar_t* name, size_t size);
+
+	// Destroys named shared memory.
+	void DestroyNamedSharedMemory(NamedSharedMemory*& memory);
 
-	void Read(void* buffer, size_t size);
-	void Write(const void* buffer, size_t size);
+	// Reads from named shared memory into the given buffer.
+	void ReadNamedSharedMemory(const NamedSharedMemory* memory, void* buffer, size_t size);
 
+	// Writes from the given buffer into named shared memory.
+	void WriteNamedSharedMemory(NamedSharedMemory* memory, const void* buffer, size_t size);
+
+	// Convenience function for reading a value of a certain type from named shared memory.
 	template <typename T>
-	T Read(void)
+	T ReadNamedSharedMemory(const NamedSharedMemory* memory)
 	{
 		T value = {};
-		Read(&value, sizeof(T));
+		ReadNamedSharedMemory(memory, &value, sizeof(T));
+
 		return value;
 	}
 
+	// Convenience function for writing a value of a certain type to named shared memory.
 	template <typename T>
-	void Write(const T& value)
+	void WriteNamedSharedMemory(NamedSharedMemory* memory, const T& value)
 	{
-		Write(&value, sizeof(T));
+		WriteNamedSharedMemory(memory, &value, sizeof(T));
 	}
-
-private:
-	Windows::HANDLE m_file;
-	void* m_view;
-	bool m_isOwned;
-};
+}
diff --git a/Engine/Source/Developer/Windows/LiveCoding/Private/External/LC_Platform.h b/Engine/Source/Developer/Windows/LiveCoding/Private/External/LC_Platform.h
--- a/Engine/Source/Developer/Windows/LiveCoding/Private/External/LC_Platform.h
+++ b/Engine/Source/Developer/Windows/LiveCoding/Private/External/LC_Platform.h
@@ -1,41 +1,40 @@
-// Copyright 2011-2019 Molecular Matters GmbH, all rights reserved.
+// Copyright 2011-2020 Molecular Matters GmbH, all rights reserved.
 
 #pragma once
 
+// BEGIN EPIC MOD
 #include "CoreTypes.h"
 #include "Misc/AssertionMacros.h"
+// END EPIC MOD
 #include "LC_Preprocessor.h"
 
-#if DO_CHECK
-#	define LC_BREAKPOINT()						((::IsDebuggerPresent() == 0) ? (void)true : __debugbreak())
+
+// determine the platform we're compiling for
+#if defined(_WIN32)
+	// Windows
+#	if defined(_WIN64)
+		// 64-bit
+#		define LC_64_BIT	1
+#		define LC_32_BIT	0
+#	else
+		// 32-bit
+#		define LC_64_BIT	0
+#		define LC_32_BIT	1
+#	endif
+#	define LC_PLATFORM_SUFFIX _Windows.h
+#elif defined(__ORBIS__)
+	// PlayStation 4
+#	define LC_64_BIT	1
+#	define LC_32_BIT	0
+#	define LC_PLATFORM_SUFFIX _Orbis.h
 #else
-#	define LC_BREAKPOINT()						(void)true
+	#pragma error "Unknown platform"
 #endif
 
-#define LC_ASSERT(_condition, _msg)				checkf(_condition, TEXT("%s"), TEXT(_msg))
-
-#define LC_FILE									__FILE__
-#define LC_LINE									__LINE__
-#define LC_FUNCTION_NAME						__FUNCTION__
-#define LC_FUNCTION_SIGNATURE					__FUNCSIG__
-#define LC_UNUSED(_value)						(void)(_value)
-
-#define LC_DISABLE_COPY(_name)					_name(const _name&) = delete
-#define LC_DISABLE_MOVE(_name)					_name(_name&&) = delete
-#define LC_DISABLE_ASSIGNMENT(_name)			_name& operator=(const _name&) = delete
-#define LC_DISABLE_MOVE_ASSIGNMENT(_name)		_name& operator=(_name&&) = delete
 
-#define LC_ALWAYS_INLINE						__forceinline
-#define LC_NEVER_INLINE							__declspec(noinline)
-
-#if PLATFORM_64BITS
-#	define LC_64_BIT 1
-#	define LC_32_BIT 0
-#else
-#	define LC_64_BIT 0
-#	define LC_32_BIT 1
-#endif
+#define LC_PLATFORM_INCLUDE(_header) LC_PP_STRINGIFY(LC_PP_JOIN(_header, LC_PLATFORM_SUFFIX))
 
+// BEGIN EPIC MOD
 // convenience macro for referring to symbol identifiers.
 // some identifiers contain an extra leading underscore in 32-bit builds.
 #if LC_64_BIT
@@ -43,3 +42,4 @@
 #else
 #	define LC_IDENTIFIER(_name)		"_" _name
 #endif
+// END EPIC MOD
\ No newline at end of file
diff --git a/Engine/Source/Developer/Windows/LiveCoding/Private/External/LC_PointerUtil.h b/Engine/Source/Developer/Windows/LiveCoding/Private/External/LC_PointerUtil.h
--- a/Engine/Source/Developer/Windows/LiveCoding/Private/External/LC_PointerUtil.h
+++ b/Engine/Source/Developer/Windows/LiveCoding/Private/External/LC_PointerUtil.h
@@ -1,18 +1,22 @@
-// Copyright 2011-2019 Molecular Matters GmbH, all rights reserved.
+// Copyright 2011-2020 Molecular Matters GmbH, all rights reserved.
 
 #pragma once
 
+// BEGIN EPIC MOD
 #include "CoreTypes.h"
+// END EPIC MOD
 #include <type_traits>
 
 
 namespace pointer
 {
+	// BEGIN EPIC MOD
 	template <typename T>
 	T As(void* anyPointer);
 
 	template <typename T>
 	T As(const void* anyPointer);
+	// END EPIC MOD
 
 	template <typename T>
 	T AsInteger(const void* anyPointer)
diff --git a/Engine/Source/Developer/Windows/LiveCoding/Private/External/LC_Preprocessor.h b/Engine/Source/Developer/Windows/LiveCoding/Private/External/LC_Preprocessor.h
--- a/Engine/Source/Developer/Windows/LiveCoding/Private/External/LC_Preprocessor.h
+++ b/Engine/Source/Developer/Windows/LiveCoding/Private/External/LC_Preprocessor.h
@@ -1,12 +1,16 @@
-// Copyright 2011-2019 Molecular Matters GmbH, all rights reserved.
+// Copyright 2011-2020 Molecular Matters GmbH, all rights reserved.
 
 #pragma once
 
 
-// joins two tokens, even when the tokens are macros themselves
+// joins two tokens, even when the tokens are macros
 #define LC_PP_JOIN_HELPER_HELPER(_0, _1)		_0##_1
 #define LC_PP_JOIN_HELPER(_0, _1)				LC_PP_JOIN_HELPER_HELPER(_0, _1)
 #define LC_PP_JOIN(_0, _1)						LC_PP_JOIN_HELPER(_0, _1)
 
+// stringifies a token, even when the token is a macro
+#define LC_PP_STRINGIFY_HELPER(_0)				#_0
+#define LC_PP_STRINGIFY(_0)						LC_PP_STRINGIFY_HELPER(_0)
+
 // generates a unique name
 #define LC_PP_UNIQUE_NAME(_name)				LC_PP_JOIN(_name, __COUNTER__)
diff --git a/Engine/Source/Developer/Windows/LiveCoding/Private/External/LC_PrimitiveNames.cpp b/Engine/Source/Developer/Windows/LiveCoding/Private/External/LC_PrimitiveNames.cpp
--- a/Engine/Source/Developer/Windows/LiveCoding/Private/External/LC_PrimitiveNames.cpp
+++ b/Engine/Source/Developer/Windows/LiveCoding/Private/External/LC_PrimitiveNames.cpp
@@ -1,18 +1,25 @@
-// Copyright 2011-2019 Molecular Matters GmbH, all rights reserved.
+// Copyright 2011-2020 Molecular Matters GmbH, all rights reserved.
 
+// BEGIN EPIC MOD
+//#include PCH_INCLUDE
+// END EPIC MOD
 #include "LC_PrimitiveNames.h"
 
 
+// BEGIN EPIC MOD
 // main suffix for all live coding primitives except pipes
 #define LPP						L"_UE_LC"
+// END EPIC MOD
 #define LPP_JOB					LPP L"_JOB"
 #define LPP_MUTEX				LPP L"_IPM"
 #define LPP_MEMORY				LPP L"_NSM"
 #define LPP_EVENT				LPP L"_EVT"
 #define LPP_SERVER_READY		LPP L"_SR"
 #define LPP_COMPILE				LPP L"_CMP"
+// BEGIN EPIC MOD
 #define LPP_PIPE				L"\\\\.\\pipe\\UE_LC"
 #define LPP_EXCEPTION_PIPE		L"\\\\.\\pipe\\UE_LC_EXC"
+// END EPIC MOD
 #define LPP_HEARTBEAT_MUTEX		LPP_MUTEX L"_HB"
 #define LPP_HEARTBEAT_MEMORY	LPP_MEMORY L"_HB"
 #define LPP_RESTART_REQUESTED	LPP_EVENT L"_RST_REQ"
@@ -104,63 +111,63 @@ std::wstring primitiveNames::ExceptionPipe(const std::wstring& processGroupName)
 }
 
 
-std::wstring primitiveNames::HeartBeatMutex(const std::wstring& processGroupName, unsigned int processId)
+std::wstring primitiveNames::HeartBeatMutex(const std::wstring& processGroupName, Process::Id processId)
 {
 	std::wstring name;
 	name.reserve(128u);
 
 	name += processGroupName;
 	name += LPP_HEARTBEAT_MUTEX;
-	name += std::to_wstring(processId);
+	name += std::to_wstring(+processId);
 
 	return name;
 }
 
 
-std::wstring primitiveNames::HeartBeatNamedSharedMemory(const std::wstring& processGroupName, unsigned int processId)
+std::wstring primitiveNames::HeartBeatNamedSharedMemory(const std::wstring& processGroupName, Process::Id processId)
 {
 	std::wstring name;
 	name.reserve(128u);
 
 	name += processGroupName;
 	name += LPP_HEARTBEAT_MEMORY;
-	name += std::to_wstring(processId);
+	name += std::to_wstring(+processId);
 
 	return name;
 }
 
 
-std::wstring primitiveNames::RequestRestart(unsigned int processId)
+std::wstring primitiveNames::RequestRestart(Process::Id processId)
 {
 	std::wstring name;
 	name.reserve(128u);
 
 	name += LPP_RESTART_REQUESTED;
-	name += std::to_wstring(processId);
+	name += std::to_wstring(+processId);
 
 	return name;
 }
 
 
-std::wstring primitiveNames::PreparedRestart(unsigned int processId)
+std::wstring primitiveNames::PreparedRestart(Process::Id processId)
 {
 	std::wstring name;
 	name.reserve(128u);
 
 	name += LPP_RESTART_PREPARED;
-	name += std::to_wstring(processId);
+	name += std::to_wstring(+processId);
 
 	return name;
 }
 
 
-std::wstring primitiveNames::Restart(unsigned int processId)
+std::wstring primitiveNames::Restart(Process::Id processId)
 {
 	std::wstring name;
 	name.reserve(128u);
 
 	name += LPP_RESTART;
-	name += std::to_wstring(processId);
+	name += std::to_wstring(+processId);
 
 	return name;
 }
diff --git a/Engine/Source/Developer/Windows/LiveCoding/Private/External/LC_PrimitiveNames.h b/Engine/Source/Developer/Windows/LiveCoding/Private/External/LC_PrimitiveNames.h
--- a/Engine/Source/Developer/Windows/LiveCoding/Private/External/LC_PrimitiveNames.h
+++ b/Engine/Source/Developer/Windows/LiveCoding/Private/External/LC_PrimitiveNames.h
@@ -1,9 +1,15 @@
-// Copyright 2011-2019 Molecular Matters GmbH, all rights reserved.
+// Copyright 2011-2020 Molecular Matters GmbH, all rights reserved.
 
 #pragma once
 
+// BEGIN EPIC MOD
 #include "CoreTypes.h"
+// END EPIC MOD
+#include "LC_ProcessTypes.h"
+// BEGIN EPIC MOD
 #include <string>
+// END EPIC MOD
+
 
 namespace primitiveNames
 {
@@ -21,11 +27,11 @@ namespace primitiveNames
 	std::wstring ExceptionPipe(const std::wstring& processGroupName);
 
 	// heart beat
-	std::wstring HeartBeatMutex(const std::wstring& processGroupName, unsigned int processId);
-	std::wstring HeartBeatNamedSharedMemory(const std::wstring& processGroupName, unsigned int processId);
+	std::wstring HeartBeatMutex(const std::wstring& processGroupName, Process::Id processId);
+	std::wstring HeartBeatNamedSharedMemory(const std::wstring& processGroupName, Process::Id processId);
 
 	// restart
-	std::wstring RequestRestart(unsigned int processId);
-	std::wstring PreparedRestart(unsigned int processId);
-	std::wstring Restart(unsigned int processId);
+	std::wstring RequestRestart(Process::Id processId);
+	std::wstring PreparedRestart(Process::Id processId);
+	std::wstring Restart(Process::Id processId);
 }
