Integration of Live++ 1.6.3

diff --git a/Engine/Source/Developer/Windows/LiveCoding/Private/External/LC_Process.cpp b/Engine/Source/Developer/Windows/LiveCoding/Private/External/LC_Process.cpp
--- a/Engine/Source/Developer/Windows/LiveCoding/Private/External/LC_Process.cpp
+++ b/Engine/Source/Developer/Windows/LiveCoding/Private/External/LC_Process.cpp
@@ -1,15 +1,20 @@
-// Copyright 2011-2019 Molecular Matters GmbH, all rights reserved.
+// Copyright 2011-2020 Molecular Matters GmbH, all rights reserved.
 
+// BEGIN EPIC MOD
+//#include PCH_INCLUDE
+// END EPIC MOD
 #include "LC_Process.h"
 #include "LC_Memory.h"
 #include "LC_PointerUtil.h"
 #include "LC_VirtualMemory.h"
-#include "LC_WindowsInternals.h"
 #include "LC_WindowsInternalFunctions.h"
+#include "LC_Thread.h"
+// BEGIN EPIC MOD
 #include "LC_Logging.h"
 #include <Psapi.h>
 #include "Containers/Map.h"
 #include "Containers/UnrealString.h"
+// END EPIC MOD
 
 // BEGIN EPIC MODS
 #pragma warning(push)
@@ -17,18 +22,90 @@
 #pragma warning(disable:6335) // warning C6335: Leaking process information handle 'context->pi.hProcess'.
 // END EPIC MODS
 
+namespace Process
+{
+	struct Context
+	{
+		uint32_t flags;
+		HANDLE pipeReadEnd;
+		PROCESS_INFORMATION processInformation;
+		Thread::Handle drainStdoutThread;
+		std::wstring stdoutData;
+	};
+}
+
 
-namespace process
+void* Process::Current::GetBase(void)
 {
-	static unsigned int __stdcall DrainPipe(void* data)
-	{
-		Context* context = static_cast<Context*>(data);
+	return ::GetModuleHandle(NULL);
+}
+
+
+Process::Handle Process::Current::GetHandle(void)
+{
+	return Process::Handle(::GetCurrentProcess());
+}
+
+
+Process::Id Process::Current::GetId(void)
+{
+	return Process::Id(::GetCurrentProcessId());
+}
+
+
+Filesystem::Path Process::Current::GetImagePath(void)
+{
+	wchar_t filename[Filesystem::Path::CAPACITY] = { '\0' };
+	::GetModuleFileNameW(NULL, filename, Filesystem::Path::CAPACITY);
+
+	return Filesystem::Path(filename);
+}
+
+
+Filesystem::Path Process::Current::GetWorkingDirectory(void)
+{
+	wchar_t workingDirectory[Filesystem::Path::CAPACITY] = { '\0' };
+	::GetCurrentDirectoryW(Filesystem::Path::CAPACITY, workingDirectory);
+
+	return Filesystem::Path(workingDirectory);
+}
+
+
+std::wstring Process::Current::GetCommandLine(void)
+{
+	return std::wstring(::GetCommandLineW());
+}
+
+
+Process::Id Process::GetId(const Context* context)
+{
+	return Process::Id(context->processInformation.dwProcessId);
+}
+
+
+Process::Handle Process::GetHandle(const Context* context)
+{
+	return Process::Handle(context->processInformation.hProcess);
+}
+
+
+std::wstring Process::GetStdOutData(const Context* context)
+{
+	return context->stdoutData;
+}
 
+
+namespace
+{
+	static Thread::ReturnValue DrainPipe(Process::Context* context)
+	{
 		std::vector<char> stdoutData;
+		stdoutData.reserve(1024u);
+
+		char buffer[1024u];
 		for (;;)
 		{
 			DWORD bytesRead = 0u;
-			char buffer[256] = {};
 			if (!::ReadFile(context->pipeReadEnd, buffer, sizeof(buffer) - 1u, &bytesRead, NULL))
 			{
 				// error while trying to read from the pipe, process has probably ended and closed its end of the pipe
@@ -59,622 +136,614 @@ namespace process
 			delete[] strTo;
 		}
 
-		return 0u;
-	}
-
-
-	unsigned int GetId(void)
-	{
-		return ::GetCurrentProcessId();
+		return Thread::ReturnValue(0u);
 	}
+}
 
 
-	Context* Spawn(const wchar_t* exePath, const wchar_t* workingDirectory, const wchar_t* commandLine, const void* environmentBlock, uint32_t flags)
-	{
-		Context* context = new Context { flags, nullptr, PROCESS_INFORMATION {}, nullptr, {} };
+Process::Context* Process::Spawn(const wchar_t* exePath, const wchar_t* workingDirectory, const wchar_t* commandLine, const void* environmentBlock, uint32_t flags)
+{
+	Context* context = new Context { flags, nullptr, ::PROCESS_INFORMATION {}, Thread::Handle(nullptr), {} };
 
-		::SECURITY_ATTRIBUTES saAttr;
-		saAttr.nLength = sizeof(SECURITY_ATTRIBUTES);
-		saAttr.bInheritHandle = Windows::TRUE;
-		saAttr.lpSecurityDescriptor = NULL;
+	::SECURITY_ATTRIBUTES saAttr;
+	saAttr.nLength = sizeof(SECURITY_ATTRIBUTES);
+	// BEGIN EPIC MOD
+	saAttr.bInheritHandle = Windows::TRUE;
+	// END EPIC MOD
+	saAttr.lpSecurityDescriptor = NULL;
 
-		::STARTUPINFOW startupInfo = {};
-		startupInfo.cb = sizeof(startupInfo);
+	::STARTUPINFOW startupInfo = {};
+	startupInfo.cb = sizeof(startupInfo);
 
-		HANDLE hProcessStdOutRead = NULL;
-		HANDLE hProcessStdOutWrite = NULL;
-		HANDLE hProcessStdErrWrite = NULL;
+	// BEGIN EPIC MOD
+	Windows::HANDLE hProcessStdOutRead = NULL;
+	Windows::HANDLE hProcessStdOutWrite = NULL;
+	Windows::HANDLE hProcessStdErrWrite = NULL;
+	// END EPIC MOD
 
-		if (flags & SpawnFlags::REDIRECT_STDOUT)
+	if (flags & SpawnFlags::REDIRECT_STDOUT)
+	{
+		// create a STD_OUT pipe for the child process
+		if (!::CreatePipe(&hProcessStdOutRead, &hProcessStdOutWrite, &saAttr, 0))
 		{
-			// create a STD_OUT pipe for the child process
-			if (!CreatePipe(&hProcessStdOutRead, &hProcessStdOutWrite, &saAttr, 0))
-			{
-				LC_ERROR_USER("Cannot create stdout pipe. Error: 0x%X", ::GetLastError());
-				delete context;
-				return nullptr;
-			}
-
-			// create a duplicate of the STD_OUT write handle for the STD_ERR write handle. this is necessary in case the child
-			// application closes one of its STD output handles.
-			if (!::DuplicateHandle(::GetCurrentProcess(), hProcessStdOutWrite, ::GetCurrentProcess(),
-				&hProcessStdErrWrite, 0, Windows::TRUE, DUPLICATE_SAME_ACCESS))
-			{
-				LC_ERROR_USER("Cannot duplicate stdout pipe. Error: 0x%X", ::GetLastError());
-				::CloseHandle(hProcessStdOutRead);
-				::CloseHandle(hProcessStdOutWrite);
-				delete context;
-				return nullptr;
-			}
-
-			// the spawned process will output data into the write-end of the pipe, and our process will read from the
-			// read-end. because pipes can only do some buffering, we need to ensure that pipes never get clogged, otherwise
-			// the spawned process could block due to the pipe being full.
-			// therefore, we also create a new thread that continuously reads data from the pipe on our end.
-			context->pipeReadEnd = hProcessStdOutRead;
-			context->threadId = thread::Create(64u * 1024u, &DrainPipe, context);
-
-			startupInfo.hStdOutput = hProcessStdOutWrite;
-			startupInfo.hStdError = hProcessStdErrWrite;
-			startupInfo.dwFlags = STARTF_USESTDHANDLES;
+			LC_ERROR_USER("Cannot create stdout pipe. Error: 0x%X", ::GetLastError());
+			delete context;
+			return nullptr;
 		}
 
-		wchar_t* commandLineBuffer = nullptr;
-		if (commandLine)
+		// create a duplicate of the STD_OUT write handle for the STD_ERR write handle. this is necessary in case the child
+		// application closes one of its STD output handles.
+		// BEGIN EPIC MOD
+		if (!::DuplicateHandle(::GetCurrentProcess(), hProcessStdOutWrite, ::GetCurrentProcess(),
+			&hProcessStdErrWrite, 0, Windows::TRUE, DUPLICATE_SAME_ACCESS))
+		// END EPIC MOD
 		{
-			commandLineBuffer = new wchar_t[32768];
-			wcscpy_s(commandLineBuffer, 32768u, commandLine);
+			LC_ERROR_USER("Cannot duplicate stdout pipe. Error: 0x%X", ::GetLastError());
+			::CloseHandle(hProcessStdOutRead);
+			::CloseHandle(hProcessStdOutWrite);
+			delete context;
+			return nullptr;
 		}
 
-		LC_LOG_DEV("Spawning process:");
-		{
-			LC_LOG_INDENT_DEV;
-			LC_LOG_DEV("Executable: %S", exePath);
-			LC_LOG_DEV("Command line: %S", commandLineBuffer ? commandLineBuffer : L"none");
-			LC_LOG_DEV("Working directory: %S", workingDirectory ? workingDirectory : L"none");
-			LC_LOG_DEV("Custom environment block: %S", environmentBlock ? L"yes" : L"no");
-			LC_LOG_DEV("Flags: %u", flags);
-		}
+		// the spawned process will output data into the write-end of the pipe, and our process will read from the
+		// read-end. because pipes can only do some buffering, we need to ensure that pipes never get clogged, otherwise
+		// the spawned process could block due to the pipe being full.
+		// therefore, we also create a new thread that continuously reads data from the pipe on our end.
+		context->pipeReadEnd = hProcessStdOutRead;
+		context->drainStdoutThread = Thread::CreateFromFunction("DrainStdoutPipe", 16u * 1024u, &DrainPipe, context);
 
-		DWORD creationFlags = CREATE_UNICODE_ENVIRONMENT;
-		if (flags & SpawnFlags::NO_WINDOW)
-		{
-			creationFlags |= CREATE_NO_WINDOW;
-		}
-		if (flags & SpawnFlags::SUSPENDED)
-		{
-			creationFlags |= CREATE_SUSPENDED;
-		}
+		startupInfo.hStdOutput = hProcessStdOutWrite;
+		startupInfo.hStdError = hProcessStdErrWrite;
+		startupInfo.dwFlags = STARTF_USESTDHANDLES;
+	}
 
-		// the environment block is not written to by CreateProcess, so it is safe to const_cast (it's a Win32 API mistake)
-		const BOOL success = ::CreateProcessW(exePath, commandLineBuffer, NULL, NULL, Windows::TRUE, creationFlags, const_cast<void*>(environmentBlock), workingDirectory, &startupInfo, &context->pi);
-		if (success == 0)
-		{
-			LC_ERROR_USER("Could not spawn process %S. Error: %d", exePath, ::GetLastError());
-		}
+	wchar_t* commandLineBuffer = nullptr;
+	if (commandLine)
+	{
+		commandLineBuffer = new wchar_t[32768];
+		::wcscpy_s(commandLineBuffer, 32768u, commandLine);
+	}
 
-		delete[] commandLineBuffer;
+	LC_LOG_DEV("Spawning process:");
+	{
+		LC_LOG_INDENT_DEV;
+		LC_LOG_DEV("Executable: %S", exePath);
+		LC_LOG_DEV("Command line: %S", commandLineBuffer ? commandLineBuffer : L"none");
+		LC_LOG_DEV("Working directory: %S", workingDirectory ? workingDirectory : L"none");
+		LC_LOG_DEV("Custom environment block: %S", environmentBlock ? L"yes" : L"no");
+		LC_LOG_DEV("Flags: %u", flags);
+	}
 
-		if (flags & SpawnFlags::REDIRECT_STDOUT)
-		{
-			// we don't need those ends of the pipe
-			::CloseHandle(hProcessStdOutWrite);
-			::CloseHandle(hProcessStdErrWrite);
-		}
+	DWORD creationFlags = CREATE_UNICODE_ENVIRONMENT;
+	if (flags & SpawnFlags::NO_WINDOW)
+	{
+		creationFlags |= CREATE_NO_WINDOW;
+	}
+	if (flags & SpawnFlags::SUSPENDED)
+	{
+		creationFlags |= CREATE_SUSPENDED;
+	}
 
-		return context;
+	// the environment block is not written to by CreateProcess, so it is safe to const_cast (it's a Win32 API mistake)
+	// BEGIN EPIC MOD
+	const BOOL success = ::CreateProcessW(exePath, commandLineBuffer, NULL, NULL, Windows::TRUE, creationFlags, const_cast<void*>(environmentBlock), workingDirectory, &startupInfo, &context->processInformation);
+	// END EPIC MOD
+	if (success == 0)
+	{
+		LC_ERROR_USER("Could not spawn process %S. Error: %d", exePath, ::GetLastError());
 	}
 
+	delete[] commandLineBuffer;
 
-	void ResumeMainThread(Context* context)
+	if (flags & SpawnFlags::REDIRECT_STDOUT)
 	{
-		::ResumeThread(context->pi.hThread);
+		// we don't need those ends of the pipe
+		::CloseHandle(hProcessStdOutWrite);
+		::CloseHandle(hProcessStdErrWrite);
 	}
 
+	return context;
+}
 
-	unsigned int Wait(Context* context)
-	{
-		// wait until process terminates
-		::WaitForSingleObject(context->pi.hProcess, INFINITE);
 
-		if (context->flags & SpawnFlags::REDIRECT_STDOUT)
-		{
-			// wait until all data is drained from the pipe
-			thread::Join(context->threadId);
-			thread::Close(context->threadId);
+void Process::Destroy(Context*& context)
+{
+	::CloseHandle(context->processInformation.hProcess);
+	::CloseHandle(context->processInformation.hThread);
 
-			// close remaining pipe handles
-			::CloseHandle(context->pipeReadEnd);
-		}
+	memory::DeleteAndNull(context);
+}
 
-		DWORD exitCode = 0xFFFFFFFFu;
-		::GetExitCodeProcess(context->pi.hProcess, &exitCode);
 
-		return exitCode;
-	}
+void Process::ResumeMainThread(Context* context)
+{
+	::ResumeThread(context->processInformation.hThread);
+}
 
 
-	unsigned int Wait(Handle handle)
-	{
-		// wait until process terminates
-		::WaitForSingleObject(handle, INFINITE);
+unsigned int Process::Wait(Context* context)
+{
+	// wait until process terminates
+	::WaitForSingleObject(context->processInformation.hProcess, INFINITE);
 
-		DWORD exitCode = 0xFFFFFFFFu;
-		::GetExitCodeProcess(handle, &exitCode);
+	if (context->flags & SpawnFlags::REDIRECT_STDOUT)
+	{
+		// wait until all data is drained from the pipe
+		Thread::Join(context->drainStdoutThread);
+		Thread::Close(context->drainStdoutThread);
 
-		return exitCode;
+		// close remaining pipe handles
+		::CloseHandle(context->pipeReadEnd);
 	}
 
+	DWORD exitCode = 0xFFFFFFFFu;
+	::GetExitCodeProcess(context->processInformation.hProcess, &exitCode);
 
-	void Destroy(Context*& context)
-	{
-		::CloseHandle(context->pi.hProcess);
-		::CloseHandle(context->pi.hThread);
+	return exitCode;
+}
 
-		memory::DeleteAndNull(context);
-	}
 
+unsigned int Process::Wait(Handle handle)
+{
+	// wait until process terminates
+	::WaitForSingleObject(+handle, INFINITE);
 
-	void Terminate(Handle processHandle)
-	{
-		::TerminateProcess(processHandle, 0u);
+	DWORD exitCode = 0xFFFFFFFFu;
+	::GetExitCodeProcess(+handle, &exitCode);
 
-		// termination is asynchronous, wait until the process is really gone
-		::WaitForSingleObject(processHandle, INFINITE);
-	}
+	return exitCode;
+}
 
 
-	Handle Open(unsigned int processId)
-	{
-		return ::OpenProcess(PROCESS_ALL_ACCESS, Windows::FALSE, processId);
-	}
+void Process::Terminate(Context* context)
+{
+	Terminate(Process::Handle(context->processInformation.hProcess));
+}
 
 
-	void Close(Handle& handle)
-	{
-		::CloseHandle(handle);
-		handle = INVALID_HANDLE_VALUE;
-	}
+void Process::Terminate(Handle handle)
+{
+	::TerminateProcess(+handle, 0u);
 
+	// termination is asynchronous, wait until the process is really gone
+	::WaitForSingleObject(+handle, INFINITE);
+}
 
-	std::wstring GetImagePath(Handle handle)
-	{
-		DWORD charCount = MAX_PATH + 1u;
-		wchar_t processName[MAX_PATH + 1u] = {};
-		::QueryFullProcessImageName(handle, 0u, processName, &charCount);
 
-		return std::wstring(processName);
-	}
+Process::Handle Process::Open(Id processId)
+{
+	// BEGIN EPIC MOD
+	return Process::Handle(::OpenProcess(PROCESS_ALL_ACCESS, Windows::FALSE, +processId));
+	// END EPIC MOD
+}
 
 
-	std::wstring GetWorkingDirectory(void)
-	{
-		wchar_t workingDirectory[MAX_PATH + 1] = {};
-		::GetCurrentDirectoryW(MAX_PATH + 1u, workingDirectory);
+void Process::Close(Handle& handle)
+{
+	::CloseHandle(+handle);
+	handle = INVALID_HANDLE;
+}
 
-		return std::wstring(workingDirectory);
-	}
 
+void Process::Suspend(Context* context)
+{
+	WindowsInternals::NtSuspendProcess(context->processInformation.hProcess);
+}
+
+
+void Process::Suspend(Handle handle)
+{
+	WindowsInternals::NtSuspendProcess(+handle);
+}
+
+
+void Process::Resume(Context* context)
+{
+	WindowsInternals::NtResumeProcess(context->processInformation.hProcess);
+}
 
-	std::wstring GetCommandLine(void)
-	{
-		return std::wstring(::GetCommandLineW());
-	}
+
+void Process::Resume(Handle handle)
+{
+	WindowsInternals::NtResumeProcess(+handle);
+}
 
 
-	void* GetBase(void)
+bool Process::IsActive(Handle handle)
+{
+	DWORD exitCode = 0u;
+	const BOOL success = ::GetExitCodeProcess(+handle, &exitCode);
+	if ((success != 0) && (exitCode == STILL_ACTIVE))
 	{
-		return ::GetModuleHandle(NULL);
+		return true;
 	}
 
+	// either the function has failed (because the process terminated unexpectedly) or the exit code
+	// signals that the process exited already.
+	return false;
+}
+
 
-	std::wstring GetImagePath(void)
-	{
-		wchar_t filename[MAX_PATH + 1u] = {};
-		::GetModuleFileNameW(NULL, filename, MAX_PATH + 1u);
+bool Process::IsWoW64(Handle handle)
+{
+	// a WoW64 process has a PEB32 instead of a real PEB.
+	// if we get a meaningful pointer to this PEB32, the process is running under WoW64.
+	ULONG_PTR peb32 = 0u;
+	WindowsInternals::NtQueryInformationProcess(+handle, WindowsInternals::ProcessWow64Information, &peb32, sizeof(ULONG_PTR), NULL);
 
-		return std::wstring(filename);
-	}
+	return (peb32 != 0u);
+}
 
 
-	uint32_t GetImageSize(Handle handle, void* moduleBase)
+void Process::ReadProcessMemory(Handle handle, const void* srcAddress, void* destBuffer, size_t size)
+{
+	WindowsInternals::NtReadVirtualMemory(+handle, const_cast<PVOID>(srcAddress), destBuffer, static_cast<ULONG>(size), NULL);
+}
+
+
+void Process::WriteProcessMemory(Handle handle, void* destAddress, const void* srcBuffer, size_t size)
+{
+	DWORD oldProtect = 0u;
+	::VirtualProtectEx(+handle, destAddress, size, PAGE_READWRITE, &oldProtect);
 	{
-		MODULEINFO info = {};
-		::GetModuleInformation(handle, static_cast<HMODULE>(moduleBase), &info, sizeof(MODULEINFO));
-		return info.SizeOfImage;
+		// instead of the regular WriteProcessMemory function, we use an undocumented function directly.
+		// this is because Windows 10 introduced a performance regression that causes WriteProcessMemory to be 100 times slower (!)
+		// than in previous versions of Windows.
+		// this bug was reported here:
+		// https://developercommunity.visualstudio.com/content/problem/228061/writeprocessmemory-slowdown-on-windows-10.html
+		WindowsInternals::NtWriteVirtualMemory(+handle, destAddress, const_cast<PVOID>(srcBuffer), static_cast<ULONG>(size), NULL);
 	}
+	::VirtualProtectEx(+handle, destAddress, size, oldProtect, &oldProtect);
+}
 
 
-	bool IsActive(Handle handle)
+void* Process::ScanMemoryRange(Handle handle, const void* lowerBound, const void* upperBound, size_t size, size_t alignment)
+{
+	for (const void* scan = lowerBound; /* nothing */; /* nothing */)
 	{
-		DWORD exitCode = 0u;
-		const BOOL success = ::GetExitCodeProcess(handle, &exitCode);
-		if ((success != 0) && (exitCode == STILL_ACTIVE))
+		// align address to be scanned
+		scan = pointer::AlignTop<const void*>(scan, alignment);
+		if (pointer::Offset<const void*>(scan, size) >= upperBound)
+		{
+			// outside of range to scan
+			LC_ERROR_DEV("Could not find memory range that fits 0x%X bytes with alignment 0x%X in range from 0x%p to 0x%p (scan: 0x%p)", size, alignment, lowerBound, upperBound, scan);
+			return nullptr;
+		}
+		else if (scan < lowerBound)
 		{
-			return true;
+			// outside of range (possible wrap-around)
+			LC_ERROR_DEV("Could not find memory range that fits 0x%X bytes with alignment 0x%X in range from 0x%p to 0x%p (scan: 0x%p)", size, alignment, lowerBound, upperBound, scan);
+			return nullptr;
 		}
 
-		// either the function has failed (because the process terminated unexpectedly) or the exit code
-		// signals that the process exited already.
-		return false;
-	}
+		::MEMORY_BASIC_INFORMATION memoryInfo = {};
+		::VirtualQueryEx(+handle, scan, &memoryInfo, sizeof(::MEMORY_BASIC_INFORMATION));
 
-	
-	void ReadProcessMemory(Handle handle, const void* srcAddress, void* destBuffer, size_t size)
-	{
-		const BOOL success = ::ReadProcessMemory(handle, srcAddress, destBuffer, size, NULL);
-		if (success == 0)
+		if ((memoryInfo.RegionSize >= size) && (memoryInfo.State == MEM_FREE))
 		{
-			LC_ERROR_USER("Cannot read %zu bytes from remote process at address 0x%p. Error: 0x%X", size, srcAddress, ::GetLastError());
+			return memoryInfo.BaseAddress;
 		}
+
+		// keep on searching
+		scan = pointer::Offset<const void*>(memoryInfo.BaseAddress, memoryInfo.RegionSize);
 	}
+}
 
 
-	void WriteProcessMemory(Handle handle, void* destAddress, const void* srcBuffer, size_t size)
+uint32_t Process::ConvertPageProtectionToExecutableProtection(uint32_t protection)
+{
+	// cut off PAGE_GUARD, PAGE_NOCACHE, PAGE_WRITECOMBINE, and PAGE_REVERT_TO_FILE_MAP
+	const uint32_t extraBits = protection & 0xFFFFFF00u;
+	const uint32_t pageProtection = protection & 0x000000FFu;
+
+	switch (pageProtection)
 	{
-		DWORD oldProtect = 0u;
-		::VirtualProtectEx(handle, destAddress, size, PAGE_READWRITE, &oldProtect);
-		{
-			// instead of the regular WriteProcessMemory function, we use an undocumented function directly.
-			// this is because Windows 10 introduced a performance regression that causes WriteProcessMemory to be 100 times slower (!)
-			// than in previous versions of Windows.
-			// this bug was reported here:
-			// https://developercommunity.visualstudio.com/content/problem/228061/writeprocessmemory-slowdown-on-windows-10.html
-			windowsInternal::NtWriteVirtualMemory(handle, destAddress, const_cast<PVOID>(srcBuffer), static_cast<ULONG>(size), NULL);
-		}
-		::VirtualProtectEx(handle, destAddress, size, oldProtect, &oldProtect);
+		case PAGE_NOACCESS:
+		case PAGE_READONLY:
+		case PAGE_READWRITE:
+		case PAGE_WRITECOPY:
+			return (pageProtection << 4u) | extraBits;
+
+		case PAGE_EXECUTE:
+		case PAGE_EXECUTE_READ:
+		case PAGE_EXECUTE_READWRITE:
+		case PAGE_EXECUTE_WRITECOPY:
+		default:
+			return protection;
 	}
+}
+
 
+void Process::MakePagesExecutable(Handle handle, void* address, size_t size)
+{
+	const uint32_t pageSize = VirtualMemory::GetPageSize();
+	const void* endOfRegion = pointer::Offset<const void*>(address, size);
 
-	void* ScanMemoryRange(Handle handle, const void* lowerBound, const void* upperBound, size_t size, size_t alignment)
+	for (const void* scan = address; /* nothing */; /* nothing */)
 	{
-		for (const void* scan = lowerBound; /* nothing */; /* nothing */)
+		::MEMORY_BASIC_INFORMATION memoryInfo = {};
+		const SIZE_T bytesInBuffer = ::VirtualQueryEx(+handle, scan, &memoryInfo, sizeof(::MEMORY_BASIC_INFORMATION));
+		if (bytesInBuffer == 0u)
 		{
-			// align address to be scanned
-			scan = pointer::AlignTop<const void*>(scan, alignment);
-			if (pointer::Offset<const void*>(scan, size) >= upperBound)
-			{
-				// outside of range to scan
-				LC_ERROR_DEV("Could not find memory range that fits 0x%X bytes with alignment 0x%X in range from 0x%p to 0x%p (scan: 0x%p)", size, alignment, lowerBound, upperBound, scan);
-				return nullptr;
-			}
-			else if (scan < lowerBound)
-			{
-				// outside of range (possible wrap-around)
-				LC_ERROR_DEV("Could not find memory range that fits 0x%X bytes with alignment 0x%X in range from 0x%p to 0x%p (scan: 0x%p)", size, alignment, lowerBound, upperBound, scan);
-				return nullptr;
-			}
-
-			::MEMORY_BASIC_INFORMATION memoryInfo = {};
-			::VirtualQueryEx(handle, scan, &memoryInfo, sizeof(::MEMORY_BASIC_INFORMATION));
+			// could not query the protection, bail out
+			break;
+		}
 
-			if ((memoryInfo.RegionSize >= size) && (memoryInfo.State == MEM_FREE))
-			{
-				return memoryInfo.BaseAddress;
-			}
+		const uint32_t executableProtection = ConvertPageProtectionToExecutableProtection(memoryInfo.Protect);
+		if (executableProtection != memoryInfo.Protect)
+		{
+			// change this page into an executable one
+			DWORD oldProtection = 0u;
+			::VirtualProtectEx(+handle, memoryInfo.BaseAddress, pageSize, executableProtection, &oldProtection);
+		}
 
-			// keep on searching
-			scan = pointer::Offset<const void*>(memoryInfo.BaseAddress, memoryInfo.RegionSize);
+		const void* endOfThisRegion = pointer::Offset<const void*>(memoryInfo.BaseAddress, pageSize);
+		if (endOfThisRegion >= endOfRegion)
+		{
+			// we are done
+			break;
 		}
+
+		// keep on walking pages
+		scan = endOfThisRegion;
 	}
+}
 
 
-	void MakePagesExecutable(Handle handle, void* address, size_t size)
-	{
-		const uint32_t pageSize = virtualMemory::GetPageSize();
-		const void* endOfRegion = pointer::Offset<const void*>(address, size);
+void Process::FlushInstructionCache(Handle handle, void* address, size_t size)
+{
+	::FlushInstructionCache(+handle, address, size);
+}
 
-		for (const void* scan = address; /* nothing */; /* nothing */)
-		{
-			::MEMORY_BASIC_INFORMATION memoryInfo = {};
-			const SIZE_T bytesInBuffer = ::VirtualQueryEx(handle, scan, &memoryInfo, sizeof(::MEMORY_BASIC_INFORMATION));
-			if (bytesInBuffer == 0u)
-			{
-				// could not query the protection, bail out
-				break;
-			}
 
-			const uint32_t executableProtection = ConvertPageProtectionToExecutableProtection(memoryInfo.Protect);
-			if (executableProtection != memoryInfo.Protect)
-			{
-				// change this page into an executable one
-				DWORD oldProtection = 0u;
-				::VirtualProtectEx(handle, memoryInfo.BaseAddress, pageSize, executableProtection, &oldProtection);
-			}
 
-			const void* endOfThisRegion = pointer::Offset<const void*>(memoryInfo.BaseAddress, pageSize);
-			if (endOfThisRegion >= endOfRegion)
-			{
-				// we are done
-				break;
-			}
+Process::Environment Process::CreateEnvironment(Handle handle)
+{
+	const void* processEnvironment = nullptr;
 
-			// keep on walking pages
-			scan = endOfThisRegion;
-		}
-	}
+	const bool isWow64 = IsWoW64(handle);
+	if (!isWow64)
+	{
+		// this is either a 32-bit process running on 32-bit Windows, or a 64-bit process running on 64-bit Windows.
+		// the environment can be retrieved directly from the process' PEB and process parameters.
+		WindowsInternals::NT_PROCESS_BASIC_INFORMATION pbi = {};
+		WindowsInternals::NtQueryInformationProcess(+handle, WindowsInternals::ProcessBasicInformation, &pbi, sizeof(pbi), NULL);
 
+		const WindowsInternals::NT_PEB peb = ReadProcessMemory<WindowsInternals::NT_PEB>(handle, pbi.PebBaseAddress);
+		const WindowsInternals::RTL_USER_PROCESS_PARAMETERS parameters = ReadProcessMemory<WindowsInternals::RTL_USER_PROCESS_PARAMETERS>(handle, peb.ProcessParameters);
 
-	void FlushInstructionCache(Handle handle, void* address, size_t size)
-	{
-		::FlushInstructionCache(handle, address, size);
+		processEnvironment = parameters.Environment;
 	}
+	else
+	{
+		// this is a process running under WoW64.
+		// we must get the environment from the PEB32 of the process, rather than the "real" PEB.
+		ULONG_PTR peb32Wow64 = 0u;
+		WindowsInternals::NtQueryInformationProcess(+handle, WindowsInternals::ProcessWow64Information, &peb32Wow64, sizeof(ULONG_PTR), NULL);
 
+		const WindowsInternals::NT_PEB32 peb32 = ReadProcessMemory<WindowsInternals::NT_PEB32>(handle, pointer::FromInteger<const void*>(peb32Wow64));
+		const WindowsInternals::RTL_USER_PROCESS_PARAMETERS32 parameters32 = ReadProcessMemory<WindowsInternals::RTL_USER_PROCESS_PARAMETERS32>(handle, pointer::FromInteger<const void*>(peb32.ProcessParameters32));
 
-	void Suspend(Handle handle)
-	{
-		windowsInternal::NtSuspendProcess(handle);
+		processEnvironment = pointer::FromInteger<const void*>(parameters32.Environment);
 	}
 
-
-	void Resume(Handle handle)
+	if (!processEnvironment)
 	{
-		windowsInternal::NtResumeProcess(handle);
+		return Environment { 0u, nullptr };
 	}
 
-
-	void Continue(CONTEXT* threadContext)
+	// query the size of the page(s) the environment is stored in
+	::MEMORY_BASIC_INFORMATION memoryInfo = {};
+	const SIZE_T bytesInBuffer = ::VirtualQueryEx(+handle, processEnvironment, &memoryInfo, sizeof(::MEMORY_BASIC_INFORMATION));
+	if (bytesInBuffer == 0u)
 	{
-		windowsInternal::NtContinue(threadContext, Windows::FALSE);
+		// operation failed, bail out
+		return Environment { 0u, nullptr };
 	}
 
+	Environment environment = {};
+	environment.size = memoryInfo.RegionSize - (reinterpret_cast<uintptr_t>(processEnvironment) - reinterpret_cast<uintptr_t>(memoryInfo.BaseAddress));
+	environment.data = ::malloc(environment.size);
 
-	std::vector<unsigned int> EnumerateThreads(unsigned int processId)
-	{
-		std::vector<unsigned int> threadIds;
-		threadIds.reserve(256u);
+	ReadProcessMemory(handle, processEnvironment, environment.data, environment.size);
 
-		// 2MB should be enough for getting the process information, even on systems with high load
-		ULONG bufferSize = 2048u * 1024u;
-		void* processSnapshot = nullptr;
-		windowsInternal::NTSTATUS status = 0;
+	return environment;
+}
 
-		do
-		{
-			processSnapshot = ::malloc(bufferSize);
 
-			// try getting a process snapshot into the provided buffer
-			status = windowsInternal::NtQuerySystemInformation(windowsInternal::SystemProcessInformation, processSnapshot, bufferSize, NULL);
+Process::Environment Process::CreateEnvironment(Context* context)
+{
+	return CreateEnvironment(Process::Handle(context->processInformation.hProcess));
+}
 
-			if (status == STATUS_INFO_LENGTH_MISMATCH)
-			{
-				// buffer is too small, try again
-				::free(processSnapshot);
-				processSnapshot = nullptr;
 
-				bufferSize *= 2u;
-			}
-			else if (status < 0)
-			{
-				// something went wrong
-				LC_ERROR_USER("Cannot enumerate threads in process (PID: %d)", processId);
-				::free(processSnapshot);
+void Process::DestroyEnvironment(Environment& environment)
+{
+	::free(environment.data);
+	environment.data = nullptr;
+	environment.size = 0u;
+}
 
-				return threadIds;
-			}
-		}
-		while (status == STATUS_INFO_LENGTH_MISMATCH);
+// BEGIN EPIC MOD
+std::vector<Thread::Id> Process::EnumerateThreads(Id processId)
+// END EPIC MOD
+{
+	// BEGIN EPIC MOD
+	std::vector<Thread::Id> threadIds;
+	// END EPIC MOD
+	threadIds.reserve(256u);
 
-		// find the process information for the given process ID
-		{
-			windowsInternal::NT_SYSTEM_PROCESS_INFORMATION* processInfo = static_cast<windowsInternal::NT_SYSTEM_PROCESS_INFORMATION*>(processSnapshot);
+	// 2MB should be enough for getting the process information, even on systems with high load
+	ULONG bufferSize = 2048u * 1024u;
+	void* processSnapshot = nullptr;
+	WindowsInternals::NTSTATUS status = 0;
 
-			while (processInfo != nullptr)
-			{
-				if (processInfo->UniqueProcessId == reinterpret_cast<HANDLE>(static_cast<DWORD_PTR>(processId)))
-				{
-					// we found the process we're looking for
-					break;
-				}
+	do
+	{
+		processSnapshot = ::malloc(bufferSize);
 
-				if (processInfo->NextEntryOffset == 0u)
-				{
-					// we couldn't find our process
-					LC_ERROR_USER("Cannot enumerate threads, process not found (PID: %d)", processId);
-					::free(processSnapshot);
+		// try getting a process snapshot into the provided buffer
+		status = WindowsInternals::NtQuerySystemInformation(WindowsInternals::SystemProcessInformation, processSnapshot, bufferSize, NULL);
 
-					return threadIds;
-				}
-				else
-				{
-					// walk to the next process info
-					processInfo = pointer::Offset<windowsInternal::NT_SYSTEM_PROCESS_INFORMATION*>(processInfo, processInfo->NextEntryOffset);
-				}
-			}
+		if (status == STATUS_INFO_LENGTH_MISMATCH)
+		{
+			// buffer is too small, try again
+			::free(processSnapshot);
+			processSnapshot = nullptr;
 
-			// record all threads belonging to the given process
-			if (processInfo)
-			{
-				for (ULONG i = 0u; i < processInfo->NumberOfThreads; ++i)
-				{
-					const DWORD threadId = static_cast<DWORD>(reinterpret_cast<DWORD_PTR>(processInfo->Threads[i].ClientId.UniqueThread));
-					threadIds.push_back(threadId);
-				}
-			}
+			bufferSize *= 2u;
 		}
+		else if (status < 0)
+		{
+			// something went wrong
+			LC_ERROR_USER("Cannot enumerate threads in process (PID: %d)", processId);
+			::free(processSnapshot);
 
-		::free(processSnapshot);
-
-		return threadIds;
+			return threadIds;
+		}
 	}
+	while (status == STATUS_INFO_LENGTH_MISMATCH);
 
-
-	std::vector<Module> EnumerateModules(Handle handle)
+	// find the process information for the given process ID
 	{
-		// 1024 modules should be enough for most processes
-		std::vector<Module> modules;
-		modules.reserve(1024u);
-
-		windowsInternal::NT_PROCESS_BASIC_INFORMATION pbi = {};
-		windowsInternal::NtQueryInformationProcess(handle, windowsInternal::ProcessBasicInformation, &pbi, sizeof(pbi), NULL);
+		WindowsInternals::NT_SYSTEM_PROCESS_INFORMATION* processInfo = static_cast<WindowsInternals::NT_SYSTEM_PROCESS_INFORMATION*>(processSnapshot);
 
-		const windowsInternal::NT_PEB processPEB = ReadProcessMemory<windowsInternal::NT_PEB>(handle, pbi.PebBaseAddress);
-		const windowsInternal::NT_PEB_LDR_DATA loaderData = ReadProcessMemory<windowsInternal::NT_PEB_LDR_DATA>(handle, processPEB.Ldr);
-
-		LIST_ENTRY* listHeader = loaderData.InLoadOrderModuleList.Flink;
-		LIST_ENTRY* currentNode = listHeader;
-		do
+		while (processInfo != nullptr)
 		{
-			const windowsInternal::NT_LDR_DATA_TABLE_ENTRY entry = ReadProcessMemory<windowsInternal::NT_LDR_DATA_TABLE_ENTRY>(handle, currentNode);
-
-			wchar_t fullDllName[MAX_PATH] = {};
-
-			// certain modules don't have a name and DLL base, skip those
-			if ((entry.DllBase != nullptr) && (entry.FullDllName.Length > 0) && (entry.FullDllName.Buffer != nullptr))
+			if (processInfo->UniqueProcessId == reinterpret_cast<HANDLE>(static_cast<DWORD_PTR>(+processId)))
 			{
-				ReadProcessMemory(handle, entry.FullDllName.Buffer, fullDllName, entry.FullDllName.Length);
-				modules.emplace_back(Module { fullDllName, entry.DllBase, entry.SizeOfImage });
+				// we found the process we're looking for
+				break;
 			}
 
-			currentNode = entry.InLoadOrderLinks.Flink;
-			if (currentNode == nullptr)
+			if (processInfo->NextEntryOffset == 0u)
 			{
-				break;
+				// we couldn't find our process
+				LC_ERROR_USER("Cannot enumerate threads, process not found (PID: %d)", +processId);
+				::free(processSnapshot);
+
+				return threadIds;
+			}
+			else
+			{
+				// walk to the next process info
+				processInfo = pointer::Offset<WindowsInternals::NT_SYSTEM_PROCESS_INFORMATION*>(processInfo, processInfo->NextEntryOffset);
 			}
 		}
-		while (listHeader != currentNode);
-
-		return modules;
-	}
-
 
-	uint32_t ConvertPageProtectionToExecutableProtection(uint32_t protection)
-	{
-		// cut off PAGE_GUARD, PAGE_NOCACHE, PAGE_WRITECOMBINE, and PAGE_REVERT_TO_FILE_MAP
-		const uint32_t extraBits = protection & 0xFFFFFF00u;
-		const uint32_t pageProtection = protection & 0x000000FFu;
-
-		switch (pageProtection)
+		// record all threads belonging to the given process
+		if (processInfo)
 		{
-			case PAGE_NOACCESS:
-			case PAGE_READONLY:
-			case PAGE_READWRITE:
-			case PAGE_WRITECOPY:
-				return (pageProtection << 4u) | extraBits;
-
-			case PAGE_EXECUTE:
-			case PAGE_EXECUTE_READ:
-			case PAGE_EXECUTE_READWRITE:
-			case PAGE_EXECUTE_WRITECOPY:
-			default:
-				return protection;
+			for (ULONG i = 0u; i < processInfo->NumberOfThreads; ++i)
+			{
+				const DWORD threadId = static_cast<DWORD>(reinterpret_cast<DWORD_PTR>(processInfo->Threads[i].ClientId.UniqueThread));
+				threadIds.push_back(Thread::Id(threadId));
+			}
 		}
 	}
 
+	::free(processSnapshot);
 
-	bool IsWoW64(Handle handle)
-	{
-		// a WoW64 process has a PEB32 instead of a real PEB.
-		// if we get a meaningful pointer to this PEB32, the process is running under WoW64.
-		ULONG_PTR peb32 = 0u;
-		windowsInternal::NtQueryInformationProcess(handle, windowsInternal::ProcessWow64Information, &peb32, sizeof(ULONG_PTR), NULL);
+	return threadIds;
+}
 
-		return (peb32 != 0u);
-	}
+// BEGIN EPIC MOD
+std::vector<Process::Module> Process::EnumerateModules(Handle handle)
+// END EPIC MOD
+{
+	// 1024 modules should be enough for most processes
+	// BEGIN EPIC MOD
+	std::vector<Module> modules;
+	// END EPIC MOD
+	modules.reserve(1024u);
 
+	WindowsInternals::NT_PROCESS_BASIC_INFORMATION pbi = {};
+	WindowsInternals::NtQueryInformationProcess(+handle, WindowsInternals::ProcessBasicInformation, &pbi, sizeof(pbi), NULL);
 
-	Environment* CreateEnvironment(Handle handle)
-	{
-		const void* processEnvironment = nullptr;
+	const WindowsInternals::NT_PEB processPEB = ReadProcessMemory<WindowsInternals::NT_PEB>(handle, pbi.PebBaseAddress);
+	const WindowsInternals::NT_PEB_LDR_DATA loaderData = ReadProcessMemory<WindowsInternals::NT_PEB_LDR_DATA>(handle, processPEB.Ldr);
 
-		const bool isWow64 = IsWoW64(handle);
-		if (!isWow64)
-		{
-			// this is either a 32-bit process running on 32-bit Windows, or a 64-bit process running on 64-bit Windows.
-			// the environment can be retrieved directly from the process' PEB and process parameters.
-			windowsInternal::NT_PROCESS_BASIC_INFORMATION pbi = {};
-			windowsInternal::NtQueryInformationProcess(handle, windowsInternal::ProcessBasicInformation, &pbi, sizeof(pbi), NULL);
+	::LIST_ENTRY* listHeader = loaderData.InLoadOrderModuleList.Flink;
+	::LIST_ENTRY* currentNode = listHeader;
+	do
+	{
+		const WindowsInternals::NT_LDR_DATA_TABLE_ENTRY entry = ReadProcessMemory<WindowsInternals::NT_LDR_DATA_TABLE_ENTRY>(handle, currentNode);
 
-			const windowsInternal::NT_PEB peb = ReadProcessMemory<windowsInternal::NT_PEB>(handle, pbi.PebBaseAddress);
-			const windowsInternal::RTL_USER_PROCESS_PARAMETERS parameters = ReadProcessMemory<windowsInternal::RTL_USER_PROCESS_PARAMETERS>(handle, peb.ProcessParameters);
+		wchar_t fullDllName[Filesystem::Path::CAPACITY] = { '\0' };
 
-			processEnvironment = parameters.Environment;
-		}
-		else
+		// certain modules don't have a name and DLL base, skip those
+		if ((entry.DllBase != nullptr) && (entry.FullDllName.Length > 0) && (entry.FullDllName.Buffer != nullptr))
 		{
-			// this is a process running under WoW64.
-			// we must get the environment from the PEB32 of the process, rather than the "real" PEB.
-			ULONG_PTR peb32Wow64 = 0u;
-			windowsInternal::NtQueryInformationProcess(handle, windowsInternal::ProcessWow64Information, &peb32Wow64, sizeof(ULONG_PTR), NULL);
-
-			const windowsInternal::NT_PEB32 peb32 = ReadProcessMemory<windowsInternal::NT_PEB32>(handle, pointer::FromInteger<const void*>(peb32Wow64));
-			const windowsInternal::RTL_USER_PROCESS_PARAMETERS32 parameters32 = ReadProcessMemory<windowsInternal::RTL_USER_PROCESS_PARAMETERS32>(handle, pointer::FromInteger<const void*>(peb32.ProcessParameters32));
-
-			processEnvironment = pointer::FromInteger<const void*>(parameters32.Environment);		
+			ReadProcessMemory(handle, entry.FullDllName.Buffer, fullDllName, entry.FullDllName.Length);
+			modules.emplace_back(Module { Filesystem::Path(fullDllName), entry.DllBase, entry.SizeOfImage });
 		}
 
-		if (!processEnvironment)
+		currentNode = entry.InLoadOrderLinks.Flink;
+		if (currentNode == nullptr)
 		{
-			return nullptr;
+			break;
 		}
+	}
+	while (listHeader != currentNode);
 
-		// query the size of the page(s) the environment is stored in
-		::MEMORY_BASIC_INFORMATION memoryInfo = {};
-		const SIZE_T bytesInBuffer = ::VirtualQueryEx(handle, processEnvironment, &memoryInfo, sizeof(::MEMORY_BASIC_INFORMATION));
-		if (bytesInBuffer == 0u)
-		{
-			// operation failed, bail out
-			return nullptr;
-		}
+	return modules;
+}
 
-		Environment* environment = new Environment;
-		environment->size = memoryInfo.RegionSize - (reinterpret_cast<uintptr_t>(processEnvironment) - reinterpret_cast<uintptr_t>(memoryInfo.BaseAddress));
-		environment->data = ::malloc(environment->size);
 
-		ReadProcessMemory(handle, processEnvironment, environment->data, environment->size);
+Filesystem::Path Process::GetImagePath(Handle handle)
+{
+	DWORD charCount = Filesystem::Path::CAPACITY;
+	wchar_t processName[Filesystem::Path::CAPACITY] = { '\0' };
+	::QueryFullProcessImageName(+handle, 0u, processName, &charCount);
 
-		return environment;
-	}
+	return Filesystem::Path(processName);
+}
 
 
-	// BEGIN EPIC MOD - Allow passing environment block for linker
-	Environment* CreateEnvironmentFromMap(const TMap<FString, FString>& Pairs)
-	{
-		std::vector<wchar_t> environmentData;
-		for (const TPair<FString, FString>& Pair : Pairs)
-		{
-			FString Variable = FString::Printf(TEXT("%s=%s"), *Pair.Key, *Pair.Value);
-			environmentData.insert(environmentData.end(), *Variable, *Variable + (Variable.Len() + 1));
-		}
-		environmentData.push_back('\0');
+uint32_t Process::GetModuleSize(Handle handle, void* moduleBase)
+{
+	::MODULEINFO info = {};
+	::GetModuleInformation(+handle, static_cast<HMODULE>(moduleBase), &info, sizeof(::MODULEINFO));
+	
+	return info.SizeOfImage;
+}
 
-		Environment* environment = new Environment;
-		environment->size = environmentData.size();
-		environment->data = ::malloc(environmentData.size() * sizeof(wchar_t));
-		if (environment->data != nullptr)
-		{
-			memcpy(environment->data, environmentData.data(), environmentData.size() * sizeof(wchar_t));
-		}
-		return environment;
+// BEGIN EPIC MOD - Allow passing environment block for linker
+Process::Environment* Process::CreateEnvironmentFromMap(const TMap<FString, FString>& Pairs)
+{
+	std::vector<wchar_t> environmentData;
+	for (const TPair<FString, FString>& Pair : Pairs)
+	{
+		FString Variable = FString::Printf(TEXT("%s=%s"), *Pair.Key, *Pair.Value);
+		environmentData.insert(environmentData.end(), *Variable, *Variable + (Variable.Len() + 1));
 	}
-	// END EPIC MOD
-
+	environmentData.push_back('\0');
 
-	void DestroyEnvironment(Environment*& environment)
+	Environment* environment = new Environment;
+	environment->size = environmentData.size();
+	environment->data = ::malloc(environmentData.size() * sizeof(wchar_t));
+	if (environment->data != nullptr)
 	{
-		if (environment)
-		{
-			::free(environment->data);
-		}
-
-		memory::DeleteAndNull(environment);
+		memcpy(environment->data, environmentData.data(), environmentData.size() * sizeof(wchar_t));
 	}
+	return environment;
+}
+// END EPIC MOD
 
 
-	void DumpMemory(Handle handle, const void* address, size_t size)
-	{
-		uint8_t* memory = new uint8_t[size];
-		ReadProcessMemory(handle, address, memory, size);
-
-		LC_LOG_DEV("Raw data:");
-		LC_LOG_INDENT_DEV;
-		for (size_t i = 0u; i < size; ++i)
-		{
-			LC_LOG_DEV("0x%02X", memory[i]);
-		}
+void Process::DumpMemory(Handle handle, const void* address, size_t size)
+{
+	uint8_t* memory = new uint8_t[size];
+	ReadProcessMemory(handle, address, memory, size);
 
-		delete[] memory;
+	LC_LOG_DEV("Raw data:");
+	LC_LOG_INDENT_DEV;
+	for (size_t i = 0u; i < size; ++i)
+	{
+		LC_LOG_DEV("0x%02X", memory[i]);
 	}
+
+	delete[] memory;
 }
 
 // BEGIN EPIC MODS
diff --git a/Engine/Source/Developer/Windows/LiveCoding/Private/External/LC_Process.h b/Engine/Source/Developer/Windows/LiveCoding/Private/External/LC_Process.h
--- a/Engine/Source/Developer/Windows/LiveCoding/Private/External/LC_Process.h
+++ b/Engine/Source/Developer/Windows/LiveCoding/Private/External/LC_Process.h
@@ -1,105 +1,109 @@
-// Copyright 2011-2019 Molecular Matters GmbH, all rights reserved.
+// Copyright 2011-2020 Molecular Matters GmbH, all rights reserved.
 
 #pragma once
 
+// BEGIN EPIC MOD
 #include "CoreTypes.h"
-#include "LC_Thread.h"
+//#include "LC_Types.h"
+// END EPIC MOD
+#include "LC_ProcessTypes.h"
+#include "LC_ThreadTypes.h"
+// BEGIN EPIC MOD
 #include <string>
 #include <vector>
+// END EPIC MOD
 
-namespace process
+namespace Process
 {
-	struct Context
-	{
-		uint32_t flags;
-		HANDLE pipeReadEnd;
-		PROCESS_INFORMATION pi;
-		thread::Handle threadId;
-		std::wstring stdoutData;
-	};
-
-	struct Module
-	{
-		std::wstring fullPath;
-		void* baseAddress;
-		uint32_t sizeOfImage;
-	};
+	// Opaque data structure identifying a spawned process.
+	struct Context;
 
-	struct Environment
+	// Current/calling process.
+	namespace Current
 	{
-		size_t size;
-		void* data;
-	};
+		// Returns the base address of the image of the calling process.
+		void* GetBase(void);
 
-	typedef HANDLE Handle;
+		// Returns the process handle of the calling process.
+		Handle GetHandle(void);
 
-	// returns the process ID for the calling process
-	unsigned int GetId(void);
+		// Returns the process Id of the calling process.
+		Id GetId(void);
 
+		// Returns the path to the executable image of the calling process.
+		Filesystem::Path GetImagePath(void);
 
-	struct SpawnFlags
-	{
-		enum Enum : uint32_t
-		{
-			NONE = 0u,
-			REDIRECT_STDOUT = 1u << 0u,
-			NO_WINDOW = 1u << 1u,
-			SUSPENDED = 1u << 2u
-		};
-	};
-
-	// spawns a new process
+		// returns the working directory of the calling process.
+		Filesystem::Path GetWorkingDirectory(void);
+
+		// Returns the command line of the calling process.
+		std::wstring GetCommandLine(void);
+	}
+
+	// Returns the underlying process Id for a process identified by the given context.
+	Id GetId(const Context* context);
+
+	// Returns the underlying process handle for a process identified by the given context.
+	Handle GetHandle(const Context* context);
+
+	// Returns the underlying stdout data for a process identified by the given context.
+	std::wstring GetStdOutData(const Context* context);
+
+
+	// Spawns a new process.
 	Context* Spawn(const wchar_t* exePath, const wchar_t* workingDirectory, const wchar_t* commandLine, const void* environmentBlock, uint32_t flags);
 
-	// resumes a process that was spawned in a suspended state
+	// Destroys a spawned process.
+	void Destroy(Context*& context);
+
+	// Resumes a process that was spawned in a suspended state.
 	void ResumeMainThread(Context* context);
 
-	// waits until a spawned process has exited
+	// Waits until a spawned process has exited.
 	unsigned int Wait(Context* context);
 
-	// waits until a process has exited
+	// Waits until the given process has exited.
 	unsigned int Wait(Handle handle);
 
-	// destroys a spawned process
-	void Destroy(Context*& context);
+	// Terminates a spawned process.
+	void Terminate(Context* context);
 
-	// terminates a spawned process
-	void Terminate(Handle processHandle);
+	// Terminates the given process.
+	void Terminate(Handle handle);
 
 
+	// Opens a process with the given Id for full access.
+	Handle Open(Id processId);
 
-	// opens a process
-	Handle Open(unsigned int processId);
-
-	// closes a process
+	// Closes an opened process.
 	void Close(Handle& handle);
 
-	// returns the full path for a process' image
-	std::wstring GetImagePath(Handle handle);
-
-	// returns the base address of the calling process
-	void* GetBase(void);
 
-	// returns the path to the executable of the calling process
-	std::wstring GetImagePath(void);
+	// Suspends a process.
+	void Suspend(Context* context);
 
-	// returns the working directory of the calling process
-	std::wstring GetWorkingDirectory(void);
+	// Suspends a process.
+	void Suspend(Handle handle);
 
-	// returns the command line of the calling process
-	std::wstring GetCommandLine(void);
+	// Resumes a suspended process.
+	void Resume(Context* context);
 
-	// returns the size of a module loaded into the virtual address space of a given process
-	uint32_t GetImageSize(Handle handle, void* moduleBase);
+	// Resumes a suspended process.
+	void Resume(Handle handle);
 
-	// returns whether the process with the given handle is still active
+	// Returns whether the given process is still active.
 	bool IsActive(Handle handle);
 
+	// Returns whether the given process runs under Wow64 (32-bit emulation on 64-bit versions of Windows).
+	bool IsWoW64(Handle handle);
 
-
-	// reads from process memory
+	// Reads from process memory.
 	void ReadProcessMemory(Handle handle, const void* srcAddress, void* destBuffer, size_t size);
 
+	// Writes to process memory.
+	void WriteProcessMemory(Handle handle, void* destAddress, const void* srcBuffer, size_t size);
+
+	// Convenience function for reading a value of a certain type from process memory.
 	template <typename T>
 	T ReadProcessMemory(Handle handle, const void* srcAddress)
 	{
@@ -109,68 +113,57 @@ namespace process
 		return value;
 	}
 
-	// writes to process memory
-	void WriteProcessMemory(Handle handle, void* destAddress, const void* srcBuffer, size_t size);
-
+	// Convenience function for writing a value of a certain type to process memory.
 	template <typename T>
 	void WriteProcessMemory(Handle handle, void* destAddress, const T& value)
 	{
 		WriteProcessMemory(handle, destAddress, &value, sizeof(T));
 	}
 
-	// scans a region of memory in the given process until a free block of a given size is found.
-	// will only consider blocks at addresses with a certain alignment.
+	// Scans a region of memory in the given process until a free block of a given size is found. Will only consider blocks at addresses with the given alignment.
 	void* ScanMemoryRange(Handle handle, const void* lowerBound, const void* upperBound, size_t size, size_t alignment);
 
-	// makes the memory pages in the given region executable (in case they aren't already) while keeping other protection flags intact
-	void MakePagesExecutable(Handle handle, void* address, size_t size);
-
+	// Converts any combination of page protection flags (e.g. PAGE_NOACCESS, PAGE_GUARD, ...) to protection flags that specify an executable page (e.g. PAGE_EXECUTE).
+	uint32_t ConvertPageProtectionToExecutableProtection(uint32_t protection);
 
+	// Makes the memory pages in the given region executable (in case they aren't already), while keeping other protection flags intact.
+	void MakePagesExecutable(Handle handle, void* address, size_t size);
 
-	// flushes the process' instruction cache
+	// Flushes a process's instruction cache.
 	void FlushInstructionCache(Handle handle, void* address, size_t size);
 
 
+	// Reads and reads the environment of a given process.
+	Environment CreateEnvironment(Handle handle);
 
-	// suspends a process
-	void Suspend(Handle handle);
-
-	// resumes a suspended process
-	void Resume(Handle handle);
+	// Reads the environment of a process identified by the given context.
+	Environment CreateEnvironment(Context* context);
 
-	// continues the calling thread of a process with the given thread context
-	void Continue(CONTEXT* threadContext);
+	// Destroys an environment.
+	void DestroyEnvironment(Environment& environment);
 
 
-	// enumerates all threads of a process, returning their thread IDs.
-	// NOTE: only call on suspended processes!
-	std::vector<unsigned int> EnumerateThreads(unsigned int processId);
+	// Enumerates all threads of a process. Only call on suspended processes.
+	// BEGIN EPIC MOD
+	std::vector<Thread::Id> EnumerateThreads(Id processId);
+	// END EPIC MOD
 
-	// enumerates all modules of a process, returning their info.
-	// NOTE: only call on suspended processes!
+	// Enumerates all modules of a process. Only call on suspended processes.
+	// BEGIN EPIC MOD
 	std::vector<Module> EnumerateModules(Handle handle);
+	// END EPIC MOD
 
+	// Returns the full path of a process's image, e.g. "C:\Directory\Application.exe".
+	Filesystem::Path GetImagePath(Handle handle);
 
-	// converts any combination of page protection flags (e.g. PAGE_NOACCESS, PAGE_GUARD, ...) to protection flags
-	// that specify an executable page (e.g. PAGE_EXECUTE).
-	uint32_t ConvertPageProtectionToExecutableProtection(uint32_t protection);
+	// Returns the size of a module loaded into the virtual address space of a given process.
+	uint32_t GetModuleSize(Handle handle, void* moduleBase);
 
 
-	// returns whether a process runs under Wow64 (32-bit emulation on 64-bit versions of Windows)
-	bool IsWoW64(Handle handle);
-
-	
-	// reads the environment of any process
-	Environment* CreateEnvironment(Handle handle);
+	// Dumps any process's memory.
+	void DumpMemory(Handle handle, const void* address, size_t size);
 
 	// BEGIN EPIC MOD - Allow passing environment block for linker
 	Environment* CreateEnvironmentFromMap(const TMap<FString, FString>& Pairs);
 	// END EPIC MOD
-
-	// destroys an environment
-	void DestroyEnvironment(Environment*& environment);
-
-
-	// dumps raw memory for a given process
-	void DumpMemory(Handle handle, const void* address, size_t size);
 }
diff --git a/Engine/Source/Developer/Windows/LiveCoding/Private/External/LC_Restart.cpp b/Engine/Source/Developer/Windows/LiveCoding/Private/External/LC_Restart.cpp
--- a/Engine/Source/Developer/Windows/LiveCoding/Private/External/LC_Restart.cpp
+++ b/Engine/Source/Developer/Windows/LiveCoding/Private/External/LC_Restart.cpp
@@ -1,12 +1,16 @@
-// Copyright 2011-2019 Molecular Matters GmbH, all rights reserved.
+// Copyright 2011-2020 Molecular Matters GmbH, all rights reserved.
 
+// BEGIN EPIC MOD
+//#include PCH_INCLUDE
+// END EPIC MOD
 #include "LC_Restart.h"
 #include "LC_Event.h"
 #include "LC_PrimitiveNames.h"
 #include "LC_Process.h"
+// BEGIN EPIC MOD
 #include "LC_Memory.h"
 #include "HAL/PlatformMisc.h"
-
+// END EPIC MOD
 
 namespace
 {
@@ -18,7 +22,7 @@ namespace
 
 void restart::Startup(void)
 {
-	const unsigned int processId = process::GetId();
+	const Process::Id processId = Process::Current::GetId();
 
 	g_requestRestart = new Event(primitiveNames::RequestRestart(processId).c_str(), Event::Type::AUTO_RESET);
 	g_restartPrepared = new Event(primitiveNames::PreparedRestart(processId).c_str(), Event::Type::AUTO_RESET);
diff --git a/Engine/Source/Developer/Windows/LiveCoding/Private/External/LC_Restart.h b/Engine/Source/Developer/Windows/LiveCoding/Private/External/LC_Restart.h
--- a/Engine/Source/Developer/Windows/LiveCoding/Private/External/LC_Restart.h
+++ b/Engine/Source/Developer/Windows/LiveCoding/Private/External/LC_Restart.h
@@ -1,4 +1,4 @@
-// Copyright 2011-2019 Molecular Matters GmbH, all rights reserved.
+// Copyright 2011-2020 Molecular Matters GmbH, all rights reserved.
 
 #pragma once
 
diff --git a/Engine/Source/Developer/Windows/LiveCoding/Private/External/LC_Semaphore.cpp b/Engine/Source/Developer/Windows/LiveCoding/Private/External/LC_Semaphore.cpp
--- a/Engine/Source/Developer/Windows/LiveCoding/Private/External/LC_Semaphore.cpp
+++ b/Engine/Source/Developer/Windows/LiveCoding/Private/External/LC_Semaphore.cpp
@@ -1,9 +1,13 @@
-// Copyright 2011-2019 Molecular Matters GmbH, all rights reserved.
+// Copyright 2011-2020 Molecular Matters GmbH, all rights reserved.
 
+// BEGIN EPIC MOD
+//#include PCH_INCLUDE
+// END EPIC MOD
 #include "LC_Semaphore.h"
+// BEGIN EPIC MOD
 #include "LC_Logging.h"
 #include "Windows/WindowsHWrapper.h"
-
+// END EPIC MOD
 
 Semaphore::Semaphore(unsigned int initialValue, unsigned int maximumValue) :
 	m_sema(::CreateSemaphore(nullptr, static_cast<LONG>(initialValue), static_cast<LONG>(maximumValue), nullptr))
diff --git a/Engine/Source/Developer/Windows/LiveCoding/Private/External/LC_StringUtil.cpp b/Engine/Source/Developer/Windows/LiveCoding/Private/External/LC_StringUtil.cpp
--- a/Engine/Source/Developer/Windows/LiveCoding/Private/External/LC_StringUtil.cpp
+++ b/Engine/Source/Developer/Windows/LiveCoding/Private/External/LC_StringUtil.cpp
@@ -1,8 +1,12 @@
-// Copyright 2011-2019 Molecular Matters GmbH, all rights reserved.
+// Copyright 2011-2020 Molecular Matters GmbH, all rights reserved.
 
+// BEGIN EPIC MOD
+//#include PCH_INCLUDE
+// END EPIC MOD
 #include "LC_StringUtil.h"
+// BEGIN EPIC MOD
 #include "Windows/WindowsHWrapper.h"
-
+// END EPIC MOD
 
 namespace detail
 {
@@ -127,6 +131,23 @@ namespace string
 	}
 
 
+	std::string ReplaceAll(const std::string& str, const std::string& from, const std::string& to)
+	{
+		std::string result(str);
+
+		for (;;)
+		{
+			const size_t pos = result.find(from);
+			if (pos == std::string::npos)
+			{
+				return result;
+			}
+
+			result.replace(pos, from.length(), to);
+		}
+	}
+
+
 	std::wstring ReplaceAll(const std::wstring& str, const std::wstring& from, const std::wstring& to)
 	{
 		std::wstring result(str);
diff --git a/Engine/Source/Developer/Windows/LiveCoding/Private/External/LC_StringUtil.h b/Engine/Source/Developer/Windows/LiveCoding/Private/External/LC_StringUtil.h
--- a/Engine/Source/Developer/Windows/LiveCoding/Private/External/LC_StringUtil.h
+++ b/Engine/Source/Developer/Windows/LiveCoding/Private/External/LC_StringUtil.h
@@ -1,9 +1,11 @@
-// Copyright 2011-2019 Molecular Matters GmbH, all rights reserved.
+// Copyright 2011-2020 Molecular Matters GmbH, all rights reserved.
 
 #pragma once
 
+// BEGIN EPIC MOD
 #include "CoreTypes.h"
 #include <string>
+// END EPIC MOD
 
 namespace string
 {
@@ -13,6 +15,7 @@ namespace string
 
 	std::string Replace(const std::string& str, const std::string& from, const std::string& to);
 	std::wstring Replace(const std::wstring& str, const std::wstring& from, const std::wstring& to);
+	std::string ReplaceAll(const std::string& str, const std::string& from, const std::string& to);
 	std::wstring ReplaceAll(const std::wstring& str, const std::wstring& from, const std::wstring& to);
 
 	std::string EraseAll(const std::string& str, const std::string& subString);
@@ -34,6 +37,27 @@ namespace string
 	bool StartsWith(const char* str, const char* subString);
 	bool StartsWith(const wchar_t* str, const wchar_t* subString);
 
+	inline char ToLower(char c)
+	{
+		return static_cast<char>(::tolower(c));
+	}
+
+	inline wchar_t ToLower(wchar_t c)
+	{
+		return static_cast<wchar_t>(::towlower(c));
+	}
+
+	inline char ToUpper(char c)
+	{
+		return static_cast<char>(::toupper(c));
+	}
+
+	inline wchar_t ToUpper(wchar_t c)
+	{
+		return static_cast<wchar_t>(::towupper(c));
+	}
+
+
 	std::string ToUpper(const char* str);
 	std::string ToUpper(const std::string& str);
 	std::wstring ToUpper(const wchar_t* str);
@@ -45,9 +69,34 @@ namespace string
 	// Turns invalid characters (\ / : * ? " < > | : ; , .) in file names, names for OS objects, etc. into underscores
 	std::wstring MakeSafeName(const std::wstring& name);
 
+	// Returns the length of the given string without null terminator
+	inline size_t GetLength(const char* str)
+	{
+		return strlen(str);
+	}
+
+	// Returns the length of the given string without null terminator
+	inline size_t GetLength(const wchar_t* str)
+	{
+		return wcslen(str);
+	}
+
+	// TODO: temporary fix for Orbis
+#if _WIN32
 	template <typename T>
 	inline T StringToInt(const wchar_t* str)
 	{
 		return static_cast<T>(::_wtoi(str));
 	}
+
+	template <typename T>
+	inline std::wstring IntToString(T value)
+	{
+		// ensure that the largest 64-bit integers & sign & a null-terminator fit into the buffer
+		wchar_t result[22u] = {};
+		_itow_s(static_cast<int>(value), result, 10);
+
+		return std::wstring(result);
+	}
+#endif
 }
diff --git a/Engine/Source/Developer/Windows/LiveCoding/Private/External/LC_SyncPoint.cpp b/Engine/Source/Developer/Windows/LiveCoding/Private/External/LC_SyncPoint.cpp
--- a/Engine/Source/Developer/Windows/LiveCoding/Private/External/LC_SyncPoint.cpp
+++ b/Engine/Source/Developer/Windows/LiveCoding/Private/External/LC_SyncPoint.cpp
@@ -1,5 +1,8 @@
-// Copyright 2011-2019 Molecular Matters GmbH, all rights reserved.
+// Copyright 2011-2020 Molecular Matters GmbH, all rights reserved.
 
+// BEGIN EPIC MOD
+//#include PCH_INCLUDE
+// END EPIC MOD
 #include "LC_SyncPoint.h"
 #include "LC_CriticalSection.h"
 #include "LC_Semaphore.h"
diff --git a/Engine/Source/Developer/Windows/LiveCoding/Private/External/LC_Thread.cpp b/Engine/Source/Developer/Windows/LiveCoding/Private/External/LC_Thread.cpp
--- a/Engine/Source/Developer/Windows/LiveCoding/Private/External/LC_Thread.cpp
+++ b/Engine/Source/Developer/Windows/LiveCoding/Private/External/LC_Thread.cpp
@@ -1,9 +1,14 @@
-// Copyright 2011-2019 Molecular Matters GmbH, all rights reserved.
+// Copyright 2011-2020 Molecular Matters GmbH, all rights reserved.
 
+// BEGIN EPIC MOD
+//#include PCH_INCLUDE
+// END EPIC MOD
 #include "LC_Thread.h"
+// BEGIN EPIC MOD
 #include "LC_Platform.h"
 #include "LC_Logging.h"
 #include <process.h>
+// END EPIC MOD
 
 // BEGIN EPIC MODS
 #pragma warning(push)
@@ -11,9 +16,9 @@
 #pragma warning(disable:6258) // warning C6258: Using TerminateThread does not allow proper thread clean up.
 // END EPIC MODS
 
-namespace
+void Thread::Current::SetName(const char* name)
 {
-	static const DWORD MS_VC_EXCEPTION = 0x406D1388;
+	const DWORD MS_VC_EXCEPTION = 0x406D1388;
 
 	struct THREADNAME_INFO
 	{
@@ -23,162 +28,153 @@ namespace
 		ULONG_PTR dwFlags;		// Reserved for future use, must be zero.
 	};
 
+	// code for setting a thread's name taken from http://msdn.microsoft.com/en-us/library/xcb2z8hs.aspx
+	// note that the pragma directives in the code sample are wrong, and will not work for 32-bit builds.
+	THREADNAME_INFO info;
+	info.dwType = 0x1000;
+	info.szName = reinterpret_cast<ULONG_PTR>(name);
+	info.dwThreadID = static_cast<ULONG_PTR>(-1);
+	info.dwFlags = 0;
 
-	static void SetThreadName(const char* threadName)
+	__try
+	{
+		RaiseException(MS_VC_EXCEPTION, 0, 4u, reinterpret_cast<ULONG_PTR*>(&info));
+	}
+	__except (EXCEPTION_EXECUTE_HANDLER)
 	{
-		// code for setting a thread's name taken from http://msdn.microsoft.com/en-us/library/xcb2z8hs.aspx
-		// note that the pragma directives in the code sample are wrong, and will not work for 32-bit builds.
-		THREADNAME_INFO info;
-		info.dwType = 0x1000;
-		info.szName = reinterpret_cast<ULONG_PTR>(threadName);
-		info.dwThreadID = static_cast<ULONG_PTR>(-1);
-		info.dwFlags = 0;
-
-		__try
-		{
-			RaiseException(MS_VC_EXCEPTION, 0, 4u, reinterpret_cast<ULONG_PTR*>(&info));
-		}
-		__except (EXCEPTION_EXECUTE_HANDLER)
-		{
-		}
 	}
 }
 
 
-namespace thread
+void Thread::Current::SleepSeconds(unsigned int seconds)
 {
-	unsigned int GetId(void)
-	{
-		return ::GetCurrentThreadId();
-	}
+	::Sleep(seconds * 1000u);
+}
 
 
-	unsigned int GetId(Handle handle)
-	{
-		return ::GetThreadId(handle);
-	}
+void Thread::Current::SleepMilliSeconds(unsigned int milliSeconds)
+{
+	::Sleep(milliSeconds);
+}
 
 
-	Handle Create(unsigned int stackSize, Function function, void* context)
-	{
-		const uintptr_t result = _beginthreadex(nullptr, stackSize, function, context, 0u, nullptr);
-		if (result == 0u)
-		{
-			const DWORD error = ::GetLastError();
-			LC_ERROR_USER("Error 0x%X while trying to create thread", error);
-		}
-
-		return reinterpret_cast<Handle>(result);
-	}
+void Thread::Current::Yield(void)
+{
+	::_mm_pause();
+}
 
 
-	void Join(Handle handle)
-	{
-		::WaitForSingleObject(handle, INFINITE);
-	}
+Thread::Id Thread::Current::GetId(void)
+{
+	return Thread::Id(::GetCurrentThreadId());
+}
 
 
-	void Terminate(Handle handle)
+Thread::Handle Thread::Create(unsigned int stackSize, Function function, void* context)
+{
+	const uintptr_t result = _beginthreadex(nullptr, stackSize, function, context, 0u, nullptr);
+	if (result == 0u)
 	{
-		::TerminateThread(handle, 0u);
+		const DWORD error = ::GetLastError();
+		LC_ERROR_USER("Error 0x%X while trying to create thread", error);
 	}
 
-
-	void Yield(void)
-	{
-		::_mm_pause();
-	}
+	// BEGIN EPIC MOD
+	return Thread::Handle(reinterpret_cast<Windows::HANDLE>(result));
+	// END EPIC MOD
+}
 
 
-	void Sleep(unsigned int milliSeconds)
-	{
-		::Sleep(milliSeconds);
-	}
+void Thread::Join(Handle handle)
+{
+	::WaitForSingleObject(+handle, INFINITE);
+}
 
 
-	void CancelIO(Handle handle)
-	{
-		::CancelSynchronousIo(handle);
-	}
+void Thread::Terminate(Handle handle)
+{
+	::TerminateThread(+handle, 0u);
+}
 
 
-	Handle Open(unsigned int threadId)
-	{
-		return ::OpenThread(THREAD_ALL_ACCESS, Windows::FALSE, threadId);
-	}
+Thread::Handle Thread::Open(Id threadId)
+{
+	// BEGIN EPIC MOD
+	return Thread::Handle(::OpenThread(THREAD_ALL_ACCESS, Windows::FALSE, +threadId));
+	// END EPIC MOD
+}
 
 
-	void Close(Handle& handle)
-	{
-		::CloseHandle(handle);
-		handle = INVALID_HANDLE_VALUE;
-	}
+void Thread::Close(Handle& handle)
+{
+	::CloseHandle(+handle);
+	handle = INVALID_HANDLE;
+}
 
 
-	void Suspend(Handle handle)
-	{
-		::SuspendThread(handle);
-	}
+void Thread::Suspend(Handle handle)
+{
+	::SuspendThread(+handle);
+}
 
 
-	void Resume(Handle handle)
-	{
-		::ResumeThread(handle);
-	}
+void Thread::Resume(Handle handle)
+{
+	::ResumeThread(+handle);
+}
 
 
-	int GetPriority(Handle handle)
-	{
-		return ::GetThreadPriority(handle);
-	}
+void Thread::SetPriority(Handle handle, int priority)
+{
+	::SetThreadPriority(+handle, priority);
+}
 
 
-	void SetPriority(Handle handle, int priority)
-	{
-		::SetThreadPriority(handle, priority);
-	}
+int Thread::GetPriority(Handle handle)
+{
+	return ::GetThreadPriority(+handle);
+}
 
 
-	Context GetContext(Handle handle)
-	{
-		CONTEXT threadContext = {};
-		threadContext.ContextFlags = CONTEXT_ALL;
-		::GetThreadContext(handle, &threadContext);
+void Thread::SetContext(Handle handle, const Context* context)
+{
+	::SetThreadContext(+handle, context->Access());
+}
 
-		return threadContext;
-	}
 
+Thread::Context Thread::GetContext(Handle handle)
+{
+	CONTEXT threadContext = {};
+	threadContext.ContextFlags = CONTEXT_ALL;
+	::GetThreadContext(+handle, &threadContext);
 
-	void SetContext(Handle handle, const Context& context)
-	{
-		::SetThreadContext(handle, &context);
-	}
+	return Thread::Context(threadContext);
+}
 
 
-	const void* ReadInstructionPointer(const Context& context)
-	{
+const void* Thread::ReadInstructionPointer(const Context* context)
+{
 #if LC_64_BIT
-		return reinterpret_cast<const void*>(context.Rip);
+	return reinterpret_cast<const void*>(context->Access()->Rip);
 #else
-		return reinterpret_cast<const void*>(context.Eip);
+	return reinterpret_cast<const void*>(context->Access()->Eip);
 #endif
-	}
+}
 
 
-	void WriteInstructionPointer(Context& context, const void* ip)
-	{
+void Thread::WriteInstructionPointer(Context* context, const void* ip)
+{
 #if LC_64_BIT
-		context.Rip = reinterpret_cast<DWORD64>(ip);
+	context->Access()->Rip = reinterpret_cast<DWORD64>(ip);
 #else
-		context.Eip = reinterpret_cast<DWORD>(ip);
+	context->Access()->Eip = reinterpret_cast<DWORD>(ip);
 #endif
-	}
+}
 
 
-	void SetName(const char* name)
-	{
-		SetThreadName(name);
-	}
+Thread::Id Thread::GetId(Handle handle)
+{
+	return Thread::Id(::GetThreadId(+handle));
 }
 
 // BEGIN EPIC MODS
diff --git a/Engine/Source/Developer/Windows/LiveCoding/Private/External/LC_Thread.h b/Engine/Source/Developer/Windows/LiveCoding/Private/External/LC_Thread.h
--- a/Engine/Source/Developer/Windows/LiveCoding/Private/External/LC_Thread.h
+++ b/Engine/Source/Developer/Windows/LiveCoding/Private/External/LC_Thread.h
@@ -1,91 +1,102 @@
-// Copyright 2011-2019 Molecular Matters GmbH, all rights reserved.
+// Copyright 2011-2020 Molecular Matters GmbH, all rights reserved.
 
 #pragma once
 
+// BEGIN EPIC MOD
 #include "CoreTypes.h"
+// END EPIC MOD
+#include "LC_ThreadTypes.h"
+// BEGIN EPIC MOD
 #include "Windows/WindowsHWrapper.h"
+// END EPIC MOD
 
-namespace thread
+
+namespace Thread
 {
-	typedef CONTEXT Context;
-	typedef HANDLE Handle;
-	typedef unsigned int(__stdcall *Function)(void*);
+	// Current/calling thread.
+	namespace Current
+	{
+		// Sets the name of the calling thread.
+		void SetName(const char* name);
+
+		// Sleeps the calling thread for the given number of seconds.
+		void SleepSeconds(unsigned int seconds);
+
+		// Sleeps the calling thread for the given number of milliseconds.
+		void SleepMilliSeconds(unsigned int milliSeconds);
 
-	// returns the thread ID of the calling thread
-	unsigned int GetId(void);
+		// Yields the calling thread.
+		void Yield(void);
 
-	// returns the thread ID of the given thread
-	unsigned int GetId(Handle handle);
+		// Returns the thread Id of the calling thread.
+		Id GetId(void);
+	}
 
-	// raw thread creation function
+	// Creates a raw thread.
 	Handle Create(unsigned int stackSize, Function function, void* context);
 
-	// creates a thread from a given function and provided arguments
+	// Creates a thread from a given function and provided arguments.
 	template <typename F, typename... Args>
-	inline Handle Create(const char* threadName, unsigned int stackSize, F ptrToFunction, Args&&... args);
+	inline Handle CreateFromFunction(const char* threadName, unsigned int stackSize, F ptrToFunction, Args&&... args);
 
-	// creates a thread from a given member function, instance and provided arguments
+	// Creates a thread from a given member function, instance and provided arguments.
 	template <class C, typename F, typename... Args>
-	inline Handle Create(const char* threadName, unsigned int stackSize, F ptrToMemberFunction, C* instance, Args&&... args);
-
+	inline Handle CreateFromMemberFunction(const char* threadName, unsigned int stackSize, C* instance, F ptrToMemberFunction, Args&&... args);
 
+	// Joins the given thread.
 	void Join(Handle handle);
-	void Terminate(Handle handle);
 
-	void Yield(void);
-	void Sleep(unsigned int milliSeconds);
-
-	void CancelIO(Handle handle);
+	// Terminates the given thread.
+	void Terminate(Handle handle);
 
 
-	// opens a thread
-	Handle Open(unsigned int threadId);
+	// Opens a thread with the given Id for full access.
+	Handle Open(Id threadId);
 
-	// closes a thread
+	// Closes a created or opened thread.
 	void Close(Handle& handle);
 
-	// suspends a thread
+
+	// Suspends a thread.
 	void Suspend(Handle handle);
 
-	// resumes a thread
+	// Resumes a suspended thread.
 	void Resume(Handle handle);
 
-	// returns a thread's priority
+	// Sets a thread's priority.
+	void SetPriority(Handle handle, int priority);
+
+	// Returns a thread's priority.
 	int GetPriority(Handle handle);
 
-	// sets a thread's priority
-	void SetPriority(Handle handle, int priority);
+	// Sets a thread's context. Only call on suspended threads.
+	void SetContext(Handle handle, const Context* context);
 
-	// returns a thread's context.
-	// NOTE: only use on suspended threads!
+	// Returns a thread's context. Only call on suspended threads.
 	Context GetContext(Handle handle);
 
-	// sets a thread's context.
-	// NOTE: only use on suspended threads!
-	void SetContext(Handle handle, const Context& context);
+	// Reads a thread context's instruction pointer.
+	const void* ReadInstructionPointer(const Context* context);
 
-	// reads a context' instruction pointer
-	const void* ReadInstructionPointer(const Context& context);
+	// Writes a thread context's instruction pointer.
+	void WriteInstructionPointer(Context* context, const void* ip);
 
-	// writes a context' instruction pointer
-	void WriteInstructionPointer(Context& context, const void* ip);
 
-
-	// sets the name of the calling thread
-	void SetName(const char* name);
+	// Returns the thread Id of the given thread.
+	Id GetId(Handle handle);
 }
 
 
 template <typename F, typename... Args>
-inline thread::Handle thread::Create(const char* threadName, unsigned int stackSize, F ptrToFunction, Args&&... args)
+inline Thread::Handle Thread::CreateFromFunction(const char* threadName, unsigned int stackSize, F ptrToFunction, Args&&... args)
 {
 	// a helper lambda that calls the given functions with the provided arguments.
 	// because this lambda needs to capture, it cannot be converted to a function pointer implicitly,
 	// and therefore cannot be used as a thread function directly.
 	// however, we can solve this by another indirection.
-	auto captureLambda = [threadName, ptrToFunction, args...]() -> unsigned int
+	auto captureLambda = [threadName, ptrToFunction, args...]() -> Thread::ReturnValue
 	{
-		thread::SetName(threadName);
+		Thread::Current::SetName(threadName);
 
 		return (*ptrToFunction)(args...);
 	};
@@ -103,30 +114,29 @@ inline thread::Handle thread::Create(const char* threadName, unsigned int stackS
 	auto capturelessLambda = [](void* lambdaContext) -> unsigned int
 	{
 		CaptureLambdaType* lambdaOnHeap = static_cast<CaptureLambdaType*>(lambdaContext);
-		const unsigned int result = (*lambdaOnHeap)();
+		const Thread::ReturnValue result = (*lambdaOnHeap)();
 		delete lambdaOnHeap;
 
-		return result;
+		return +result;
 	};
 
-	unsigned int(__stdcall *threadFunction)(void*) = capturelessLambda;
-
-	thread::Handle handle = thread::Create(stackSize, threadFunction, lambdaOnHeap);
+	Thread::Function threadFunction = capturelessLambda;
+	Thread::Handle handle = Thread::Create(stackSize, threadFunction, lambdaOnHeap);
 
 	return handle;
 }
 
 
 template <class C, typename F, typename... Args>
-inline thread::Handle thread::Create(const char* threadName, unsigned int stackSize, F ptrToMemberFunction, C* instance, Args&&... args)
+inline Thread::Handle Thread::CreateFromMemberFunction(const char* threadName, unsigned int stackSize, C* instance, F ptrToMemberFunction, Args&&... args)
 {
 	// a helper lambda that calls the given functions with the provided arguments.
 	// because this lambda needs to capture, it cannot be converted to a function pointer implicitly,
 	// and therefore cannot be used as a thread function directly.
 	// however, we can solve this by another indirection.
-	auto captureLambda = [threadName, ptrToMemberFunction, instance, args...]() -> unsigned int
+	auto captureLambda = [threadName, ptrToMemberFunction, instance, args...]() -> Thread::ReturnValue
 	{
-		thread::SetName(threadName);
+		Thread::Current::SetName(threadName);
 
 		return (instance->*ptrToMemberFunction)(args...);
 	};
@@ -144,15 +154,14 @@ inline thread::Handle thread::Create(const char* threadName, unsigned int stackS
 	auto capturelessLambda = [](void* lambdaContext) -> unsigned int
 	{
 		CaptureLambdaType* lambdaOnHeap = static_cast<CaptureLambdaType*>(lambdaContext);
-		const unsigned int result = (*lambdaOnHeap)();
+		const Thread::ReturnValue result = (*lambdaOnHeap)();
 		delete lambdaOnHeap;
 
-		return result;
+		return +result;
 	};
 
-	unsigned int (__stdcall *threadFunction)(void*) = capturelessLambda;
-
-	thread::Handle handle = thread::Create(stackSize, threadFunction, lambdaOnHeap);
+	Thread::Function threadFunction = capturelessLambda;
+	Thread::Handle handle = Thread::Create(stackSize, threadFunction, lambdaOnHeap);
 
 	return handle;
 }
diff --git a/Engine/Source/Developer/Windows/LiveCoding/Private/External/LC_VirtualMemory.cpp b/Engine/Source/Developer/Windows/LiveCoding/Private/External/LC_VirtualMemory.cpp
--- a/Engine/Source/Developer/Windows/LiveCoding/Private/External/LC_VirtualMemory.cpp
+++ b/Engine/Source/Developer/Windows/LiveCoding/Private/External/LC_VirtualMemory.cpp
@@ -1,5 +1,8 @@
-// Copyright 2011-2019 Molecular Matters GmbH, all rights reserved.
+// Copyright 2011-2020 Molecular Matters GmbH, all rights reserved.
 
+// BEGIN EPIC MOD
+//#include PCH_INCLUDE
+// END EPIC MOD
 #include "LC_VirtualMemory.h"
 
 
@@ -7,7 +10,7 @@ namespace
 {
 	static uint32_t GetAllocationGranularity(void)
 	{
-		SYSTEM_INFO info = {};
+		::SYSTEM_INFO info = {};
 		::GetSystemInfo(&info);
 		return info.dwAllocationGranularity;
 	}
@@ -15,45 +18,45 @@ namespace
 
 	static uint32_t GetPageSize(void)
 	{
-		SYSTEM_INFO info = {};
+		::SYSTEM_INFO info = {};
 		::GetSystemInfo(&info);
 		return info.dwPageSize;
 	}
 }
 
 
-void* virtualMemory::Allocate(size_t size)
+void* VirtualMemory::Allocate(size_t size)
 {
 	return ::VirtualAlloc(nullptr, size, MEM_COMMIT | MEM_RESERVE, PAGE_READWRITE);
 }
 
 
-void virtualMemory::Free(void* ptr)
+void VirtualMemory::Free(void* ptr)
 {
 	::VirtualFree(ptr, 0u, MEM_RELEASE);
 }
 
 
-void* virtualMemory::Allocate(process::Handle handle, size_t size, PageType::Enum pageType)
+void* VirtualMemory::Allocate(Process::Handle handle, size_t size, PageType::Enum pageType)
 {
-	return ::VirtualAllocEx(handle, nullptr, size, MEM_COMMIT | MEM_RESERVE, pageType);
+	return ::VirtualAllocEx(+handle, nullptr, size, MEM_COMMIT | MEM_RESERVE, pageType);
 }
 
 
-void virtualMemory::Free(process::Handle handle, void* ptr)
+void VirtualMemory::Free(Process::Handle handle, void* ptr)
 {
-	::VirtualFreeEx(handle, ptr, 0u, MEM_RELEASE);
+	::VirtualFreeEx(+handle, ptr, 0u, MEM_RELEASE);
 }
 
 
-uint32_t virtualMemory::GetAllocationGranularity(void)
+uint32_t VirtualMemory::GetAllocationGranularity(void)
 {
 	static const uint32_t allocationGranularity = ::GetAllocationGranularity();
 	return allocationGranularity;
 }
 
 
-uint32_t virtualMemory::GetPageSize(void)
+uint32_t VirtualMemory::GetPageSize(void)
 {
 	static const uint32_t pageSize = ::GetPageSize();
 	return pageSize;
diff --git a/Engine/Source/Developer/Windows/LiveCoding/Private/External/LC_VirtualMemory.h b/Engine/Source/Developer/Windows/LiveCoding/Private/External/LC_VirtualMemory.h
--- a/Engine/Source/Developer/Windows/LiveCoding/Private/External/LC_VirtualMemory.h
+++ b/Engine/Source/Developer/Windows/LiveCoding/Private/External/LC_VirtualMemory.h
@@ -1,28 +1,32 @@
-// Copyright 2011-2019 Molecular Matters GmbH, all rights reserved.
+// Copyright 2011-2020 Molecular Matters GmbH, all rights reserved.
 
 #pragma once
 
+// BEGIN EPIC MOD
 #include "CoreTypes.h"
-#include "LC_Process.h"
+// END EPIC MOD
+#include "LC_ProcessTypes.h"
+#include "LC_VirtualMemoryTypes.h"
 
 
-namespace virtualMemory
+namespace VirtualMemory
 {
-	struct PageType
-	{
-		enum Enum
-		{
-			READ_WRITE = PAGE_READWRITE,
-			EXECUTE_READ_WRITE = PAGE_EXECUTE_READWRITE
-		};
-	};
-
+	// Allocates virtual memory.
 	void* Allocate(size_t size);
+
+	// Frees virtual memory.
 	void Free(void* ptr);
 
-	void* Allocate(process::Handle handle, size_t size, PageType::Enum pageType);
-	void Free(process::Handle handle, void* ptr);
+	// Allocates virtual memory in the given process.
+	void* Allocate(Process::Handle handle, size_t size, PageType::Enum pageType);
+
+	// Frees virtual memory in the given process.
+	void Free(Process::Handle handle, void* ptr);
 
+
+	// Returns the allocation granularity of the virtual memory system.
 	uint32_t GetAllocationGranularity(void);
+
+	// Returns the page size used by the virtual memory system.
 	uint32_t GetPageSize(void);
 }
diff --git a/Engine/Source/Developer/Windows/LiveCoding/Private/External/LC_WindowsInternalFunctions.cpp b/Engine/Source/Developer/Windows/LiveCoding/Private/External/LC_WindowsInternalFunctions.cpp
--- a/Engine/Source/Developer/Windows/LiveCoding/Private/External/LC_WindowsInternalFunctions.cpp
+++ b/Engine/Source/Developer/Windows/LiveCoding/Private/External/LC_WindowsInternalFunctions.cpp
@@ -1,11 +1,15 @@
-// Copyright 2011-2019 Molecular Matters GmbH, all rights reserved.
+// Copyright 2011-2020 Molecular Matters GmbH, all rights reserved.
 
+// BEGIN EPIC MOD
+//#include PCH_INCLUDE
+// END EPIC MOD
 #include "LC_WindowsInternalFunctions.h"
 
 
-windowsInternal::Function<decltype(NtSuspendProcess)> windowsInternal::NtSuspendProcess("ntdll.dll", "NtSuspendProcess");
-windowsInternal::Function<decltype(NtResumeProcess)> windowsInternal::NtResumeProcess("ntdll.dll", "NtResumeProcess");
-windowsInternal::Function<decltype(NtWriteVirtualMemory)> windowsInternal::NtWriteVirtualMemory("ntdll.dll", "NtWriteVirtualMemory");
-windowsInternal::Function<decltype(NtQuerySystemInformation)> windowsInternal::NtQuerySystemInformation("ntdll.dll", "NtQuerySystemInformation");
-windowsInternal::Function<decltype(NtQueryInformationProcess)> windowsInternal::NtQueryInformationProcess("ntdll.dll", "NtQueryInformationProcess");
-windowsInternal::Function<decltype(NtContinue)> windowsInternal::NtContinue("ntdll.dll", "NtContinue");
+WindowsInternals::Function<decltype(NtSuspendProcess)> WindowsInternals::NtSuspendProcess("ntdll.dll", "NtSuspendProcess");
+WindowsInternals::Function<decltype(NtResumeProcess)> WindowsInternals::NtResumeProcess("ntdll.dll", "NtResumeProcess");
+WindowsInternals::Function<decltype(NtReadVirtualMemory)> WindowsInternals::NtReadVirtualMemory("ntdll.dll", "NtReadVirtualMemory");
+WindowsInternals::Function<decltype(NtWriteVirtualMemory)> WindowsInternals::NtWriteVirtualMemory("ntdll.dll", "NtWriteVirtualMemory");
+WindowsInternals::Function<decltype(NtQuerySystemInformation)> WindowsInternals::NtQuerySystemInformation("ntdll.dll", "NtQuerySystemInformation");
+WindowsInternals::Function<decltype(NtQueryInformationProcess)> WindowsInternals::NtQueryInformationProcess("ntdll.dll", "NtQueryInformationProcess");
+WindowsInternals::Function<decltype(NtContinue)> WindowsInternals::NtContinue("ntdll.dll", "NtContinue");
diff --git a/Engine/Source/Developer/Windows/LiveCoding/Private/External/LC_WindowsInternalFunctions.h b/Engine/Source/Developer/Windows/LiveCoding/Private/External/LC_WindowsInternalFunctions.h
--- a/Engine/Source/Developer/Windows/LiveCoding/Private/External/LC_WindowsInternalFunctions.h
+++ b/Engine/Source/Developer/Windows/LiveCoding/Private/External/LC_WindowsInternalFunctions.h
@@ -1,19 +1,20 @@
-// Copyright 2011-2019 Molecular Matters GmbH, all rights reserved.
+// Copyright 2011-2020 Molecular Matters GmbH, all rights reserved.
 
 #pragma once
 
 #include "LC_WindowsInternals.h"
+// BEGIN EPIC MODS
 #include "LC_Logging.h"
+// END EPIC MODS
 
-
-namespace windowsInternal
+namespace WindowsInternals
 {
-	// helper class that allows us to call any function in any Windows DLL, as long as it is exported and we know its signature.
-	// base template.
+	// Helper class that allows us to call any function in any Windows DLL, as long as it is exported and we know its signature.
+	// Base template.
 	template <typename T>
 	class Function {};
 
-	// partial specialization for matching any function signature.
+	// Partial specialization for matching any function signature.
 	template <typename R, typename... Args>
 	class Function<R (Args...)>
 	{
@@ -25,12 +26,12 @@ namespace windowsInternal
 		inline R operator()(Args... args) const;
 
 	private:
-		// helper for letting us check the result for arbitrary return types.
-		// base template.
+		// Helper for letting us check the result for arbitrary return types.
+		// Base template.
 		template <typename T>
 		inline void CheckResult(T) const {}
 
-		// explicit specialization for NTSTATUS return values
+		// Explicit specialization for NTSTATUS return values.
 		template <>
 		inline void CheckResult(NTSTATUS result) const
 		{
@@ -48,7 +49,7 @@ namespace windowsInternal
 
 
 template <typename R, typename... Args>
-inline windowsInternal::Function<R (Args...)>::Function(const char* moduleName, const char* functionName)
+inline WindowsInternals::Function<R (Args...)>::Function(const char* moduleName, const char* functionName)
 	: m_moduleName(moduleName)
 	, m_functionName(functionName)
 	, m_function(nullptr)
@@ -69,7 +70,7 @@ inline windowsInternal::Function<R (Args...)>::Function(const char* moduleName,
 
 
 template <typename R, typename... Args>
-inline R windowsInternal::Function<R (Args...)>::operator()(Args... args) const
+inline R WindowsInternals::Function<R (Args...)>::operator()(Args... args) const
 {
 	const R result = m_function(args...);
 	CheckResult(result);
@@ -78,21 +79,23 @@ inline R windowsInternal::Function<R (Args...)>::operator()(Args... args) const
 }
 
 
-// these are undocumented functions found in ntdll.dll.
-// we don't call them directly, but use them for "extracting" their signature.
-extern "C" windowsInternal::NTSTATUS NtSuspendProcess(HANDLE ProcessHandle);
-extern "C" windowsInternal::NTSTATUS NtResumeProcess(HANDLE ProcessHandle);
-extern "C" windowsInternal::NTSTATUS NtWriteVirtualMemory(HANDLE ProcessHandle, PVOID BaseAddress, PVOID Buffer, ULONG NumberOfBytesToWrite, PULONG NumberOfBytesWritten);
-extern "C" windowsInternal::NTSTATUS NtQuerySystemInformation(windowsInternal::NT_SYSTEM_INFORMATION_CLASS SystemInformationClass, PVOID SystemInformation, ULONG SystemInformationLength, PULONG ReturnLength);
-extern "C" windowsInternal::NTSTATUS NtQueryInformationProcess(HANDLE ProcessHandle, windowsInternal::NT_PROCESS_INFORMATION_CLASS ProcessInformationClass, PVOID ProcessInformation, ULONG ProcessInformationLength, PULONG ReturnLength);
-extern "C" windowsInternal::NTSTATUS NtContinue(CONTEXT* ThreadContext, BOOLEAN RaiseAlert);
+// These are undocumented functions found in ntdll.dll.
+// We don't call them directly, but use them for "extracting" their signature using decltype.
+extern "C" WindowsInternals::NTSTATUS NtSuspendProcess(HANDLE ProcessHandle);
+extern "C" WindowsInternals::NTSTATUS NtResumeProcess(HANDLE ProcessHandle);
+extern "C" WindowsInternals::NTSTATUS NtReadVirtualMemory(HANDLE ProcessHandle, PVOID BaseAddress, PVOID Buffer, ULONG NumberOfBytesToRead, PULONG NumberOfBytesRead);
+extern "C" WindowsInternals::NTSTATUS NtWriteVirtualMemory(HANDLE ProcessHandle, PVOID BaseAddress, PVOID Buffer, ULONG NumberOfBytesToWrite, PULONG NumberOfBytesWritten);
+extern "C" WindowsInternals::NTSTATUS NtQuerySystemInformation(WindowsInternals::NT_SYSTEM_INFORMATION_CLASS SystemInformationClass, PVOID SystemInformation, ULONG SystemInformationLength, PULONG ReturnLength);
+extern "C" WindowsInternals::NTSTATUS NtQueryInformationProcess(HANDLE ProcessHandle, WindowsInternals::NT_PROCESS_INFORMATION_CLASS ProcessInformationClass, PVOID ProcessInformation, ULONG ProcessInformationLength, PULONG ReturnLength);
+extern "C" WindowsInternals::NTSTATUS NtContinue(CONTEXT* ThreadContext, BOOLEAN RaiseAlert);
 
 
-// cache important undocumented functions
-namespace windowsInternal
+// Cache important undocumented functions.
+namespace WindowsInternals
 {
 	extern Function<decltype(NtSuspendProcess)> NtSuspendProcess;
 	extern Function<decltype(NtResumeProcess)> NtResumeProcess;
+	extern Function<decltype(NtReadVirtualMemory)> NtReadVirtualMemory;
 	extern Function<decltype(NtWriteVirtualMemory)> NtWriteVirtualMemory;
 	extern Function<decltype(NtQuerySystemInformation)> NtQuerySystemInformation;
 	extern Function<decltype(NtQueryInformationProcess)> NtQueryInformationProcess;
diff --git a/Engine/Source/Developer/Windows/LiveCoding/Private/External/LC_WindowsInternals.h b/Engine/Source/Developer/Windows/LiveCoding/Private/External/LC_WindowsInternals.h
--- a/Engine/Source/Developer/Windows/LiveCoding/Private/External/LC_WindowsInternals.h
+++ b/Engine/Source/Developer/Windows/LiveCoding/Private/External/LC_WindowsInternals.h
@@ -1,18 +1,20 @@
-// Copyright 2011-2019 Molecular Matters GmbH, all rights reserved.
+// Copyright 2011-2020 Molecular Matters GmbH, all rights reserved.
 
 #pragma once
 
+// BEGIN EPIC MOD
 #include "Windows/WindowsHWrapper.h"
+// END EPIC MOD
 
 // https://docs.microsoft.com/en-us/openspecs/windows_protocols/ms-erref/87fba13e-bf06-450e-83b1-9241dc81e781
 // found in <ntstatus.h>
-#define STATUS_INFO_LENGTH_MISMATCH		((windowsInternal::NTSTATUS)0xC0000004L)
+#define STATUS_INFO_LENGTH_MISMATCH		((WindowsInternals::NTSTATUS)0xC0000004L)
 
 // found in <winternl.h>
-#define NT_SUCCESS(Status)				((windowsInternal::NTSTATUS)(Status) >= 0)
+#define NT_SUCCESS(Status)				((WindowsInternals::NTSTATUS)(Status) >= 0)
 
 
-namespace windowsInternal
+namespace WindowsInternals
 {
 	// most of these types are defined in <winternl.h>, but we cannot include that header because it lacks
 	// a few undocumented values/members that we need for our purposes.
