Integration of Live++ 1.6.3

diff --git a/Engine/Source/Developer/Windows/LiveCoding/Private/External/LPP_API.h b/Engine/Source/Developer/Windows/LiveCoding/Private/External/LPP_API.h
--- a/Engine/Source/Developer/Windows/LiveCoding/Private/External/LPP_API.h
+++ b/Engine/Source/Developer/Windows/LiveCoding/Private/External/LPP_API.h
@@ -1,7 +1,8 @@
-// Copyright 2011-2019 Molecular Matters GmbH, all rights reserved.
+// Copyright 2011-2020 Molecular Matters GmbH, all rights reserved.
 
 #pragma once
 
+
 /******************************************************************************/
 /* HOOKS                                                                      */
 /******************************************************************************/
@@ -15,46 +16,47 @@
 #define LPP_IDENTIFIER(_identifier)					LPP_CONCATENATE(_identifier, __LINE__)
 
 // custom section names for hooks
-#define LPP_PREPATCH_SECTION					".lpp_prepatch_hooks"
-#define LPP_POSTPATCH_SECTION					".lpp_postpatch_hooks"
-#define LPP_COMPILE_START_SECTION				".lpp_compile_start_hooks"
-#define LPP_COMPILE_SUCCESS_SECTION				".lpp_compile_success_hooks"
-#define LPP_COMPILE_ERROR_SECTION				".lpp_compile_error_hooks"
-#define LPP_COMPILE_ERROR_MESSAGE_SECTION		".lpp_compile_error_message_hooks"
+#define LPP_PREPATCH_SECTION						".lpp_prepatch_hooks"
+#define LPP_POSTPATCH_SECTION						".lpp_postpatch_hooks"
+#define LPP_COMPILE_START_SECTION					".lpp_compile_start_hooks"
+#define LPP_COMPILE_SUCCESS_SECTION					".lpp_compile_success_hooks"
+#define LPP_COMPILE_ERROR_SECTION					".lpp_compile_error_hooks"
+#define LPP_COMPILE_ERROR_MESSAGE_SECTION			".lpp_compile_error_message_hooks"
 // BEGIN EPIC MOD - Add the ability for pre and post compile notifications
-#define LPP_PRECOMPILE_SECTION					".lpp_precompile_hooks"
-#define LPP_POSTCOMPILE_SECTION					".lpp_postcompile_hooks"
+#define LPP_PRECOMPILE_SECTION						".lpp_precompile_hooks"
+#define LPP_POSTCOMPILE_SECTION						".lpp_postcompile_hooks"
 // END EPIC MOD
 
 // register a pre-patch hook in a custom section
-#define LPP_PREPATCH_HOOK(_function)																							\
-	__pragma(section(LPP_PREPATCH_SECTION, read))																				\
-	__declspec(allocate(LPP_PREPATCH_SECTION)) extern void (*LPP_IDENTIFIER(lpp_prepatch_hook_function))(void) = &_function
+#define LPP_PREPATCH_HOOK(_function)																																													\
+	__pragma(section(LPP_PREPATCH_SECTION, read)) __declspec(allocate(LPP_PREPATCH_SECTION)) extern void (*LPP_IDENTIFIER(lpp_prepatch_hook_function))(void);															\
+	__pragma(section(LPP_PREPATCH_SECTION, read)) __declspec(allocate(LPP_PREPATCH_SECTION)) void (*LPP_IDENTIFIER(lpp_prepatch_hook_function))(void) = &_function
 
 // register a post-patch hook in a custom section
-#define LPP_POSTPATCH_HOOK(_function)																							\
-	__pragma(section(LPP_POSTPATCH_SECTION, read))																				\
-	__declspec(allocate(LPP_POSTPATCH_SECTION)) extern void (*LPP_IDENTIFIER(lpp_postpatch_hook_function))(void) = &_function
+#define LPP_POSTPATCH_HOOK(_function)																																													\
+	__pragma(section(LPP_POSTPATCH_SECTION, read)) __declspec(allocate(LPP_POSTPATCH_SECTION)) extern void (*LPP_IDENTIFIER(lpp_postpatch_hook_function))(void);														\
+	__pragma(section(LPP_POSTPATCH_SECTION, read)) __declspec(allocate(LPP_POSTPATCH_SECTION)) void (*LPP_IDENTIFIER(lpp_postpatch_hook_function))(void) = &_function
 
 // register a compile start hook in a custom section
-#define LPP_COMPILE_START_HOOK(_function)																								\
-	__pragma(section(LPP_COMPILE_START_SECTION, read))																					\
-	__declspec(allocate(LPP_COMPILE_START_SECTION)) extern void (*LPP_IDENTIFIER(lpp_compile_start_hook_function))(void) = &_function
+#define LPP_COMPILE_START_HOOK(_function)																																												\
+	__pragma(section(LPP_COMPILE_START_SECTION, read)) __declspec(allocate(LPP_COMPILE_START_SECTION)) extern void (*LPP_IDENTIFIER(lpp_compile_start_hook_function))(void);											\
+	__pragma(section(LPP_COMPILE_START_SECTION, read)) __declspec(allocate(LPP_COMPILE_START_SECTION)) void (*LPP_IDENTIFIER(lpp_compile_start_hook_function))(void) = &_function
 
 // register a compile success hook in a custom section
-#define LPP_COMPILE_SUCCESS_HOOK(_function)																									\
-	__pragma(section(LPP_COMPILE_SUCCESS_SECTION, read))																					\
-	__declspec(allocate(LPP_COMPILE_SUCCESS_SECTION)) extern void (*LPP_IDENTIFIER(lpp_compile_success_hook_function))(void) = &_function
+#define LPP_COMPILE_SUCCESS_HOOK(_function)																																												\
+	__pragma(section(LPP_COMPILE_SUCCESS_SECTION, read)) __declspec(allocate(LPP_COMPILE_SUCCESS_SECTION)) extern void (*LPP_IDENTIFIER(lpp_compile_success_hook_function))(void);										\
+	__pragma(section(LPP_COMPILE_SUCCESS_SECTION, read)) __declspec(allocate(LPP_COMPILE_SUCCESS_SECTION)) void (*LPP_IDENTIFIER(lpp_compile_success_hook_function))(void) = &_function
 
 // register a compile error hook in a custom section
-#define LPP_COMPILE_ERROR_HOOK(_function)																										\
-	__pragma(section(LPP_COMPILE_ERROR_SECTION, read))																							\
-	__declspec(allocate(LPP_COMPILE_ERROR_SECTION)) extern void (*LPP_IDENTIFIER(lpp_compile_error_hook_function))(void) = &_function
+#define LPP_COMPILE_ERROR_HOOK(_function)																																												\
+	__pragma(section(LPP_COMPILE_ERROR_SECTION, read)) __declspec(allocate(LPP_COMPILE_ERROR_SECTION)) extern void (*LPP_IDENTIFIER(lpp_compile_error_hook_function))(void);											\
+	__pragma(section(LPP_COMPILE_ERROR_SECTION, read)) __declspec(allocate(LPP_COMPILE_ERROR_SECTION)) void (*LPP_IDENTIFIER(lpp_compile_error_hook_function))(void) = &_function
 
 // register a compile error message hook in a custom section
-#define LPP_COMPILE_ERROR_MESSAGE_HOOK(_function)																												\
-	__pragma(section(LPP_COMPILE_ERROR_MESSAGE_SECTION, read))																									\
-	__declspec(allocate(LPP_COMPILE_ERROR_MESSAGE_SECTION)) extern void (*LPP_IDENTIFIER(lpp_compile_error_message_hook_function))(const wchar_t*) = &_function
+#define LPP_COMPILE_ERROR_MESSAGE_HOOK(_function)																																										\
+	__pragma(section(LPP_COMPILE_ERROR_MESSAGE_SECTION, read)) __declspec(allocate(LPP_COMPILE_ERROR_MESSAGE_SECTION)) extern void (*LPP_IDENTIFIER(lpp_compile_error_message_hook_function))(const wchar_t*);			\
+	__pragma(section(LPP_COMPILE_ERROR_MESSAGE_SECTION, read)) __declspec(allocate(LPP_COMPILE_ERROR_MESSAGE_SECTION)) void (*LPP_IDENTIFIER(lpp_compile_error_message_hook_function))(const wchar_t*) = &_function
+
 
 // BEGIN EPIC MOD - Add the ability for pre and post compile notifications
 // register a pre-compile hook in a custom section
@@ -69,31 +71,157 @@
 // END EPIC MOD
 
 /******************************************************************************/
-/* API                                                                        */
+/* HELPERS                                                                    */
 /******************************************************************************/
 
-// version string
-#define LPP_VERSION "1.5.2"
-
 #ifdef __cplusplus
 #	define LPP_NS_BEGIN			namespace lpp {
 #	define LPP_NS_END			}
 #	define LPP_API				inline
+#	define LPP_CAST(_type)		reinterpret_cast<_type>
+#	define LPP_NULL				nullptr
 #else
 #	define LPP_NS_BEGIN
 #	define LPP_NS_END
 #	define LPP_API				static inline
+#	define LPP_CAST(_type)		(_type)
+#	define LPP_NULL				NULL
 #endif
 
 // helper macros to call a function in a DLL with an arbitrary signature without a compiler warning
-#define LPP_CALL1(_module, _functionName, _return, _args1)								((_return (__cdecl*)(_args1))((uintptr_t)GetProcAddress(_module, _functionName)))
-#define LPP_CALL2(_module, _functionName, _return, _args1, _args2)						((_return (__cdecl*)(_args1, _args2))((uintptr_t)GetProcAddress(_module, _functionName)))
-#define LPP_CALL3(_module, _functionName, _return, _args1, _args2, _args3)				((_return (__cdecl*)(_args1, _args2, _args3))((uintptr_t)GetProcAddress(_module, _functionName)))
-#define LPP_CALL4(_module, _functionName, _return, _args1, _args2, _args3, _args4)		((_return (__cdecl*)(_args1, _args2, _args3, _args4))((uintptr_t)GetProcAddress(_module, _functionName)))
+#define LPP_CALL1(_module, _functionName, _returnType, _args1)								(LPP_CAST(_returnType (__cdecl*)(_args1))(LPP_CAST(uintptr_t)(GetProcAddress(_module, _functionName))))
+#define LPP_CALL2(_module, _functionName, _returnType, _args1, _args2)						(LPP_CAST(_returnType (__cdecl*)(_args1, _args2))(LPP_CAST(uintptr_t)(GetProcAddress(_module, _functionName))))
+#define LPP_CALL3(_module, _functionName, _returnType, _args1, _args2, _args3)				(LPP_CAST(_returnType (__cdecl*)(_args1, _args2, _args3))(LPP_CAST(uintptr_t)(GetProcAddress(_module, _functionName))))
+#define LPP_CALL4(_module, _functionName, _returnType, _args1, _args2, _args3, _args4)		(LPP_CAST(_returnType (__cdecl*)(_args1, _args2, _args3, _args4))(LPP_CAST(uintptr_t)(GetProcAddress(_module, _functionName))))
+
+// linker pseudo-variable representing the DOS header of the module we're being compiled into:
+// Raymond Chen: https://blogs.msdn.microsoft.com/oldnewthing/20041025-00/?p=37483
+// BEGIN EPIC MOD
+//EXTERN_C IMAGE_DOS_HEADER __ImageBase;
+// END EPIC MOD
+
+/******************************************************************************/
+/* API                                                                        */
+/******************************************************************************/
+
+// version string
+#define LPP_VERSION "1.6.3"
+
+// macros to temporarily enable/disable optimizations
+#ifdef __clang__
+#	define LPP_ENABLE_OPTIMIZATIONS		_Pragma("clang optimize on")
+#	define LPP_DISABLE_OPTIMIZATIONS	_Pragma("clang optimize off")
+#else
+#	define LPP_ENABLE_OPTIMIZATIONS		__pragma(optimize("", on))
+#	define LPP_DISABLE_OPTIMIZATIONS	__pragma(optimize("", off))
+#endif
 
 
 LPP_NS_BEGIN
 
+// BEGIN EPIC MOD
+#if 0
+// starts up Live++. must be called after loading the Live++ DLL, before registering a process group.
+LPP_API void lppStartup(HMODULE livePP)
+{
+	return LPP_CALL1(livePP, "LppStartup", void, void)();
+}
+
+
+// shuts down Live++. must be called before unloading the Live++ DLL.
+LPP_API void lppShutdown(HMODULE livePP)
+{
+	return LPP_CALL1(livePP, "LppShutdown", void, void)();
+}
+
+
+// retrieves the version number of the DLL.
+LPP_API const char* lppGetVersion(HMODULE livePP)
+{
+	return LPP_CALL1(livePP, "LppGetVersion", const char*, void)();
+}
+
+
+// checks whether the API and DLL versions match.
+// returns 1 on success, 0 on failure.
+LPP_API int lppCheckVersion(HMODULE livePP)
+{
+	return LPP_CALL1(livePP, "LppCheckVersion", int, const char*)(LPP_VERSION);
+}
+
+
+// registers a Live++ process group.
+LPP_API void lppRegisterProcessGroup(HMODULE livePP, const char* groupName)
+{
+	LPP_CALL1(livePP, "LppRegisterProcessGroup", void, const char*)(groupName);
+}
+
+
+// calls the synchronization point.
+LPP_API void lppSyncPoint(HMODULE livePP)
+{
+	LPP_CALL1(livePP, "LppSyncPoint", void, void)();
+}
+
+
+// waits until the operation identified by the given token is finished.
+LPP_API void lppWaitForToken(HMODULE livePP, void* token)
+{
+	LPP_CALL1(livePP, "LppWaitForToken", void, void*)(token);
+}
+
+
+// triggers a recompile.
+LPP_API void lppTriggerRecompile(HMODULE livePP)
+{
+	LPP_CALL1(livePP, "LppTriggerRecompile", void, void)();
+}
+
+
+// logs a message in the Live+ UI.
+LPP_API void lppLogMessage(HMODULE livePP, const wchar_t* message)
+{
+	LPP_CALL1(livePP, "LppLogMessage", void, const wchar_t*)(message);
+}
+
+
+// builds a patch using the given object files.
+LPP_API void lppBuildPatch(HMODULE livePP, const wchar_t* moduleNames[], const wchar_t* objPaths[], const wchar_t* amalgamatedObjPaths[], unsigned int count)
+{
+	LPP_CALL4(livePP, "LppBuildPatch", void, const wchar_t*[], const wchar_t*[], const wchar_t*[], unsigned int)(moduleNames, objPaths, amalgamatedObjPaths, count);
+}
+
+
+// installs Live++'s exception handler.
+LPP_API void lppInstallExceptionHandler(HMODULE livePP)
+{
+	LPP_CALL1(livePP, "LppInstallExceptionHandler", void, void)();
+}
+
+
+// makes Live++ listen to changed .obj files compiled by an external build system.
+LPP_API void lppUseExternalBuildSystem(HMODULE livePP)
+{
+	LPP_CALL1(livePP, "LppUseExternalBuildSystem", void, void)();
+}
+
+
+// triggers a restart.
+LPP_API void lppTriggerRestart(HMODULE livePP)
+{
+	LPP_CALL1(livePP, "LppTriggerRestart", int, void)();
+}
+
+
+// returns 1 if the calling process wants to be restarted.
+LPP_API int lppWantsRestart(HMODULE livePP)
+{
+	return LPP_CALL1(livePP, "LppWantsRestart", int, void)();
+}
+#endif
+// END EPIC MOD
+
+
 enum RestartBehaviour
 {
 	// BEGIN EPIC MODS - Use UE4 codepath for termination to ensure logs are flushed and session analytics are sent
@@ -105,9 +233,264 @@ enum RestartBehaviour
 	LPP_RESTART_BEHAVIOUR_INSTANT_TERMINATION		// TerminateProcess
 };
 
+// BEGIN EPIC MOD
+#if 0
+// restarts the calling process, does not return.
+LPP_API void lppRestart(HMODULE livePP, enum RestartBehaviour behaviour, unsigned int exitCode)
+{
+	LPP_CALL2(livePP, "LppRestart", void, enum RestartBehaviour, unsigned int)(behaviour, exitCode);
+}
+
+
+// asynchronously enables Live++ for just the given .exe or .dll, but not its import modules.
+// returns a token that can be waited upon using lppWaitForToken().
+LPP_API void* lppEnableModuleAsync(HMODULE livePP, const wchar_t* nameOfExeOrDll)
+{
+	return LPP_CALL1(livePP, "LppEnableModule", void*, const wchar_t*)(nameOfExeOrDll);
+}
+
+
+// asynchronously enables Live++ for all given .exes and .dlls, but not its import modules.
+// returns a token that can be waited upon using lppWaitForToken().
+LPP_API void* lppEnableModulesAsync(HMODULE livePP, const wchar_t* namesOfExeOrDll[], unsigned int count)
+{
+	return LPP_CALL2(livePP, "LppEnableModules", void*, const wchar_t*[], unsigned int)(namesOfExeOrDll, count);
+}
+
+
+// asynchronously enables Live++ for the given .exe or .dll and all its import modules.
+// returns a token that can be waited upon using lppWaitForToken().
+LPP_API void* lppEnableAllModulesAsync(HMODULE livePP, const wchar_t* nameOfExeOrDll)
+{
+	return LPP_CALL1(livePP, "LppEnableAllModules", void*, const wchar_t*)(nameOfExeOrDll);
+}
+
+
+// asynchronously enables Live++ for just the .exe or .dll this function is being called from, but not its import modules.
+// returns a token that can be waited upon using lppWaitForToken().
+LPP_API void* lppEnableCallingModuleAsync(HMODULE livePP)
+{
+	wchar_t path[MAX_PATH] = { 0 };
+	GetModuleFileNameW(LPP_CAST(HMODULE)(&__ImageBase), path, MAX_PATH);
+	return lppEnableModuleAsync(livePP, path);
+}
+
+
+// asynchronously enables Live++ for the .exe or .dll this function is being called from and all its import modules.
+// returns a token that can be waited upon using lppWaitForToken().
+LPP_API void* lppEnableAllCallingModulesAsync(HMODULE livePP)
+{
+	wchar_t path[MAX_PATH] = { 0 };
+	GetModuleFileNameW(LPP_CAST(HMODULE)(&__ImageBase), path, MAX_PATH);
+	return lppEnableAllModulesAsync(livePP, path);
+}
+
+
+// synchronously enables Live++ for just the given .exe or .dll, but not its import modules.
+LPP_API void lppEnableModuleSync(HMODULE livePP, const wchar_t* nameOfExeOrDll)
+{
+	void* token = lppEnableModuleAsync(livePP, nameOfExeOrDll);
+	lppWaitForToken(livePP, token);
+}
+
+
+// synchronously enables Live++ for all given .exes and .dlls, but not its import modules.
+LPP_API void lppEnableModulesSync(HMODULE livePP, const wchar_t* namesOfExeOrDll[], unsigned int count)
+{
+	void* token = lppEnableModulesAsync(livePP, namesOfExeOrDll, count);
+	lppWaitForToken(livePP, token);
+}
+
+
+// synchronously enables Live++ for the given .exe or .dll and all its import modules.
+LPP_API void lppEnableAllModulesSync(HMODULE livePP, const wchar_t* nameOfExeOrDll)
+{
+	void* token = lppEnableAllModulesAsync(livePP, nameOfExeOrDll);
+	lppWaitForToken(livePP, token);
+}
+
+
+// synchronously enables Live++ for just the .exe or .dll this function is being called from, but not its import modules.
+LPP_API void lppEnableCallingModuleSync(HMODULE livePP)
+{
+	void* token = lppEnableCallingModuleAsync(livePP);
+	lppWaitForToken(livePP, token);
+}
+
+
+// synchronously enables Live++ for the .exe or .dll this function is being called from and all its import modules.
+LPP_API void lppEnableAllCallingModulesSync(HMODULE livePP)
+{
+	void* token = lppEnableAllCallingModulesAsync(livePP);
+	lppWaitForToken(livePP, token);
+}
+
+
+// asynchronously disables Live++ for just the given .exe or .dll, but not its import modules.
+// returns a token that can be waited upon using lppWaitForToken().
+LPP_API void* lppDisableModuleAsync(HMODULE livePP, const wchar_t* nameOfExeOrDll)
+{
+	return LPP_CALL1(livePP, "LppDisableModule", void*, const wchar_t*)(nameOfExeOrDll);
+}
+
+
+// asynchronously disables Live++ for all given .exes and .dlls, but not its import modules.
+// returns a token that can be waited upon using lppWaitForToken().
+LPP_API void* lppDisableModulesAsync(HMODULE livePP, const wchar_t* namesOfExeOrDll[], unsigned int count)
+{
+	return LPP_CALL2(livePP, "LppDisableModules", void*, const wchar_t*[], unsigned int)(namesOfExeOrDll, count);
+}
+
+
+// asynchronously disables Live++ for the given .exe or .dll and all its import modules.
+// returns a token that can be waited upon using lppWaitForToken().
+LPP_API void* lppDisableAllModulesAsync(HMODULE livePP, const wchar_t* nameOfExeOrDll)
+{
+	return LPP_CALL1(livePP, "LppDisableAllModules", void*, const wchar_t*)(nameOfExeOrDll);
+}
+
+
+// asynchronously disables Live++ for just the .exe or .dll this function is being called from, but not its import modules.
+// returns a token that can be waited upon using lppWaitForToken().
+LPP_API void* lppDisableCallingModuleAsync(HMODULE livePP)
+{
+	wchar_t path[MAX_PATH] = { 0 };
+	GetModuleFileNameW(LPP_CAST(HMODULE)(&__ImageBase), path, MAX_PATH);
+	return lppDisableModuleAsync(livePP, path);
+}
+
+
+// asynchronously disables Live++ for the .exe or .dll this function is being called from and all its import modules.
+// returns a token that can be waited upon using lppWaitForToken().
+LPP_API void* lppDisableAllCallingModulesAsync(HMODULE livePP)
+{
+	wchar_t path[MAX_PATH] = { 0 };
+	GetModuleFileNameW(LPP_CAST(HMODULE)(&__ImageBase), path, MAX_PATH);
+	return lppDisableAllModulesAsync(livePP, path);
+}
+
+
+// synchronously disables Live++ for just the given .exe or .dll, but not its import modules.
+LPP_API void lppDisableModuleSync(HMODULE livePP, const wchar_t* nameOfExeOrDll)
+{
+	void* token = lppDisableModuleAsync(livePP, nameOfExeOrDll);
+	lppWaitForToken(livePP, token);
+}
+
+
+// synchronously disables Live++ for all given .exes and .dlls, but not its import modules.
+LPP_API void lppDisableModulesSync(HMODULE livePP, const wchar_t* namesOfExeOrDll[], unsigned int count)
+{
+	void* token = lppDisableModulesAsync(livePP, namesOfExeOrDll, count);
+	lppWaitForToken(livePP, token);
+}
+
+
+// synchronously disables Live++ for the given .exe or .dll and all its import modules.
+LPP_API void lppDisableAllModulesSync(HMODULE livePP, const wchar_t* nameOfExeOrDll)
+{
+	void* token = lppDisableAllModulesAsync(livePP, nameOfExeOrDll);
+	lppWaitForToken(livePP, token);
+}
+
+
+// synchronously disables Live++ for just the .exe or .dll this function is being called from, but not its import modules.
+LPP_API void lppDisableCallingModuleSync(HMODULE livePP)
+{
+	void* token = lppDisableCallingModuleAsync(livePP);
+	lppWaitForToken(livePP, token);
+}
+
+
+// synchronously disables Live++ for the .exe or .dll this function is being called from and all its import modules.
+LPP_API void lppDisableAllCallingModulesSync(HMODULE livePP)
+{
+	void* token = lppDisableAllCallingModulesAsync(livePP);
+	lppWaitForToken(livePP, token);
+}
+
+
+// sets any boolean API setting. check the documentation for a full list of available settings.
+LPP_API void lppApplySettingBool(HMODULE livePP, const char* settingName, int value)
+{
+	LPP_CALL2(livePP, "LppApplySettingBool", void, const char*, int)(settingName, value);
+}
+
+
+// sets any integer API setting. check the documentation for a full list of available settings.
+LPP_API void lppApplySettingInt(HMODULE livePP, const char* settingName, int value)
+{
+	LPP_CALL2(livePP, "LppApplySettingInt", void, const char*, int)(settingName, value);
+}
+
+
+// sets any string API setting. check the documentation for a full list of available settings.
+LPP_API void lppApplySettingString(HMODULE livePP, const char* settingName, const wchar_t* value)
+{
+	LPP_CALL2(livePP, "LppApplySettingString", void, const char*, const wchar_t*)(settingName, value);
+}
+
+
+// loads the correct 32-bit/64-bit DLL (based on architecture), checks if the versions match, and registers a Live++ process group.
+LPP_API HMODULE lppLoadAndRegister(const wchar_t* pathWithoutTrailingSlash, const char* groupName)
+{
+	wchar_t path[1024] = { 0 };
+
+	// we deliberately do not use memcpy or strcpy here, as that could force more #includes in the user's code
+	unsigned int index = 0u;
+	while (*pathWithoutTrailingSlash != L'\0')
+	{
+		if (index >= 1023u)
+		{
+			// no space left in buffer for this character and a null terminator
+			return LPP_NULL;
+		}
+
+		path[index++] = *pathWithoutTrailingSlash++;
+	}
+
+#ifdef _WIN64
+	const wchar_t* dllName = L"/x64/LPP_x64.dll";
+#else
+	const wchar_t* dllName = L"/x86/LPP_x86.dll";
+#endif
+
+	while (*dllName != L'\0')
+	{
+		if (index >= 1023u)
+		{
+			// no space left in buffer for this character and a null terminator
+			return LPP_NULL;
+		}
+
+		path[index++] = *dllName++;
+	}
+
+	path[index++] = L'\0';
+
+	HMODULE livePP = LoadLibraryW(path);
+	if (livePP)
+	{
+		if (!lppCheckVersion(livePP))
+		{
+			// version mismatch detected
+			FreeLibrary(livePP);
+			return LPP_NULL;
+		}
+
+		lppStartup(livePP);
+		lppRegisterProcessGroup(livePP, groupName);
+	}
+
+	return livePP;
+}
+#endif
+// END EPIC MOD
+
 LPP_NS_END
 
 
 #undef LPP_CALL1
 #undef LPP_CALL2
 #undef LPP_CALL3
+#undef LPP_CALL4
diff --git a/Engine/Source/Developer/Windows/LiveCoding/Private/External/xxhash.c b/Engine/Source/Developer/Windows/LiveCoding/Private/External/xxhash.c
--- a/Engine/Source/Developer/Windows/LiveCoding/Private/External/xxhash.c
+++ b/Engine/Source/Developer/Windows/LiveCoding/Private/External/xxhash.c
@@ -1,904 +1,54 @@
 /*
-*  xxHash - Fast Hash algorithm
-*  Copyright (C) 2012-2016, Yann Collet
-*
-*  BSD 2-Clause License (http://www.opensource.org/licenses/bsd-license.php)
-*
-*  Redistribution and use in source and binary forms, with or without
-*  modification, are permitted provided that the following conditions are
-*  met:
-*
-*  * Redistributions of source code must retain the above copyright
-*  notice, this list of conditions and the following disclaimer.
-*  * Redistributions in binary form must reproduce the above
-*  copyright notice, this list of conditions and the following disclaimer
-*  in the documentation and/or other materials provided with the
-*  distribution.
-*
-*  THIS SOFTWARE IS PROVIDED BY THE COPYRIGHT HOLDERS AND CONTRIBUTORS
-*  "AS IS" AND ANY EXPRESS OR IMPLIED WARRANTIES, INCLUDING, BUT NOT
-*  LIMITED TO, THE IMPLIED WARRANTIES OF MERCHANTABILITY AND FITNESS FOR
-*  A PARTICULAR PURPOSE ARE DISCLAIMED. IN NO EVENT SHALL THE COPYRIGHT
-*  OWNER OR CONTRIBUTORS BE LIABLE FOR ANY DIRECT, INDIRECT, INCIDENTAL,
-*  SPECIAL, EXEMPLARY, OR CONSEQUENTIAL DAMAGES (INCLUDING, BUT NOT
-*  LIMITED TO, PROCUREMENT OF SUBSTITUTE GOODS OR SERVICES; LOSS OF USE,
-*  DATA, OR PROFITS; OR BUSINESS INTERRUPTION) HOWEVER CAUSED AND ON ANY
-*  THEORY OF LIABILITY, WHETHER IN CONTRACT, STRICT LIABILITY, OR TORT
-*  (INCLUDING NEGLIGENCE OR OTHERWISE) ARISING IN ANY WAY OUT OF THE USE
-*  OF THIS SOFTWARE, EVEN IF ADVISED OF THE POSSIBILITY OF SUCH DAMAGE.
-*
-*  You can contact the author at :
-*  - xxHash homepage: http://www.xxhash.com
-*  - xxHash source repository : [REDACTED_PROFILE_URL]/xxHash
-*/
-
-// BEGIN EPIC MODS
-#pragma warning(push)
-#pragma warning(disable:6385) // warning C6385: Reading invalid data from '&state'
-#pragma warning(disable:4668) // warning C4668: '__cplusplus' is not defined as a preprocessor macro, replacing with '0' for '#if/#elif'
-// END EPIC MODS
-
-// MM: tuning
-#pragma warning(disable : 4711)
-#pragma warning(disable : 4820)
-#define XXH_FORCE_MEMORY_ACCESS 2
-#define XXH_FORCE_ALIGN_CHECK 0
-#define XXH_CPU_LITTLE_ENDIAN 1
-
-
-/* *************************************
-*  Tuning parameters
-***************************************/
-/*!XXH_FORCE_MEMORY_ACCESS :
- * By default, access to unaligned memory is controlled by `memcpy()`, which is safe and portable.
- * Unfortunately, on some target/compiler combinations, the generated assembly is sub-optimal.
- * The below switch allow to select different access method for improved performance.
- * Method 0 (default) : use `memcpy()`. Safe and portable.
- * Method 1 : `__packed` statement. It depends on compiler extension (ie, not portable).
- *            This method is safe if your compiler supports it, and *generally* as fast or faster than `memcpy`.
- * Method 2 : direct access. This method doesn't depend on compiler but violate C standard.
- *            It can generate buggy code on targets which do not support unaligned memory accesses.
- *            But in some circumstances, it's the only known way to get the most performance (ie GCC + ARMv6)
- * See http://stackoverflow.com/a/32095106/646947 for details.
- * Prefer these methods in priority order (0 > 1 > 2)
+ * xxHash - Extremely Fast Hash algorithm
+ * Copyright (C) 2012-present, Yann Collet
+ *
+ * BSD 2-Clause License (https://www.opensource.org/licenses/bsd-license.php)
+ *
+ * Redistribution and use in source and binary forms, with or without
+ * modification, are permitted provided that the following conditions are
+ * met:
+ *
+ *    * Redistributions of source code must retain the above copyright
+ *      notice, this list of conditions and the following disclaimer.
+ *    * Redistributions in binary form must reproduce the above
+ *      copyright notice, this list of conditions and the following disclaimer
+ *      in the documentation and/or other materials provided with the
+ *      distribution.
+ *
+ * THIS SOFTWARE IS PROVIDED BY THE COPYRIGHT HOLDERS AND CONTRIBUTORS
+ * "AS IS" AND ANY EXPRESS OR IMPLIED WARRANTIES, INCLUDING, BUT NOT
+ * LIMITED TO, THE IMPLIED WARRANTIES OF MERCHANTABILITY AND FITNESS FOR
+ * A PARTICULAR PURPOSE ARE DISCLAIMED. IN NO EVENT SHALL THE COPYRIGHT
+ * OWNER OR CONTRIBUTORS BE LIABLE FOR ANY DIRECT, INDIRECT, INCIDENTAL,
+ * SPECIAL, EXEMPLARY, OR CONSEQUENTIAL DAMAGES (INCLUDING, BUT NOT
+ * LIMITED TO, PROCUREMENT OF SUBSTITUTE GOODS OR SERVICES; LOSS OF USE,
+ * DATA, OR PROFITS; OR BUSINESS INTERRUPTION) HOWEVER CAUSED AND ON ANY
+ * THEORY OF LIABILITY, WHETHER IN CONTRACT, STRICT LIABILITY, OR TORT
+ * (INCLUDING NEGLIGENCE OR OTHERWISE) ARISING IN ANY WAY OUT OF THE USE
+ * OF THIS SOFTWARE, EVEN IF ADVISED OF THE POSSIBILITY OF SUCH DAMAGE.
+ *
+ * You can contact the author at:
+ *   - xxHash homepage: https://www.xxhash.com
+ *   - xxHash source repository: [REDACTED_PROFILE_URL]/xxHash
  */
-#ifndef XXH_FORCE_MEMORY_ACCESS   /* can be defined externally, on command line for example */
-#  if defined(__GNUC__) && ( defined(__ARM_ARCH_6__) || defined(__ARM_ARCH_6J__) || defined(__ARM_ARCH_6K__) || defined(__ARM_ARCH_6Z__) || defined(__ARM_ARCH_6ZK__) || defined(__ARM_ARCH_6T2__) )
-#    define XXH_FORCE_MEMORY_ACCESS 2
-#  elif defined(__INTEL_COMPILER) || \
-  (defined(__GNUC__) && ( defined(__ARM_ARCH_7__) || defined(__ARM_ARCH_7A__) || defined(__ARM_ARCH_7R__) || defined(__ARM_ARCH_7M__) || defined(__ARM_ARCH_7S__) ))
-#    define XXH_FORCE_MEMORY_ACCESS 1
-#  endif
-#endif
 
-/*!XXH_ACCEPT_NULL_INPUT_POINTER :
- * If the input pointer is a null pointer, xxHash default behavior is to trigger a memory access error, since it is a bad pointer.
- * When this option is enabled, xxHash output for null input pointers will be the same as a null-length input.
- * By default, this option is disabled. To enable it, uncomment below define :
- */
-/* #define XXH_ACCEPT_NULL_INPUT_POINTER 1 */
 
-/*!XXH_FORCE_NATIVE_FORMAT :
- * By default, xxHash library provides endian-independent Hash values, based on little-endian convention.
- * Results are therefore identical for little-endian and big-endian CPU.
- * This comes at a performance cost for big-endian CPU, since some swapping is required to emulate little-endian format.
- * Should endian-independence be of no importance for your application, you may set the #define below to 1,
- * to improve speed for Big-endian CPU.
- * This option has no impact on Little_Endian CPU.
+/*
+ * xxhash.c instantiates functions defined in xxhash.h
  */
-#ifndef XXH_FORCE_NATIVE_FORMAT   /* can be defined externally */
-#  define XXH_FORCE_NATIVE_FORMAT 0
-#endif
 
-/*!XXH_FORCE_ALIGN_CHECK :
- * This is a minor performance trick, only useful with lots of very small keys.
- * It means : check for aligned/unaligned input.
- * The check costs one initial branch per hash;
- * set it to 0 when the input is guaranteed to be aligned,
- * or when alignment doesn't matter for performance.
- */
-#ifndef XXH_FORCE_ALIGN_CHECK /* can be defined externally */
-#  if defined(__i386) || defined(_M_IX86) || defined(__x86_64__) || defined(_M_X64)
-#    define XXH_FORCE_ALIGN_CHECK 0
-#  else
-#    define XXH_FORCE_ALIGN_CHECK 1
-#  endif
-#endif
 
+#define XXH_STATIC_LINKING_ONLY   /* access advanced declarations */
+#define XXH_IMPLEMENTATION   /* access definitions */
 
-/* *************************************
-*  Includes & Memory related functions
-***************************************/
-/*! Modify the local functions below should you wish to use some other memory routines
-*   for malloc(), free() */
-#include <stdlib.h>
-static void* XXH_malloc(size_t s) { return malloc(s); }
-static void  XXH_free  (void* p)  { free(p); }
-/*! and for memcpy() */
-#include <string.h>
-static void* XXH_memcpy(void* dest, const void* src, size_t size) { return memcpy(dest,src,size); }
+ // BEGIN EPIC MODS
+#pragma warning(push)
+#pragma warning(disable:6385) // warning C6385: Reading invalid data from '&state'
+#pragma warning(disable:4668) // warning C4668: '__cplusplus' is not defined as a preprocessor macro, replacing with '0' for '#if/#elif'
+// END EPIC MODS
 
-#define XXH_STATIC_LINKING_ONLY
 #include "xxhash.h"
 
-/* *************************************
-*  Compiler Specific Options
-***************************************/
-#ifdef _MSC_VER    /* Visual Studio */
-#  pragma warning(disable : 4127)      /* disable: C4127: conditional expression is constant */
-#  define FORCE_INLINE static __forceinline
-#else
-#  if defined (__cplusplus) || defined (__STDC_VERSION__) && __STDC_VERSION__ >= 199901L   /* C99 */
-#    ifdef __GNUC__
-#      define FORCE_INLINE static inline __attribute__((always_inline))
-#    else
-#      define FORCE_INLINE static inline
-#    endif
-#  else
-#    define FORCE_INLINE static
-#  endif /* __STDC_VERSION__ */
-#endif
-
-
-/* *************************************
-*  Basic Types
-***************************************/
-#ifndef MEM_MODULE
-# if !defined (__VMS) && (defined (__cplusplus) || (defined (__STDC_VERSION__) && (__STDC_VERSION__ >= 199901L) /* C99 */) )
-#   include <stdint.h>
-    typedef uint8_t  BYTE;
-    typedef uint16_t U16;
-    typedef uint32_t U32;
-# else
-    typedef unsigned char      BYTE;
-    typedef unsigned short     U16;
-    typedef unsigned int       U32;
-# endif
-#endif
-
-#if (defined(XXH_FORCE_MEMORY_ACCESS) && (XXH_FORCE_MEMORY_ACCESS==2))
-
-/* Force direct memory access. Only works on CPU which support unaligned memory access in hardware */
-static U32 XXH_read32(const void* memPtr) { return *(const U32*) memPtr; }
-
-#elif (defined(XXH_FORCE_MEMORY_ACCESS) && (XXH_FORCE_MEMORY_ACCESS==1))
-
-/* __pack instructions are safer, but compiler specific, hence potentially problematic for some compilers */
-/* currently only defined for gcc and icc */
-typedef union { U32 u32; } __attribute__((packed)) unalign;
-static U32 XXH_read32(const void* ptr) { return ((const unalign*)ptr)->u32; }
-
-#else
-
-/* portable and safe solution. Generally efficient.
- * see : http://stackoverflow.com/a/32095106/646947
- */
-static U32 XXH_read32(const void* memPtr)
-{
-    U32 val;
-    memcpy(&val, memPtr, sizeof(val));
-    return val;
-}
-
-#endif   /* XXH_FORCE_DIRECT_MEMORY_ACCESS */
-
-
-/* ****************************************
-*  Compiler-specific Functions and Macros
-******************************************/
-#define XXH_GCC_VERSION (__GNUC__ * 100 + __GNUC_MINOR__)
-
-/* Note : although _rotl exists for minGW (GCC under windows), performance seems poor */
-#if defined(_MSC_VER)
-#  define XXH_rotl32(x,r) _rotl(x,r)
-#  define XXH_rotl64(x,r) _rotl64(x,r)
-#else
-#  define XXH_rotl32(x,r) ((x << r) | (x >> (32 - r)))
-#  define XXH_rotl64(x,r) ((x << r) | (x >> (64 - r)))
-#endif
-
-#if defined(_MSC_VER)     /* Visual Studio */
-#  define XXH_swap32 _byteswap_ulong
-#elif XXH_GCC_VERSION >= 403
-#  define XXH_swap32 __builtin_bswap32
-#else
-static U32 XXH_swap32 (U32 x)
-{
-    return  ((x << 24) & 0xff000000 ) |
-            ((x <<  8) & 0x00ff0000 ) |
-            ((x >>  8) & 0x0000ff00 ) |
-            ((x >> 24) & 0x000000ff );
-}
-#endif
-
-
-/* *************************************
-*  Architecture Macros
-***************************************/
-typedef enum { XXH_bigEndian=0, XXH_littleEndian=1 } XXH_endianess;
-
-/* XXH_CPU_LITTLE_ENDIAN can be defined externally, for example on the compiler command line */
-#ifndef XXH_CPU_LITTLE_ENDIAN
-    static const int g_one = 1;
-#   define XXH_CPU_LITTLE_ENDIAN   (*(const char*)(&g_one))
-#endif
-
-
-/* ***************************
-*  Memory reads
-*****************************/
-typedef enum { XXH_aligned, XXH_unaligned } XXH_alignment;
-
-FORCE_INLINE U32 XXH_readLE32_align(const void* ptr, XXH_endianess endian, XXH_alignment align)
-{
-    if (align==XXH_unaligned)
-        return endian==XXH_littleEndian ? XXH_read32(ptr) : XXH_swap32(XXH_read32(ptr));
-    else
-        return endian==XXH_littleEndian ? *(const U32*)ptr : XXH_swap32(*(const U32*)ptr);
-}
-
-FORCE_INLINE U32 XXH_readLE32(const void* ptr, XXH_endianess endian)
-{
-    return XXH_readLE32_align(ptr, endian, XXH_unaligned);
-}
-
-static U32 XXH_readBE32(const void* ptr)
-{
-    return XXH_CPU_LITTLE_ENDIAN ? XXH_swap32(XXH_read32(ptr)) : XXH_read32(ptr);
-}
-
-
-/* *************************************
-*  Macros
-***************************************/
-#define XXH_STATIC_ASSERT(c)   { enum { XXH_static_assert = 1/(int)(!!(c)) }; }    /* use only *after* variable declarations */
-XXH_PUBLIC_API unsigned XXH_versionNumber (void) { return XXH_VERSION_NUMBER; }
-
-
-/* *******************************************************************
-*  32-bits hash functions
-*********************************************************************/
-static const U32 PRIME32_1 = 2654435761U;
-static const U32 PRIME32_2 = 2246822519U;
-static const U32 PRIME32_3 = 3266489917U;
-static const U32 PRIME32_4 =  668265263U;
-static const U32 PRIME32_5 =  374761393U;
-
-static U32 XXH32_round(U32 seed, U32 input)
-{
-    seed += input * PRIME32_2;
-    seed  = XXH_rotl32(seed, 13);
-    seed *= PRIME32_1;
-    return seed;
-}
-
-FORCE_INLINE U32 XXH32_endian_align(const void* input, size_t len, U32 seed, XXH_endianess endian, XXH_alignment align)
-{
-    const BYTE* p = (const BYTE*)input;
-    const BYTE* bEnd = p + len;
-    U32 h32;
-#define XXH_get32bits(p) XXH_readLE32_align(p, endian, align)
-
-#ifdef XXH_ACCEPT_NULL_INPUT_POINTER
-    if (p==NULL) {
-        len=0;
-        bEnd=p=(const BYTE*)(size_t)16;
-    }
-#endif
-
-    if (len>=16) {
-        const BYTE* const limit = bEnd - 16;
-        U32 v1 = seed + PRIME32_1 + PRIME32_2;
-        U32 v2 = seed + PRIME32_2;
-        U32 v3 = seed + 0;
-        U32 v4 = seed - PRIME32_1;
-
-        do {
-            v1 = XXH32_round(v1, XXH_get32bits(p)); p+=4;
-            v2 = XXH32_round(v2, XXH_get32bits(p)); p+=4;
-            v3 = XXH32_round(v3, XXH_get32bits(p)); p+=4;
-            v4 = XXH32_round(v4, XXH_get32bits(p)); p+=4;
-        } while (p<=limit);
-
-        h32 = XXH_rotl32(v1, 1) + XXH_rotl32(v2, 7) + XXH_rotl32(v3, 12) + XXH_rotl32(v4, 18);
-    } else {
-        h32  = seed + PRIME32_5;
-    }
-
-    h32 += (U32) len;
-
-    while (p+4<=bEnd) {
-        h32 += XXH_get32bits(p) * PRIME32_3;
-        h32  = XXH_rotl32(h32, 17) * PRIME32_4 ;
-        p+=4;
-    }
-
-    while (p<bEnd) {
-        h32 += (*p) * PRIME32_5;
-        h32 = XXH_rotl32(h32, 11) * PRIME32_1 ;
-        p++;
-    }
-
-    h32 ^= h32 >> 15;
-    h32 *= PRIME32_2;
-    h32 ^= h32 >> 13;
-    h32 *= PRIME32_3;
-    h32 ^= h32 >> 16;
-
-    return h32;
-}
-
-
-XXH_PUBLIC_API unsigned int XXH32 (const void* input, size_t len, unsigned int seed)
-{
-#if 0
-    /* Simple version, good for code maintenance, but unfortunately slow for small inputs */
-    XXH32_state_t state;
-    XXH32_reset(&state, seed);
-    XXH32_update(&state, input, len);
-    return XXH32_digest(&state);
-#else
-    XXH_endianess endian_detected = (XXH_endianess)XXH_CPU_LITTLE_ENDIAN;
-
-    if (XXH_FORCE_ALIGN_CHECK) {
-        if ((((size_t)input) & 3) == 0) {   /* Input is 4-bytes aligned, leverage the speed benefit */
-            if ((endian_detected==XXH_littleEndian) || XXH_FORCE_NATIVE_FORMAT)
-                return XXH32_endian_align(input, len, seed, XXH_littleEndian, XXH_aligned);
-            else
-                return XXH32_endian_align(input, len, seed, XXH_bigEndian, XXH_aligned);
-    }   }
-
-    if ((endian_detected==XXH_littleEndian) || XXH_FORCE_NATIVE_FORMAT)
-        return XXH32_endian_align(input, len, seed, XXH_littleEndian, XXH_unaligned);
-    else
-        return XXH32_endian_align(input, len, seed, XXH_bigEndian, XXH_unaligned);
-#endif
-}
-
-
-
-/*======   Hash streaming   ======*/
-
-XXH_PUBLIC_API XXH32_state_t* XXH32_createState(void)
-{
-    return (XXH32_state_t*)XXH_malloc(sizeof(XXH32_state_t));
-}
-XXH_PUBLIC_API XXH_errorcode XXH32_freeState(XXH32_state_t* statePtr)
-{
-    XXH_free(statePtr);
-    return XXH_OK;
-}
-
-XXH_PUBLIC_API void XXH32_copyState(XXH32_state_t* dstState, const XXH32_state_t* srcState)
-{
-    memcpy(dstState, srcState, sizeof(*dstState));
-}
-
-XXH_PUBLIC_API XXH_errorcode XXH32_reset(XXH32_state_t* statePtr, unsigned int seed)
-{
-    XXH32_state_t state;   /* using a local state to memcpy() in order to avoid strict-aliasing warnings */
-    memset(&state, 0, sizeof(state)-4);   /* do not write into reserved, for future removal */ //-V512
-    state.v1 = seed + PRIME32_1 + PRIME32_2;
-    state.v2 = seed + PRIME32_2;
-    state.v3 = seed + 0;
-    state.v4 = seed - PRIME32_1;
-    memcpy(statePtr, &state, sizeof(state));
-    return XXH_OK;
-}
-
-
-FORCE_INLINE XXH_errorcode XXH32_update_endian (XXH32_state_t* state, const void* input, size_t len, XXH_endianess endian)
-{
-    const BYTE* p = (const BYTE*)input;
-    const BYTE* const bEnd = p + len;
-
-#ifdef XXH_ACCEPT_NULL_INPUT_POINTER
-    if (input==NULL) return XXH_ERROR;
-#endif
-
-    state->total_len_32 += (unsigned)len;
-    state->large_len |= (len>=16) | (state->total_len_32>=16);
-
-    if (state->memsize + len < 16)  {   /* fill in tmp buffer */
-        XXH_memcpy((BYTE*)(state->mem32) + state->memsize, input, len);
-        state->memsize += (unsigned)len;
-        return XXH_OK;
-    }
-
-    if (state->memsize) {   /* some data left from previous update */
-        XXH_memcpy((BYTE*)(state->mem32) + state->memsize, input, 16-state->memsize);
-        {   const U32* p32 = state->mem32;
-            state->v1 = XXH32_round(state->v1, XXH_readLE32(p32, endian)); p32++;
-            state->v2 = XXH32_round(state->v2, XXH_readLE32(p32, endian)); p32++;
-            state->v3 = XXH32_round(state->v3, XXH_readLE32(p32, endian)); p32++;
-            state->v4 = XXH32_round(state->v4, XXH_readLE32(p32, endian));
-        }
-        p += 16-state->memsize;
-        state->memsize = 0;
-    }
-
-    if (p <= bEnd-16) {
-        const BYTE* const limit = bEnd - 16;
-        U32 v1 = state->v1;
-        U32 v2 = state->v2;
-        U32 v3 = state->v3;
-        U32 v4 = state->v4;
-
-        do {
-            v1 = XXH32_round(v1, XXH_readLE32(p, endian)); p+=4;
-            v2 = XXH32_round(v2, XXH_readLE32(p, endian)); p+=4;
-            v3 = XXH32_round(v3, XXH_readLE32(p, endian)); p+=4;
-            v4 = XXH32_round(v4, XXH_readLE32(p, endian)); p+=4;
-        } while (p<=limit);
-
-        state->v1 = v1;
-        state->v2 = v2;
-        state->v3 = v3;
-        state->v4 = v4;
-    }
-
-    if (p < bEnd) {
-        XXH_memcpy(state->mem32, p, (size_t)(bEnd-p));
-        state->memsize = (unsigned)(bEnd-p);
-    }
-
-    return XXH_OK;
-}
-
-XXH_PUBLIC_API XXH_errorcode XXH32_update (XXH32_state_t* state_in, const void* input, size_t len)
-{
-    XXH_endianess endian_detected = (XXH_endianess)XXH_CPU_LITTLE_ENDIAN;
-
-    if ((endian_detected==XXH_littleEndian) || XXH_FORCE_NATIVE_FORMAT)
-        return XXH32_update_endian(state_in, input, len, XXH_littleEndian);
-    else
-        return XXH32_update_endian(state_in, input, len, XXH_bigEndian);
-}
-
-
-
-FORCE_INLINE U32 XXH32_digest_endian (const XXH32_state_t* state, XXH_endianess endian)
-{
-    const BYTE * p = (const BYTE*)state->mem32;
-    const BYTE* const bEnd = (const BYTE*)(state->mem32) + state->memsize;
-    U32 h32;
-
-    if (state->large_len) {
-        h32 = XXH_rotl32(state->v1, 1) + XXH_rotl32(state->v2, 7) + XXH_rotl32(state->v3, 12) + XXH_rotl32(state->v4, 18);
-    } else {
-        h32 = state->v3 /* == seed */ + PRIME32_5;
-    }
-
-    h32 += state->total_len_32;
-
-    while (p+4<=bEnd) {
-        h32 += XXH_readLE32(p, endian) * PRIME32_3;
-        h32  = XXH_rotl32(h32, 17) * PRIME32_4;
-        p+=4;
-    }
-
-    while (p<bEnd) {
-        h32 += (*p) * PRIME32_5;
-        h32  = XXH_rotl32(h32, 11) * PRIME32_1;
-        p++;
-    }
-
-    h32 ^= h32 >> 15;
-    h32 *= PRIME32_2;
-    h32 ^= h32 >> 13;
-    h32 *= PRIME32_3;
-    h32 ^= h32 >> 16;
-
-    return h32;
-}
-
-
-XXH_PUBLIC_API unsigned int XXH32_digest (const XXH32_state_t* state_in)
-{
-    XXH_endianess endian_detected = (XXH_endianess)XXH_CPU_LITTLE_ENDIAN;
-
-    if ((endian_detected==XXH_littleEndian) || XXH_FORCE_NATIVE_FORMAT)
-        return XXH32_digest_endian(state_in, XXH_littleEndian);
-    else
-        return XXH32_digest_endian(state_in, XXH_bigEndian);
-}
-
-
-/*======   Canonical representation   ======*/
-
-/*! Default XXH result types are basic unsigned 32 and 64 bits.
-*   The canonical representation follows human-readable write convention, aka big-endian (large digits first).
-*   These functions allow transformation of hash result into and from its canonical format.
-*   This way, hash values can be written into a file or buffer, and remain comparable across different systems and programs.
-*/
-
-XXH_PUBLIC_API void XXH32_canonicalFromHash(XXH32_canonical_t* dst, XXH32_hash_t hash)
-{
-    XXH_STATIC_ASSERT(sizeof(XXH32_canonical_t) == sizeof(XXH32_hash_t));
-    if (XXH_CPU_LITTLE_ENDIAN) hash = XXH_swap32(hash);
-    memcpy(dst, &hash, sizeof(*dst));
-}
-
-XXH_PUBLIC_API XXH32_hash_t XXH32_hashFromCanonical(const XXH32_canonical_t* src)
-{
-    return XXH_readBE32(src);
-}
-
-
-#ifndef XXH_NO_LONG_LONG
-
-/* *******************************************************************
-*  64-bits hash functions
-*********************************************************************/
-
-/*======   Memory access   ======*/
-
-#ifndef MEM_MODULE
-# define MEM_MODULE
-# if !defined (__VMS) && (defined (__cplusplus) || (defined (__STDC_VERSION__) && (__STDC_VERSION__ >= 199901L) /* C99 */) )
-#   include <stdint.h>
-    typedef uint64_t U64;
-# else
-    typedef unsigned long long U64;   /* if your compiler doesn't support unsigned long long, replace by another 64-bit type here. Note that xxhash.h will also need to be updated. */
-# endif
-#endif
-
-
-#if (defined(XXH_FORCE_MEMORY_ACCESS) && (XXH_FORCE_MEMORY_ACCESS==2))
-
-/* Force direct memory access. Only works on CPU which support unaligned memory access in hardware */
-static U64 XXH_read64(const void* memPtr) { return *(const U64*) memPtr; }
-
-#elif (defined(XXH_FORCE_MEMORY_ACCESS) && (XXH_FORCE_MEMORY_ACCESS==1))
-
-/* __pack instructions are safer, but compiler specific, hence potentially problematic for some compilers */
-/* currently only defined for gcc and icc */
-typedef union { U32 u32; U64 u64; } __attribute__((packed)) unalign64;
-static U64 XXH_read64(const void* ptr) { return ((const unalign64*)ptr)->u64; }
-
-#else
-
-/* portable and safe solution. Generally efficient.
- * see : http://stackoverflow.com/a/32095106/646947
- */
-
-static U64 XXH_read64(const void* memPtr)
-{
-    U64 val;
-    memcpy(&val, memPtr, sizeof(val));
-    return val;
-}
-
-#endif   /* XXH_FORCE_DIRECT_MEMORY_ACCESS */
-
-#if defined(_MSC_VER)     /* Visual Studio */
-#  define XXH_swap64 _byteswap_uint64
-#elif XXH_GCC_VERSION >= 403
-#  define XXH_swap64 __builtin_bswap64
-#else
-static U64 XXH_swap64 (U64 x)
-{
-    return  ((x << 56) & 0xff00000000000000ULL) |
-            ((x << 40) & 0x00ff000000000000ULL) |
-            ((x << 24) & 0x0000ff0000000000ULL) |
-            ((x << 8)  & 0x000000ff00000000ULL) |
-            ((x >> 8)  & 0x00000000ff000000ULL) |
-            ((x >> 24) & 0x0000000000ff0000ULL) |
-            ((x >> 40) & 0x000000000000ff00ULL) |
-            ((x >> 56) & 0x00000000000000ffULL);
-}
-#endif
-
-FORCE_INLINE U64 XXH_readLE64_align(const void* ptr, XXH_endianess endian, XXH_alignment align)
-{
-    if (align==XXH_unaligned)
-        return endian==XXH_littleEndian ? XXH_read64(ptr) : XXH_swap64(XXH_read64(ptr));
-    else
-        return endian==XXH_littleEndian ? *(const U64*)ptr : XXH_swap64(*(const U64*)ptr);
-}
-
-FORCE_INLINE U64 XXH_readLE64(const void* ptr, XXH_endianess endian)
-{
-    return XXH_readLE64_align(ptr, endian, XXH_unaligned);
-}
-
-static U64 XXH_readBE64(const void* ptr)
-{
-    return XXH_CPU_LITTLE_ENDIAN ? XXH_swap64(XXH_read64(ptr)) : XXH_read64(ptr);
-}
-
-
-/*======   xxh64   ======*/
-
-static const U64 PRIME64_1 = 11400714785074694791ULL;
-static const U64 PRIME64_2 = 14029467366897019727ULL;
-static const U64 PRIME64_3 =  1609587929392839161ULL;
-static const U64 PRIME64_4 =  9650029242287828579ULL;
-static const U64 PRIME64_5 =  2870177450012600261ULL;
-
-static U64 XXH64_round(U64 acc, U64 input)
-{
-    acc += input * PRIME64_2;
-    acc  = XXH_rotl64(acc, 31);
-    acc *= PRIME64_1;
-    return acc;
-}
-
-static U64 XXH64_mergeRound(U64 acc, U64 val)
-{
-    val  = XXH64_round(0, val);
-    acc ^= val;
-    acc  = acc * PRIME64_1 + PRIME64_4;
-    return acc;
-}
-
-FORCE_INLINE U64 XXH64_endian_align(const void* input, size_t len, U64 seed, XXH_endianess endian, XXH_alignment align)
-{
-    const BYTE* p = (const BYTE*)input;
-    const BYTE* bEnd = p + len;
-    U64 h64;
-#define XXH_get64bits(p) XXH_readLE64_align(p, endian, align)
-
-#ifdef XXH_ACCEPT_NULL_INPUT_POINTER
-    if (p==NULL) {
-        len=0;
-        bEnd=p=(const BYTE*)(size_t)32;
-    }
-#endif
-
-    if (len>=32) {
-        const BYTE* const limit = bEnd - 32;
-        U64 v1 = seed + PRIME64_1 + PRIME64_2;
-        U64 v2 = seed + PRIME64_2;
-        U64 v3 = seed + 0;
-        U64 v4 = seed - PRIME64_1;
-
-        do {
-            v1 = XXH64_round(v1, XXH_get64bits(p)); p+=8;
-            v2 = XXH64_round(v2, XXH_get64bits(p)); p+=8;
-            v3 = XXH64_round(v3, XXH_get64bits(p)); p+=8;
-            v4 = XXH64_round(v4, XXH_get64bits(p)); p+=8;
-        } while (p<=limit);
-
-        h64 = XXH_rotl64(v1, 1) + XXH_rotl64(v2, 7) + XXH_rotl64(v3, 12) + XXH_rotl64(v4, 18);
-        h64 = XXH64_mergeRound(h64, v1);
-        h64 = XXH64_mergeRound(h64, v2);
-        h64 = XXH64_mergeRound(h64, v3);
-        h64 = XXH64_mergeRound(h64, v4);
-
-    } else {
-        h64  = seed + PRIME64_5;
-    }
-
-    h64 += (U64) len;
-
-    while (p+8<=bEnd) {
-        U64 const k1 = XXH64_round(0, XXH_get64bits(p));
-        h64 ^= k1;
-        h64  = XXH_rotl64(h64,27) * PRIME64_1 + PRIME64_4;
-        p+=8;
-    }
-
-    if (p+4<=bEnd) {
-        h64 ^= (U64)(XXH_get32bits(p)) * PRIME64_1;
-        h64 = XXH_rotl64(h64, 23) * PRIME64_2 + PRIME64_3;
-        p+=4;
-    }
-
-    while (p<bEnd) {
-        h64 ^= (*p) * PRIME64_5;
-        h64 = XXH_rotl64(h64, 11) * PRIME64_1;
-        p++;
-    }
-
-    h64 ^= h64 >> 33;
-    h64 *= PRIME64_2;
-    h64 ^= h64 >> 29;
-    h64 *= PRIME64_3;
-    h64 ^= h64 >> 32;
-
-    return h64;
-}
-
-
-XXH_PUBLIC_API unsigned long long XXH64 (const void* input, size_t len, unsigned long long seed)
-{
-#if 0
-    /* Simple version, good for code maintenance, but unfortunately slow for small inputs */
-    XXH64_state_t state;
-    XXH64_reset(&state, seed);
-    XXH64_update(&state, input, len);
-    return XXH64_digest(&state);
-#else
-    XXH_endianess endian_detected = (XXH_endianess)XXH_CPU_LITTLE_ENDIAN;
-
-    if (XXH_FORCE_ALIGN_CHECK) {
-        if ((((size_t)input) & 7)==0) {  /* Input is aligned, let's leverage the speed advantage */
-            if ((endian_detected==XXH_littleEndian) || XXH_FORCE_NATIVE_FORMAT)
-                return XXH64_endian_align(input, len, seed, XXH_littleEndian, XXH_aligned);
-            else
-                return XXH64_endian_align(input, len, seed, XXH_bigEndian, XXH_aligned);
-    }   }
-
-    if ((endian_detected==XXH_littleEndian) || XXH_FORCE_NATIVE_FORMAT)
-        return XXH64_endian_align(input, len, seed, XXH_littleEndian, XXH_unaligned);
-    else
-        return XXH64_endian_align(input, len, seed, XXH_bigEndian, XXH_unaligned);
-#endif
-}
-
-/*======   Hash Streaming   ======*/
-
-XXH_PUBLIC_API XXH64_state_t* XXH64_createState(void)
-{
-    return (XXH64_state_t*)XXH_malloc(sizeof(XXH64_state_t));
-}
-XXH_PUBLIC_API XXH_errorcode XXH64_freeState(XXH64_state_t* statePtr)
-{
-    XXH_free(statePtr);
-    return XXH_OK;
-}
-
-XXH_PUBLIC_API void XXH64_copyState(XXH64_state_t* dstState, const XXH64_state_t* srcState)
-{
-    memcpy(dstState, srcState, sizeof(*dstState));
-}
-
-XXH_PUBLIC_API XXH_errorcode XXH64_reset(XXH64_state_t* statePtr, unsigned long long seed)
-{
-    XXH64_state_t state;   /* using a local state to memcpy() in order to avoid strict-aliasing warnings */
-    memset(&state, 0, sizeof(state)-8);   /* do not write into reserved, for future removal */ //-V512
-    state.v1 = seed + PRIME64_1 + PRIME64_2;
-    state.v2 = seed + PRIME64_2;
-    state.v3 = seed + 0;
-    state.v4 = seed - PRIME64_1;
-    memcpy(statePtr, &state, sizeof(state));
-    return XXH_OK;
-}
-
-FORCE_INLINE XXH_errorcode XXH64_update_endian (XXH64_state_t* state, const void* input, size_t len, XXH_endianess endian)
-{
-    const BYTE* p = (const BYTE*)input;
-    const BYTE* const bEnd = p + len;
-
-#ifdef XXH_ACCEPT_NULL_INPUT_POINTER
-    if (input==NULL) return XXH_ERROR;
-#endif
-
-    state->total_len += len;
-
-    if (state->memsize + len < 32) {  /* fill in tmp buffer */
-        XXH_memcpy(((BYTE*)state->mem64) + state->memsize, input, len);
-        state->memsize += (U32)len;
-        return XXH_OK;
-    }
-
-    if (state->memsize) {   /* tmp buffer is full */
-        XXH_memcpy(((BYTE*)state->mem64) + state->memsize, input, 32-state->memsize);
-        state->v1 = XXH64_round(state->v1, XXH_readLE64(state->mem64+0, endian));
-        state->v2 = XXH64_round(state->v2, XXH_readLE64(state->mem64+1, endian));
-        state->v3 = XXH64_round(state->v3, XXH_readLE64(state->mem64+2, endian));
-        state->v4 = XXH64_round(state->v4, XXH_readLE64(state->mem64+3, endian));
-        p += 32-state->memsize;
-        state->memsize = 0;
-    }
-
-    if (p+32 <= bEnd) {
-        const BYTE* const limit = bEnd - 32;
-        U64 v1 = state->v1;
-        U64 v2 = state->v2;
-        U64 v3 = state->v3;
-        U64 v4 = state->v4;
-
-        do {
-            v1 = XXH64_round(v1, XXH_readLE64(p, endian)); p+=8;
-            v2 = XXH64_round(v2, XXH_readLE64(p, endian)); p+=8;
-            v3 = XXH64_round(v3, XXH_readLE64(p, endian)); p+=8;
-            v4 = XXH64_round(v4, XXH_readLE64(p, endian)); p+=8;
-        } while (p<=limit);
-
-        state->v1 = v1;
-        state->v2 = v2;
-        state->v3 = v3;
-        state->v4 = v4;
-    }
-
-    if (p < bEnd) {
-        XXH_memcpy(state->mem64, p, (size_t)(bEnd-p));
-        state->memsize = (unsigned)(bEnd-p);
-    }
-
-    return XXH_OK;
-}
-
-XXH_PUBLIC_API XXH_errorcode XXH64_update (XXH64_state_t* state_in, const void* input, size_t len)
-{
-    XXH_endianess endian_detected = (XXH_endianess)XXH_CPU_LITTLE_ENDIAN;
-
-    if ((endian_detected==XXH_littleEndian) || XXH_FORCE_NATIVE_FORMAT)
-        return XXH64_update_endian(state_in, input, len, XXH_littleEndian);
-    else
-        return XXH64_update_endian(state_in, input, len, XXH_bigEndian);
-}
-
-FORCE_INLINE U64 XXH64_digest_endian (const XXH64_state_t* state, XXH_endianess endian)
-{
-    const BYTE * p = (const BYTE*)state->mem64;
-    const BYTE* const bEnd = (const BYTE*)state->mem64 + state->memsize;
-    U64 h64;
-
-    if (state->total_len >= 32) {
-        U64 const v1 = state->v1;
-        U64 const v2 = state->v2;
-        U64 const v3 = state->v3;
-        U64 const v4 = state->v4;
-
-        h64 = XXH_rotl64(v1, 1) + XXH_rotl64(v2, 7) + XXH_rotl64(v3, 12) + XXH_rotl64(v4, 18);
-        h64 = XXH64_mergeRound(h64, v1);
-        h64 = XXH64_mergeRound(h64, v2);
-        h64 = XXH64_mergeRound(h64, v3);
-        h64 = XXH64_mergeRound(h64, v4);
-    } else {
-        h64  = state->v3 + PRIME64_5;
-    }
-
-    h64 += (U64) state->total_len;
-
-    while (p+8<=bEnd) {
-        U64 const k1 = XXH64_round(0, XXH_readLE64(p, endian));
-        h64 ^= k1;
-        h64  = XXH_rotl64(h64,27) * PRIME64_1 + PRIME64_4;
-        p+=8;
-    }
-
-    if (p+4<=bEnd) {
-        h64 ^= (U64)(XXH_readLE32(p, endian)) * PRIME64_1;
-        h64  = XXH_rotl64(h64, 23) * PRIME64_2 + PRIME64_3;
-        p+=4;
-    }
-
-    while (p<bEnd) {
-        h64 ^= (*p) * PRIME64_5;
-        h64  = XXH_rotl64(h64, 11) * PRIME64_1;
-        p++;
-    }
-
-    h64 ^= h64 >> 33;
-    h64 *= PRIME64_2;
-    h64 ^= h64 >> 29;
-    h64 *= PRIME64_3;
-    h64 ^= h64 >> 32;
-
-    return h64;
-}
-
-XXH_PUBLIC_API unsigned long long XXH64_digest (const XXH64_state_t* state_in)
-{
-    XXH_endianess endian_detected = (XXH_endianess)XXH_CPU_LITTLE_ENDIAN;
-
-    if ((endian_detected==XXH_littleEndian) || XXH_FORCE_NATIVE_FORMAT)
-        return XXH64_digest_endian(state_in, XXH_littleEndian);
-    else
-        return XXH64_digest_endian(state_in, XXH_bigEndian);
-}
-
-
-/*====== Canonical representation   ======*/
-
-XXH_PUBLIC_API void XXH64_canonicalFromHash(XXH64_canonical_t* dst, XXH64_hash_t hash)
-{
-    XXH_STATIC_ASSERT(sizeof(XXH64_canonical_t) == sizeof(XXH64_hash_t));
-    if (XXH_CPU_LITTLE_ENDIAN) hash = XXH_swap64(hash);
-    memcpy(dst, &hash, sizeof(*dst));
-}
-
-XXH_PUBLIC_API XXH64_hash_t XXH64_hashFromCanonical(const XXH64_canonical_t* src)
-{
-    return XXH_readBE64(src);
-}
-
-#endif  /* XXH_NO_LONG_LONG */
-
 // BEGIN EPIC MODS
 #pragma warning(pop)
 // END EPIC MODS
