Integration of Live++ 1.6.3

diff --git a/Engine/Source/Developer/Windows/LiveCodingServer/Private/External/LC_Amalgamation.cpp b/Engine/Source/Developer/Windows/LiveCodingServer/Private/External/LC_Amalgamation.cpp
--- a/Engine/Source/Developer/Windows/LiveCodingServer/Private/External/LC_Amalgamation.cpp
+++ b/Engine/Source/Developer/Windows/LiveCodingServer/Private/External/LC_Amalgamation.cpp
@@ -1,11 +1,18 @@
-// Copyright 2011-2019 Molecular Matters GmbH, all rights reserved.
+// Copyright 2011-2020 Molecular Matters GmbH, all rights reserved.
 
+// BEGIN EPIC MOD
+//#include PCH_INCLUDE
+// END EPIC MOD
 #include "LC_Amalgamation.h"
 #include "LC_StringUtil.h"
-#include "LC_FileUtil.h"
+#include "LC_Filesystem.h"
 #include "LC_GrowingMemoryBlock.h"
+#include "LC_FileAttributeCache.h"
+// BEGIN EPIC MOD
+#include "LC_Allocators.h"
+#include "LC_AppSettings.h"
 #include "LC_Logging.h"
-
+// END EPIC MOD
 
 namespace
 {
@@ -29,7 +36,7 @@ namespace
 	static std::wstring GenerateDatabasePath(const symbols::ObjPath& objPath)
 	{
 		std::wstring path = string::ToWideString(objPath.c_str());
-		path = file::RemoveExtension(path);
+		path = Filesystem::RemoveExtension(path.c_str()).GetString();
 		path += L".ldb";
 
 		return path;
@@ -39,8 +46,8 @@ namespace
 	// helper function to generate a timestamp for a file
 	static uint64_t GenerateTimestamp(const wchar_t* path)
 	{
-		const file::Attributes& attributes = file::GetAttributes(path);
-		return file::GetLastModificationTime(attributes);
+		const Filesystem::PathAttributes& attributes = Filesystem::GetAttributes(path);
+		return Filesystem::GetLastModificationTime(attributes);
 	}
 
 
@@ -55,9 +62,87 @@ namespace
 	static Database::Dependency GenerateNormalizedDatabaseDependency(const ImmutableString& path)
 	{
 		const std::wstring widePath = string::ToWideString(path);
-		const std::wstring normalizedPath = file::NormalizePath(widePath.c_str());
+		const std::wstring normalizedPath = Filesystem::NormalizePath(widePath.c_str()).GetString();
 		return Database::Dependency { string::ToUtf8String(normalizedPath).c_str(), GenerateTimestamp(normalizedPath.c_str()) };
 	}
+
+
+	static bool IsMainCompilandCpp(const std::wstring& normalizedDependencySrcPath, const std::wstring& objPath)
+	{
+		// it should suffice to only check the source filename (without extension) against the object filename (without extension).
+		// comparisons involving paths are tricky in this case, because certain build systems like FASTBuild automatically
+		// generate unity files, and can do so in different directories, e.g:
+		// OBJ: Z:\Intermediate\x64\Debug\Unity11.obj
+		// SRC: Z:\Unity\Unity11.cpp
+		const std::wstring srcFile = string::ToUpper(Filesystem::RemoveExtension(Filesystem::GetFilename(normalizedDependencySrcPath.c_str()).GetString()).GetString());
+		const std::wstring objFile = string::ToUpper(Filesystem::RemoveExtension(Filesystem::GetFilename(objPath.c_str()).GetString()).GetString());
+		return string::Contains(objFile.c_str(), srcFile.c_str());
+	}
+
+	static bool IsCppOrCFile(const std::wstring& normalizedLowercaseFilename)
+	{
+		const Filesystem::Path filenameExtension = Filesystem::GetExtension(normalizedLowercaseFilename.c_str());
+		const types::vector<std::wstring>& amalgamatedCppFileExtensions = appSettings::GetAmalgamatedCppFileExtensions();
+
+		for (size_t i = 0u; i < amalgamatedCppFileExtensions.size(); ++i)
+		{
+			if (string::Matches(filenameExtension.GetString(), amalgamatedCppFileExtensions[i].c_str()))
+			{
+				return true;
+			}
+		}
+
+		return false;
+	}
+
+	static size_t FindNumberOfCppFiles(const types::vector<std::wstring>& includeFiles)
+	{
+		size_t count = 0u;
+
+		const size_t fileCount = includeFiles.size();
+		for (size_t i = 0u; i < fileCount; ++i)
+		{
+			const std::wstring wideFilename = includeFiles[i];
+			const std::wstring lowercaseFilename = string::ToLower(wideFilename);
+			const bool isCppOrCFile = IsCppOrCFile(lowercaseFilename);
+			if (isCppOrCFile)
+			{
+				++count;
+			}
+		}
+
+		return count;
+	}
+
+	static void AddFileDependency(symbols::CompilandDB* compilandDb, const ImmutableString& changedSrcFile, const ImmutableString& recompiledObjFile, uint64_t srcFileLastModificationTime, uint64_t moduleLastModificationTime)
+	{
+		// try updating dependencies for the given file and create a new dependency in case none exists yet
+		const auto& insertPair = compilandDb->dependencies.emplace(changedSrcFile, nullptr);
+		symbols::Dependency*& dependency = insertPair.first->second;
+
+		if (insertPair.second)
+		{
+			// insertion was successful, create a new dependency
+			dependency = LC_NEW(&g_dependencyAllocator, symbols::Dependency);
+
+			// if the source file is newer than the module it belongs to, let the dependency store the modification time of the module instead of the source.
+			// this ensures that changes to a source file which happened before Live++ loaded the module are automatically picked up when re-compiling.
+			if (srcFileLastModificationTime > moduleLastModificationTime)
+			{
+				dependency->lastModification = moduleLastModificationTime;
+				dependency->hadInitialChange = true;
+			}
+			else
+			{
+				dependency->lastModification = srcFileLastModificationTime;
+				dependency->hadInitialChange = false;
+			}
+			dependency->parentDirectory = nullptr;
+		}
+
+		// update entry
+		dependency->objPaths.push_back(recompiledObjFile);
+	}
 }
 
 
@@ -295,7 +380,7 @@ bool amalgamation::IsPartOfAmalgamation(const wchar_t* normalizedObjPath)
 std::wstring amalgamation::CreateObjPart(const std::wstring& normalizedFilename)
 {
 	std::wstring newObjPart(LPP_AMALGAMATION_PART_WIDE);
-	newObjPart += file::RemoveExtension(file::GetFilename(normalizedFilename));
+	newObjPart += Filesystem::RemoveExtension(Filesystem::GetFilename(normalizedFilename.c_str()).GetString()).GetString();
 	newObjPart += L".obj";
 
 	return newObjPart;
@@ -305,39 +390,284 @@ std::wstring amalgamation::CreateObjPart(const std::wstring& normalizedFilename)
 std::wstring amalgamation::CreateObjPath(const std::wstring& normalizedAmalgamatedObjPath, const std::wstring& objPart)
 {
 	std::wstring newObjPath(normalizedAmalgamatedObjPath);
-	newObjPath = file::RemoveExtension(newObjPath);
+	newObjPath = Filesystem::RemoveExtension(newObjPath.c_str()).GetString();
 	newObjPath += objPart;
 
 	return newObjPath;
 }
 
 
+void amalgamation::SplitAmalgamatedCompilands(
+	unsigned int splitAmalgamatedFilesThreshold,
+	symbols::CompilandDB* compilandDb,
+	symbols::Compiland* compiland,
+	const types::vector<std::wstring>& includeFiles,
+	const wchar_t* diaCompilandPath,
+	const std::wstring& compilandPath,
+	FileAttributeCache& fileCache,
+	bool generateLogs,
+	symbols::TimeStamp moduleLastModificationTime)
+{
+	const std::wstring normalizedCompilandPath = Filesystem::NormalizePath(compilandPath.c_str()).GetString();
+	const symbols::ObjPath objPath(string::ToUtf8String(normalizedCompilandPath));
+
+	const size_t fileCount = includeFiles.size();
+
+	const bool splitAmalgamatedFiles = (splitAmalgamatedFilesThreshold > 1u);
+	const size_t cppFileCount = splitAmalgamatedFiles ? FindNumberOfCppFiles(includeFiles) : 0u;
+
+	// make sure to treat single-part compilands as being non-amalgamated, i.e. we don't support
+	// recursive amalgamation.
+	// in case splitting of amalgamated files is turned off, this automatically takes care of
+	// treating every compiland as single-file compiland.
+	const bool isPartOfAmalgamation = amalgamation::IsPartOfAmalgamation(compilandPath.c_str());
+	if ((!splitAmalgamatedFiles) || (isPartOfAmalgamation) || (cppFileCount < splitAmalgamatedFilesThreshold))
+	{
+		LC_LOG_DEV("Single .cpp file compiland %s", objPath.c_str());
+
+		// only store source files when splitting amalgamated files in order to save memory
+		// in the general case.
+		if (splitAmalgamatedFiles)
+		{
+			// create array of source file indices for this compiland
+			compiland->sourceFiles = new symbols::CompilandSourceFiles;
+			compiland->sourceFiles->files.reserve(fileCount);
+		}
+
+		// this is not an amalgamated compiland
+		for (size_t j = 0u; j < fileCount; ++j)
+		{
+			const std::wstring& normalizedFilename = includeFiles[j];
+
+			if (Filesystem::GetDriveType(normalizedFilename.c_str()) == Filesystem::DriveType::OPTICAL)
+			{
+				LC_LOG_DEV("Ignoring file %S on optical drive", normalizedFilename.c_str());
+			}
+			else
+			{
+				const FileAttributeCache::Data& cacheData = fileCache.UpdateCacheData(normalizedFilename);
+				if (cacheData.exists)
+				{
+					if (generateLogs)
+					{
+						LC_LOG_DEV("Dependency %S", normalizedFilename.c_str());
+					}
+
+					const ImmutableString& sourceFilePath = string::ToUtf8String(normalizedFilename);
+					AddFileDependency(compilandDb, sourceFilePath, objPath, cacheData.lastModificationTime, moduleLastModificationTime);
+
+					if (splitAmalgamatedFiles)
+					{
+						compiland->sourceFiles->files.push_back(sourceFilePath);
+					}
+				}
+				else if (generateLogs)
+				{
+					LC_LOG_DEV("Missing dependency %S", normalizedFilename.c_str());
+				}
+			}
+		}
+
+		compilandDb->compilands.insert(std::make_pair(objPath, compiland));
+		compilandDb->compilandNameToObjOnDisk.emplace(string::ToUtf8String(diaCompilandPath), objPath);
+	}
+	else
+	{
+		// this is an amalgamated compiland
+		LC_LOG_DEV("Amalgamated .cpp file compiland %s", objPath.c_str());
+
+		// always add a main compiland for the .obj file.
+		// some amalgamated files don't store their main .cpp as dependency.
+		{
+			compiland->type = symbols::Compiland::Type::AMALGAMATION;
+			compilandDb->compilands.insert(std::make_pair(objPath, compiland));
+			compilandDb->compilandNameToObjOnDisk.insert(std::make_pair(string::ToUtf8String(diaCompilandPath), objPath));
+		}
+
+		for (size_t j = 0u; j < fileCount; ++j)
+		{
+			const std::wstring& normalizedFilename = includeFiles[j];
+			const std::wstring lowercaseFilename = string::ToLower(normalizedFilename);
+
+			const ImmutableString& sourceFilePath = string::ToUtf8String(normalizedFilename);
+
+			if (IsCppOrCFile(normalizedFilename))
+			{
+				if (IsMainCompilandCpp(normalizedFilename, compilandPath))
+				{
+					LC_LOG_DEV("Main .cpp %S", normalizedFilename.c_str());
+
+					if (Filesystem::GetDriveType(normalizedFilename.c_str()) == Filesystem::DriveType::OPTICAL)
+					{
+						if (generateLogs)
+						{
+							LC_LOG_DEV("Ignoring file %S on optical drive", normalizedFilename.c_str());
+						}
+					}
+					else
+					{
+						const FileAttributeCache::Data& cacheData = fileCache.UpdateCacheData(normalizedFilename);
+						if (cacheData.exists)
+						{
+							if (generateLogs)
+							{
+								LC_LOG_DEV("Dependency %S", normalizedFilename.c_str());
+							}
+
+							AddFileDependency(compilandDb, sourceFilePath, objPath, cacheData.lastModificationTime, moduleLastModificationTime);
+						}
+						else if (generateLogs)
+						{
+							LC_LOG_DEV("Missing dependency %S", normalizedFilename.c_str());
+						}
+					}
+				}
+				else
+				{
+					// this is a .cpp file included by the amalgamated file.
+					// add a separate compiland and .obj for this file, and update dependencies so that changing
+					// this source file will not trigger a build of the amalgamated file.
+					LC_LOG_DEV("Included .cpp %S", normalizedFilename.c_str());
+
+					if (Filesystem::GetDriveType(normalizedFilename.c_str()) == Filesystem::DriveType::OPTICAL)
+					{
+						if (generateLogs)
+						{
+							LC_LOG_DEV("Ignoring file %S on optical drive", normalizedFilename.c_str());
+						}
+					}
+					else
+					{
+						const FileAttributeCache::Data& cacheData = fileCache.UpdateCacheData(normalizedFilename);
+						if (cacheData.exists)
+						{
+							if (generateLogs)
+							{
+								LC_LOG_DEV("Dependency %S", normalizedFilename.c_str());
+							}
+
+							// create new .obj path by appending this file name to the real .obj, e.g.
+							// Amalgamated.obj turns into Amalgamated.lpp_part.ASingleFile.obj.
+							const std::wstring newObjPart = amalgamation::CreateObjPart(Filesystem::NormalizePath(normalizedFilename.c_str()).GetString());
+							const std::wstring newObjPath = amalgamation::CreateObjPath(compilandPath, newObjPart);
+							const std::wstring normalizedNewObjPath = Filesystem::NormalizePath(newObjPath.c_str()).GetString();
+
+							AddFileDependency(compilandDb, sourceFilePath, string::ToUtf8String(normalizedNewObjPath), cacheData.lastModificationTime, moduleLastModificationTime);
+
+							// create a new compiland matching this .obj.
+							// we could use different PDBs for different files when no PCHs are being used, but with
+							// our automatic multi-processor compilation this doesn't really gain anything performance-wise
+							// and just complicates things.
+
+							// adapt command line to accommodate new .obj path
+							const std::wstring& newCommandLine = string::Replace(string::ToWideString(compiland->commandLine), L".obj", newObjPart);
+
+							symbols::Compiland* newCompiland = LC_NEW(&g_compilandAllocator, symbols::Compiland)
+							{
+								string::ToUtf8String(newObjPath),
+								string::ToUtf8String(normalizedFilename),
+								compiland->pdbPath,
+								compiland->compilerPath,
+								string::ToUtf8String(newCommandLine),
+								compiland->workingDirectory,
+								objPath,												// .obj of the amalgamation
+								nullptr,												// file indices
+
+								// note that for the purpose of disambiguating symbols in COFF files,
+								// we treat these files as being the amalgamated file.
+								// symbols originally coming from amalgamated files need to have the same
+								// name as symbols from individual files.
+								compiland->uniqueId,
+
+								symbols::Compiland::Type::PART_OF_AMALGAMATION,			// type of file
+								compiland->isPartOfLibrary,								// isPartOfLibrary
+								false													// wasRecompiled
+							};
+
+							compilandDb->compilands.insert(std::make_pair(string::ToUtf8String(normalizedNewObjPath), newCompiland));
+
+							// try updating the amalgamated compiland for the given file and create a new one in case none exists yet
+							{
+								const auto& insertPair = compilandDb->amalgamatedCompilands.emplace(objPath, nullptr);
+								symbols::AmalgamatedCompiland*& amalgamatedCompiland = insertPair.first->second;
+
+								if (insertPair.second)
+								{
+									// insertion was successful, create a new amalgamated compiland
+									amalgamatedCompiland = LC_NEW(&g_amalgamatedCompilandAllocator, symbols::AmalgamatedCompiland);
+									amalgamatedCompiland->isSplit = false;
+								}
+
+								// update entry
+								amalgamatedCompiland->singleParts.push_back(string::ToUtf8String(normalizedNewObjPath));
+							}
+						}
+						else if (generateLogs)
+						{
+							LC_LOG_DEV("Missing dependency %S", normalizedFilename.c_str());
+						}
+					}
+				}
+			}
+			else
+			{
+				// this is a header file. add it as regular dependency for the main amalgamated .obj file
+				if (Filesystem::GetDriveType(normalizedFilename.c_str()) == Filesystem::DriveType::OPTICAL)
+				{
+					if (generateLogs)
+					{
+						LC_LOG_DEV("Ignoring file %S on optical drive", normalizedFilename.c_str());
+					}
+				}
+				else
+				{
+					const FileAttributeCache::Data& cacheData = fileCache.UpdateCacheData(normalizedFilename);
+					if (cacheData.exists)
+					{
+						if (generateLogs)
+						{
+							LC_LOG_DEV("Dependency %S", normalizedFilename.c_str());
+						}
+
+						AddFileDependency(compilandDb, sourceFilePath, objPath, cacheData.lastModificationTime, moduleLastModificationTime);
+					}
+					else if (generateLogs)
+					{
+						LC_LOG_DEV("Missing dependency %S", normalizedFilename.c_str());
+					}
+				}
+			}
+		}
+	}
+}
+
+
 bool amalgamation::ReadAndCompareDatabase(const symbols::ObjPath& objPath, const std::wstring& compilerPath, const symbols::Compiland* compiland, const std::wstring& additionalCompilerOptions)
 {
 	// check if the .obj is there. if not, there is no need to check the database at all.
 	{
-		const file::Attributes& objAttributes = file::GetAttributes(string::ToWideString(objPath).c_str());
-		if (!file::DoesExist(objAttributes))
+		const Filesystem::PathAttributes& objAttributes = Filesystem::GetAttributes(string::ToWideString(objPath).c_str());
+		if (!Filesystem::DoesExist(objAttributes))
 		{
 			return false;
 		}
 	}
 
 	const std::wstring databasePath = GenerateDatabasePath(objPath);
-	const file::Attributes& fileAttributes = file::GetAttributes(databasePath.c_str());
-	if (!file::DoesExist(fileAttributes))
+	const Filesystem::PathAttributes& fileAttributes = Filesystem::GetAttributes(databasePath.c_str());
+	if (!Filesystem::DoesExist(fileAttributes))
 	{
 		return false;
 	}
 
-	const uint64_t bytesLeftToRead = file::GetSize(fileAttributes);
+	const uint64_t bytesLeftToRead = Filesystem::GetSize(fileAttributes);
 	if (bytesLeftToRead == 0u)
 	{
 		LC_LOG_DEV("Failed to retrieve size of database file %S", databasePath.c_str());
 		return false;
 	}
 
-	file::MemoryFile* memoryFile = file::Open(databasePath.c_str(), file::OpenMode::READ_ONLY);
+	Filesystem::MemoryMappedFile* memoryFile = Filesystem::OpenMemoryMappedFile(databasePath.c_str(), Filesystem::OpenMode::READ);
 	if (!memoryFile)
 	{
 		// database cannot be opened, treat as if a change was detected
@@ -345,12 +675,12 @@ bool amalgamation::ReadAndCompareDatabase(const symbols::ObjPath& objPath, const
 	}
 
 	// start reading the database from disk, comparing against the compiland's database at the same time
-	serializationFromDisk::ReadBuffer readBuffer { file::GetData(memoryFile), bytesLeftToRead };
+	serializationFromDisk::ReadBuffer readBuffer { Filesystem::GetMemoryMappedFileData(memoryFile), bytesLeftToRead };
 	if (!serializationFromDisk::Compare(Database::MAGIC_NUMBER, &readBuffer))
 	{
 		LC_LOG_DEV("Wrong magic number in database file %S", databasePath.c_str());
 
-		file::Close(memoryFile);
+		Filesystem::CloseMemoryMappedFile(memoryFile);
 		return false;
 	}
 
@@ -358,7 +688,7 @@ bool amalgamation::ReadAndCompareDatabase(const symbols::ObjPath& objPath, const
 	{
 		LC_LOG_DEV("Version has changed in database file %S", databasePath.c_str());
 
-		file::Close(memoryFile);
+		Filesystem::CloseMemoryMappedFile(memoryFile);
 		return false;
 	}
 
@@ -366,7 +696,7 @@ bool amalgamation::ReadAndCompareDatabase(const symbols::ObjPath& objPath, const
 	{
 		LC_LOG_DEV("Compiler path has changed in database file %S", databasePath.c_str());
 
-		file::Close(memoryFile);
+		Filesystem::CloseMemoryMappedFile(memoryFile);
 		return false;
 	}
 
@@ -374,7 +704,7 @@ bool amalgamation::ReadAndCompareDatabase(const symbols::ObjPath& objPath, const
 	{
 		LC_LOG_DEV("Compiler timestamp has changed in database file %S", databasePath.c_str());
 
-		file::Close(memoryFile);
+		Filesystem::CloseMemoryMappedFile(memoryFile);
 		return false;
 	}
 
@@ -382,7 +712,7 @@ bool amalgamation::ReadAndCompareDatabase(const symbols::ObjPath& objPath, const
 	{
 		LC_LOG_DEV("Compiland compiler options have changed in database file %S", databasePath.c_str());
 
-		file::Close(memoryFile);
+		Filesystem::CloseMemoryMappedFile(memoryFile);
 		return false;
 	}
 
@@ -390,7 +720,7 @@ bool amalgamation::ReadAndCompareDatabase(const symbols::ObjPath& objPath, const
 	{
 		LC_LOG_DEV("Additional compiler options have changed in database file %S", databasePath.c_str());
 
-		file::Close(memoryFile);
+		Filesystem::CloseMemoryMappedFile(memoryFile);
 		return false;
 	}
 
@@ -398,7 +728,7 @@ bool amalgamation::ReadAndCompareDatabase(const symbols::ObjPath& objPath, const
 	{
 		LC_LOG_DEV("Source file has changed in database file %S", databasePath.c_str());
 
-		file::Close(memoryFile);
+		Filesystem::CloseMemoryMappedFile(memoryFile);
 		return false;
 	}
 
@@ -413,7 +743,7 @@ bool amalgamation::ReadAndCompareDatabase(const symbols::ObjPath& objPath, const
 		{
 			LC_LOG_DEV("Failed to read dependency count in database file %S", databasePath.c_str());
 
-			file::Close(memoryFile);
+			Filesystem::CloseMemoryMappedFile(memoryFile);
 			return false;
 		}
 
@@ -424,24 +754,24 @@ bool amalgamation::ReadAndCompareDatabase(const symbols::ObjPath& objPath, const
 			{
 				LC_LOG_DEV("Failed to read dependency in database file %S", databasePath.c_str());
 
-				file::Close(memoryFile);
+				Filesystem::CloseMemoryMappedFile(memoryFile);
 				return false;
 			}
 
 			// check dependency timestamp
-			const file::Attributes& attributes = file::GetAttributes(string::ToWideString(dependency.filename).c_str());
-			if (file::GetLastModificationTime(attributes) != dependency.timestamp)
+			const Filesystem::PathAttributes& attributes = Filesystem::GetAttributes(string::ToWideString(dependency.filename).c_str());
+			if (Filesystem::GetLastModificationTime(attributes) != dependency.timestamp)
 			{
 				LC_LOG_DEV("Dependency has changed in database file %S", databasePath.c_str());
 
-				file::Close(memoryFile);
+				Filesystem::CloseMemoryMappedFile(memoryFile);
 				return false;
 			}
 		}
 	}
 
 	// no change detected
-	file::Close(memoryFile);
+	Filesystem::CloseMemoryMappedFile(memoryFile);
 	return true;
 }
 
@@ -527,18 +857,18 @@ void amalgamation::WriteDatabase(const symbols::ObjPath& objPath, const std::wst
 	std::wstring tempDatabasePath = databasePath;
 	tempDatabasePath += L".tmp";
 
-	if (!file::CreateFileWithData(tempDatabasePath.c_str(), dbInMemory.GetData(), dbInMemory.GetSize()))
+	if (!Filesystem::CreateFileWithData(tempDatabasePath.c_str(), dbInMemory.GetData(), dbInMemory.GetSize()))
 	{
 		LC_LOG_DEV("Failed to write database for compiland %s", objPath.c_str());
 		return;
 	}
 
-	file::Move(tempDatabasePath.c_str(), databasePath.c_str());
+	Filesystem::Move(tempDatabasePath.c_str(), databasePath.c_str());
 }
 
 
 void amalgamation::DeleteDatabase(const symbols::ObjPath& objPath)
 {
 	const std::wstring& databasePath = GenerateDatabasePath(objPath);
-	file::DeleteIfExists(databasePath.c_str());
+	Filesystem::DeleteIfExists(databasePath.c_str());
 }
diff --git a/Engine/Source/Developer/Windows/LiveCodingServer/Private/External/LC_Amalgamation.h b/Engine/Source/Developer/Windows/LiveCodingServer/Private/External/LC_Amalgamation.h
--- a/Engine/Source/Developer/Windows/LiveCodingServer/Private/External/LC_Amalgamation.h
+++ b/Engine/Source/Developer/Windows/LiveCodingServer/Private/External/LC_Amalgamation.h
@@ -1,10 +1,14 @@
-// Copyright 2011-2019 Molecular Matters GmbH, all rights reserved.
+// Copyright 2011-2020 Molecular Matters GmbH, all rights reserved.
 
 #pragma once
 
+// BEGIN EPIC MOD
 #include "CoreTypes.h"
+// END EPIC MOD
 #include "LC_Symbols.h"
 
+class FileAttributeCache;
+
 
 namespace amalgamation
 {
@@ -20,6 +24,17 @@ namespace amalgamation
 	std::wstring CreateObjPath(const std::wstring& normalizedAmalgamatedObjPath, const std::wstring& objPart);
 
 
+	void SplitAmalgamatedCompilands(
+		unsigned int splitAmalgamatedFilesThreshold,
+		symbols::CompilandDB* compilandDb,
+		symbols::Compiland* compiland,
+		const types::vector<std::wstring>& includeFiles,
+		const wchar_t* diaCompilandPath,
+		const std::wstring& compilandPath,
+		FileAttributeCache& fileCache,
+		bool generateLogs,
+		symbols::TimeStamp moduleLastModificationTime);
+
 	// dependency database handling
 	
 	// reads a database from disk and compares it against the compiland's data.
diff --git a/Engine/Source/Developer/Windows/LiveCodingServer/Private/External/LC_AppSettings.cpp b/Engine/Source/Developer/Windows/LiveCodingServer/Private/External/LC_AppSettings.cpp
--- a/Engine/Source/Developer/Windows/LiveCodingServer/Private/External/LC_AppSettings.cpp
+++ b/Engine/Source/Developer/Windows/LiveCodingServer/Private/External/LC_AppSettings.cpp
@@ -1,14 +1,24 @@
-// Copyright 2011-2019 Molecular Matters GmbH, all rights reserved.
+// Copyright 2011-2020 Molecular Matters GmbH, all rights reserved.
 
+// BEGIN EPIC MOD
+//#include PCH_INCLUDE
+// END EPIC MOD
 #include "LC_AppSettings.h"
 #include "LC_Process.h"
-#include "LC_FileUtil.h"
+#include "LC_Filesystem.h"
 #include "LC_StringUtil.h"
+// BEGIN EPIC MOD
+//#include "LC_App.h"
+// END EPIC MOD
+// BEGIN EPIC MOD
 #include "LC_Logging.h"
+// END EPIC MOD
 
+// BEGIN EPIC MOD
 #include "Windows/AllowWindowsPlatformTypes.h"
 #include <ShlObj.h>
 #include "Windows/HideWindowsPlatformTypes.h"
+// END EPIC MOD
 
 namespace
 {
@@ -36,10 +46,10 @@ namespace
 	{
 		// absolute paths can be used as they are, relative paths are supposed to be relative to the Live++ executable
 		std::wstring path;
-		const bool isRelativePath = file::IsRelativePath(setting->GetValue());
+		const bool isRelativePath = Filesystem::IsRelativePath(setting->GetValue());
 		if (isRelativePath)
 		{
-			path = file::GetDirectory(process::GetImagePath());
+			path = Filesystem::GetDirectory(Process::Current::GetImagePath().GetString()).GetString();
 			path += L"\\";
 			path += setting->GetValue();
 		}
@@ -48,8 +58,8 @@ namespace
 			path = setting->GetValue();
 		}
 
-		const file::Attributes& attributes = file::GetAttributes(path.c_str());
-		if (!file::DoesExist(attributes))
+		const Filesystem::PathAttributes& attributes = Filesystem::GetAttributes(path.c_str());
+		if (!Filesystem::DoesExist(attributes))
 		{
 			if (path.length() != 0u)
 			{
@@ -58,7 +68,7 @@ namespace
 			return path;
 		}
 
-		if (!file::IsDirectory(attributes))
+		if (!Filesystem::IsDirectory(attributes))
 		{
 			// this is not a directory, but a full path
 // EPIC REMOVED			LC_SUCCESS_USER("Using %S at path %S", type, path.c_str());
@@ -66,34 +76,38 @@ namespace
 		}
 
 		// try to find the compiler/linker in the given directory or any of its child directories
-		const types::vector<std::wstring>& files = file::EnumerateFiles(path.c_str());
+		const types::vector<Filesystem::Path>& files = Filesystem::EnumerateFiles(path.c_str());
 
 		// walk over all files, grabbing only cl.exe and link.exe
 		const size_t count = files.size();
 		for (size_t i = 0u; i < count; ++i)
 		{
-			const std::wstring& originalPath = files[i];
-			const std::wstring lowerCaseFilename = string::ToLower(originalPath);
-			if (string::Contains(lowerCaseFilename.c_str(), vs2017Path))
+			Filesystem::Path originalPath;
+			originalPath = path.c_str();
+			originalPath += L"\\";
+			originalPath += files[i];
+
+			const Filesystem::Path lowerCaseFilename = originalPath.ToLower();
+			if (string::Contains(lowerCaseFilename.GetString(), vs2017Path))
 			{
 				// containing the proper sub-path is not enough, we also need to check the filename,
 				// because Visual Studio has files named cl.exe.config and link.exe.config.
-				const std::wstring filenameOnly = file::GetFilename(lowerCaseFilename);
-				if (string::Matches(filenameOnly.c_str(), exeName))
+				const Filesystem::Path filenameOnly = Filesystem::GetFilename(lowerCaseFilename.GetString());
+				if (string::Matches(filenameOnly.GetString(), exeName))
 				{
-					LC_SUCCESS_USER("Found %S at path %S", type, files[i].c_str());
-					return files[i];
+					LC_SUCCESS_USER("Found %S at path %S", type, originalPath.GetString());
+					return std::wstring(originalPath.GetString());
 				}
 			}
-			else if (string::Contains(lowerCaseFilename.c_str(), vs2015AndEarlierPath))
+			else if (string::Contains(lowerCaseFilename.GetString(), vs2015AndEarlierPath))
 			{
 				// containing the proper sub-path is not enough, we also need to check the filename,
 				// because Visual Studio has files named cl.exe.config and link.exe.config.
-				const std::wstring filenameOnly = file::GetFilename(lowerCaseFilename);
-				if (string::Matches(filenameOnly.c_str(), exeName))
+				const Filesystem::Path filenameOnly = Filesystem::GetFilename(lowerCaseFilename.GetString());
+				if (string::Matches(filenameOnly.GetString(), exeName))
 				{
-					LC_SUCCESS_USER("Found %S at path %S", type, files[i].c_str());
-					return files[i];
+					LC_SUCCESS_USER("Found %S at path %S", type, originalPath.GetString());
+					return std::wstring(originalPath.GetString());
 				}
 			}
 		}
@@ -248,6 +262,15 @@ void appSettings::Startup(const wchar_t* group)
 		0x37A	// Ctrl+Alt+F11
 	);
 
+	g_prewarmTimeout = new SettingInt
+	(
+		group,
+		L"prewarm_timeout",
+		L"Prewarm timeout",
+		L"Timeout in milliseconds when prewarming compiler/linker environments",
+		10000
+	);
+
 	g_showUndecoratedNames = new SettingBool
 	(
 		group,
@@ -383,15 +406,6 @@ void appSettings::Startup(const wchar_t* group)
 		L""
 	);	
 
-	g_forceLinkWeakSymbols = new SettingBool
-	(
-		group,
-		L"force_link_weak_symbols",
-		L"Force linking of weak symbols",
-		L"Specifies whether weak symbols should be forced to link",
-		false
-	);
-
 	g_continuousCompilationEnabled = new SettingBool
 	(
 		group,
@@ -454,7 +468,9 @@ void appSettings::Startup(const wchar_t* group)
 		L"amalgamation_split_into_single_parts",
 		L"Split into single parts",
 		L"Specifies whether amalgamated/unity files are automatically split into single files",
+		// BEGIN EPIC MOD
 		true
+		// END EPIC MOD
 	);
 
 	g_amalgamationSplitMinCppCount = new SettingInt
@@ -483,6 +499,24 @@ void appSettings::Startup(const wchar_t* group)
 		L"Specifies whether NatVis visualizers are supported by linking with a special helper library",
 		false
 	);
+
+	g_fastBuildDatabasePath = new SettingString
+	(
+		group,
+		L"fastbuild_database_path",
+		L"FASTBuild database path",
+		L"Path to the FASTBuild database to gather dependencies from",
+		L""
+	);
+
+	g_fastBuildDllName = new SettingString
+	(
+		group,
+		L"fastbuild_dll_name",
+		L"FASTBuild DLL name",
+		L"Name of the DLL used to load FASTBuild databases",
+		L"LPP_FASTBuild_1_01.dll"
+	);
 }
 
 
@@ -502,6 +536,7 @@ void appSettings::Shutdown(void)
 	delete g_playSoundOnSuccess;
 	delete g_playSoundOnError;
 	delete g_compileShortcut;
+	delete g_prewarmTimeout;
 
 	delete g_showUndecoratedNames;
 	delete g_showTimestamps;
@@ -520,7 +555,6 @@ void appSettings::Shutdown(void)
 	delete g_useLinkerOverrideAsFallback;
 	delete g_useLinkerEnvironment;
 	delete g_linkerOptions;
-	delete g_forceLinkWeakSymbols;
 
 	delete g_continuousCompilationEnabled;
 	delete g_continuousCompilationPath;
@@ -536,6 +570,9 @@ void appSettings::Shutdown(void)
 	delete g_amalgamationCppFileExtensions;
 
 	delete g_ue4EnableNatVisSupport;
+
+	delete g_fastBuildDatabasePath;
+	delete g_fastBuildDllName;
 }
 
 
@@ -573,15 +610,15 @@ std::wstring appSettings::GetUserSettingsPath(void)
 	iniPath += L"\\LPP_x86.ini";
 #endif
 
-	return file::NormalizePath(iniPath.c_str());
+	return Filesystem::NormalizePath(iniPath.c_str()).GetString();
 }
 
 
 std::wstring appSettings::GetProjectSettingsPath(void)
 {
 	// project settings are stored next to the Live++ executable
-	const std::wstring& imagePath = process::GetImagePath();
-	std::wstring iniPath(file::GetDirectory(imagePath));
+	const std::wstring& imagePath = Process::Current::GetImagePath().GetString();
+	std::wstring iniPath(Filesystem::GetDirectory(imagePath.c_str()).GetString());
 
 #if LC_64_BIT
 	iniPath += L"\\LPP_x64.ini";
@@ -589,7 +626,7 @@ std::wstring appSettings::GetProjectSettingsPath(void)
 	iniPath += L"\\LPP_x86.ini";
 #endif
 
-	return file::NormalizePath(iniPath.c_str());
+	return Filesystem::NormalizePath(iniPath.c_str()).GetString();
 }
 
 
@@ -653,7 +690,7 @@ void appSettings::UpdateAmalgamatedCppFileExtensions(void)
 
 void appSettings::ApplySettingBool(const char* const settingName, bool value)
 {
-	const unsigned int COUNT = 22u;
+	const unsigned int COUNT = 21u;
 	SettingBool* settings[COUNT] =
 	{
 		g_showFullPathInTitle,
@@ -673,7 +710,6 @@ void appSettings::ApplySettingBool(const char* const settingName, bool value)
 		g_compilerForcePchPdbs,
 		g_useLinkerOverrideAsFallback,
 		g_useLinkerEnvironment,
-		g_forceLinkWeakSymbols,
 		g_continuousCompilationEnabled,
 		g_installCompiledPatchesMultiProcess,
 		g_amalgamationSplitIntoSingleParts,
@@ -692,12 +728,13 @@ void appSettings::ApplySettingInt(const char* const settingName, int value)
 {
 	// try int settings first
 	{
-		const unsigned int COUNT = 4u;
+		const unsigned int COUNT = 5u;
 		SettingInt* settings[COUNT] =
 		{
 			g_initialWindowMode,
 			g_receiveFocusOnRecompile,
 			g_continuousCompilationTimeout,
+			g_prewarmTimeout,
 			g_amalgamationSplitMinCppCount
 		};
 
@@ -733,7 +770,7 @@ void appSettings::ApplySettingString(const char* const settingName, const wchar_
 {
 	// try string settings first
 	{
-		const unsigned int COUNT = 10u;
+		const unsigned int COUNT = 12u;
 		SettingString* settings[COUNT] =
 		{
 			g_playSoundOnSuccess,
@@ -745,7 +782,9 @@ void appSettings::ApplySettingString(const char* const settingName, const wchar_
 			g_continuousCompilationPath,
 			g_virtualDriveLetter,
 			g_virtualDrivePath,
-			g_amalgamationCppFileExtensions
+			g_amalgamationCppFileExtensions,
+			g_fastBuildDatabasePath,
+			g_fastBuildDllName
 		};
 
 		const SettingString* setting = ApplySetting(settings, COUNT, settingName, value);
@@ -801,52 +840,55 @@ void appSettings::ApplySettingString(const char* const settingName, const wchar_
 }
 
 
-extern SettingInt* appSettings::g_initialWindowMode = nullptr;
-extern SettingIntProxy* appSettings::g_initialWindowModeProxy = nullptr;
-extern SettingBool* appSettings::g_showFullPathInTitle = nullptr;
-extern SettingBool* appSettings::g_showPathFirstInTitle = nullptr;
+SettingInt* appSettings::g_initialWindowMode = nullptr;
+SettingIntProxy* appSettings::g_initialWindowModeProxy = nullptr;
+SettingBool* appSettings::g_showFullPathInTitle = nullptr;
+SettingBool* appSettings::g_showPathFirstInTitle = nullptr;
+
+SettingInt* appSettings::g_receiveFocusOnRecompile = nullptr;
+SettingIntProxy* appSettings::g_receiveFocusOnRecompileProxy = nullptr;
+SettingBool* appSettings::g_showNotificationOnRecompile = nullptr;
+SettingBool* appSettings::g_clearLogOnRecompile = nullptr;
+SettingBool* appSettings::g_minimizeOnClose = nullptr;
+SettingBool* appSettings::g_keepTrayIcon = nullptr;
 
-extern SettingInt* appSettings::g_receiveFocusOnRecompile = nullptr;
-extern SettingIntProxy* appSettings::g_receiveFocusOnRecompileProxy = nullptr;
-extern SettingBool* appSettings::g_showNotificationOnRecompile = nullptr;
-extern SettingBool* appSettings::g_clearLogOnRecompile = nullptr;
-extern SettingBool* appSettings::g_minimizeOnClose = nullptr;
-extern SettingBool* appSettings::g_keepTrayIcon = nullptr;
+SettingString* appSettings::g_playSoundOnSuccess = nullptr;
+SettingString* appSettings::g_playSoundOnError = nullptr;
+SettingShortcut* appSettings::g_compileShortcut = nullptr;
+SettingInt* appSettings::g_prewarmTimeout = nullptr;
 
-extern SettingString* appSettings::g_playSoundOnSuccess = nullptr;
-extern SettingString* appSettings::g_playSoundOnError = nullptr;
-extern SettingShortcut* appSettings::g_compileShortcut = nullptr;
+SettingBool* appSettings::g_showUndecoratedNames = nullptr;
+SettingBool* appSettings::g_showTimestamps = nullptr;
+SettingBool* appSettings::g_wordWrapOutput = nullptr;
+SettingBool* appSettings::g_enableDevLog = nullptr;
+SettingBool* appSettings::g_enableTelemetryLog = nullptr;
+SettingBool* appSettings::g_enableDevLogCompilands = nullptr;
 
-extern SettingBool* appSettings::g_showUndecoratedNames = nullptr;
-extern SettingBool* appSettings::g_showTimestamps = nullptr;
-extern SettingBool* appSettings::g_wordWrapOutput = nullptr;
-extern SettingBool* appSettings::g_enableDevLog = nullptr;
-extern SettingBool* appSettings::g_enableTelemetryLog = nullptr;
-extern SettingBool* appSettings::g_enableDevLogCompilands = nullptr;
+SettingString* appSettings::g_compilerPath = nullptr;
+SettingBool* appSettings::g_useCompilerOverrideAsFallback = nullptr;
+SettingBool* appSettings::g_useCompilerEnvironment = nullptr;
+SettingString* appSettings::g_compilerOptions = nullptr;
+SettingBool* appSettings::g_compilerForcePchPdbs = nullptr;
 
-extern SettingString* appSettings::g_compilerPath = nullptr;
-extern SettingBool* appSettings::g_useCompilerOverrideAsFallback = nullptr;
-extern SettingBool* appSettings::g_useCompilerEnvironment = nullptr;
-extern SettingString* appSettings::g_compilerOptions = nullptr;
-extern SettingBool* appSettings::g_compilerForcePchPdbs = nullptr;
+SettingString* appSettings::g_linkerPath = nullptr;
+SettingBool* appSettings::g_useLinkerOverrideAsFallback = nullptr;
+SettingBool* appSettings::g_useLinkerEnvironment = nullptr;
+SettingString* appSettings::g_linkerOptions = nullptr;
 
-extern SettingString* appSettings::g_linkerPath = nullptr;
-extern SettingBool* appSettings::g_useLinkerOverrideAsFallback = nullptr;
-extern SettingBool* appSettings::g_useLinkerEnvironment = nullptr;
-extern SettingString* appSettings::g_linkerOptions = nullptr;
-extern SettingBool* appSettings::g_forceLinkWeakSymbols = nullptr;
+SettingBool* appSettings::g_continuousCompilationEnabled = nullptr;
+SettingString* appSettings::g_continuousCompilationPath = nullptr;
+SettingInt* appSettings::g_continuousCompilationTimeout = nullptr;
 
-extern SettingBool* appSettings::g_continuousCompilationEnabled = nullptr;
-extern SettingString* appSettings::g_continuousCompilationPath = nullptr;
-extern SettingInt* appSettings::g_continuousCompilationTimeout = nullptr;
+SettingString* appSettings::g_virtualDriveLetter = nullptr;
+SettingString* appSettings::g_virtualDrivePath = nullptr;
 
-extern SettingString* appSettings::g_virtualDriveLetter = nullptr;
-extern SettingString* appSettings::g_virtualDrivePath = nullptr;
+SettingBool* appSettings::g_installCompiledPatchesMultiProcess = nullptr;
 
-extern SettingBool* appSettings::g_installCompiledPatchesMultiProcess = nullptr;
+SettingBool* appSettings::g_amalgamationSplitIntoSingleParts = nullptr;
+SettingInt* appSettings::g_amalgamationSplitMinCppCount = nullptr;
+SettingString* appSettings::g_amalgamationCppFileExtensions = nullptr;
 
-extern SettingBool* appSettings::g_amalgamationSplitIntoSingleParts = nullptr;
-extern SettingInt* appSettings::g_amalgamationSplitMinCppCount = nullptr;
-extern SettingString* appSettings::g_amalgamationCppFileExtensions = nullptr;
+SettingBool* appSettings::g_ue4EnableNatVisSupport = nullptr;
 
-extern SettingBool* appSettings::g_ue4EnableNatVisSupport = nullptr;
\ No newline at end of file
+SettingString* appSettings::g_fastBuildDatabasePath = nullptr;
+SettingString* appSettings::g_fastBuildDllName = nullptr;
diff --git a/Engine/Source/Developer/Windows/LiveCodingServer/Private/External/LC_AppSettings.h b/Engine/Source/Developer/Windows/LiveCodingServer/Private/External/LC_AppSettings.h
--- a/Engine/Source/Developer/Windows/LiveCodingServer/Private/External/LC_AppSettings.h
+++ b/Engine/Source/Developer/Windows/LiveCodingServer/Private/External/LC_AppSettings.h
@@ -1,8 +1,10 @@
-// Copyright 2011-2019 Molecular Matters GmbH, all rights reserved.
+// Copyright 2011-2020 Molecular Matters GmbH, all rights reserved.
 
 #pragma once
 
+// BEGIN EPIC MOD
 #include "CoreTypes.h"
+// END EPIC MOD
 #include "LC_Settings.h"
 
 
@@ -75,6 +77,7 @@ namespace appSettings
 	extern SettingString* g_playSoundOnSuccess;
 	extern SettingString* g_playSoundOnError;
 	extern SettingShortcut* g_compileShortcut;
+	extern SettingInt* g_prewarmTimeout;
 
 	// logging
 	extern SettingBool* g_showUndecoratedNames;
@@ -96,7 +99,6 @@ namespace appSettings
 	extern SettingBool* g_useLinkerOverrideAsFallback;
 	extern SettingBool* g_useLinkerEnvironment;
 	extern SettingString* g_linkerOptions;
-	extern SettingBool* g_forceLinkWeakSymbols;
 
 	// continuous compilation
 	extern SettingBool* g_continuousCompilationEnabled;
@@ -117,4 +119,8 @@ namespace appSettings
 
 	// UE4-specific
 	extern SettingBool* g_ue4EnableNatVisSupport;
+
+	// FASTBuild-specific
+	extern SettingString* g_fastBuildDatabasePath;
+	extern SettingString* g_fastBuildDllName;
 }
diff --git a/Engine/Source/Developer/Windows/LiveCodingServer/Private/External/LC_COMThread.cpp b/Engine/Source/Developer/Windows/LiveCodingServer/Private/External/LC_COMThread.cpp
--- a/Engine/Source/Developer/Windows/LiveCodingServer/Private/External/LC_COMThread.cpp
+++ b/Engine/Source/Developer/Windows/LiveCodingServer/Private/External/LC_COMThread.cpp
@@ -1,12 +1,19 @@
-// Copyright 2011-2019 Molecular Matters GmbH, all rights reserved.
+// Copyright 2011-2020 Molecular Matters GmbH, all rights reserved.
 
+// BEGIN EPIC MOD
+//#include PCH_INCLUDE
+// END EPIC MOD
 #include "LC_COMThread.h"
-#include "LC_Platform.h"
+#include "LC_Thread.h"
+// BEGIN EPIC MOD
 #include "LC_Logging.h"
+#include "LC_Platform.h"
+#include LC_PLATFORM_INCLUDE(LC_Foundation)
 #include "Windows/MinimalWindowsAPI.h"
 #include <objidl.h>
 
 #include "Windows/AllowWindowsPlatformAtomics.h"
+// END EPIC MOD
 
 namespace
 {
@@ -16,6 +23,8 @@ namespace
 	class COMMessageFilter : public IMessageFilter
 	{
 	public:
+		virtual ~COMMessageFilter(void) {}
+
 		virtual HRESULT STDMETHODCALLTYPE QueryInterface(
 			/* [in] */ REFIID riid,
 			/* [iid_is][out] */ _COM_Outptr_ void __RPC_FAR *__RPC_FAR *ppvObject)
@@ -127,7 +136,7 @@ COMThread::COMThread(void)
 	, m_internalThread(nullptr)
 {
 	// launch the thread that takes care of calling COM functions
-	m_internalThread = thread::Create("Live++ COM", 65536u, &COMThread::ThreadFunction, this);
+	m_internalThread = Thread::CreateFromMemberFunction("Live++ COM", 65536u, this, &COMThread::ThreadFunction);
 }
 
 
@@ -137,12 +146,12 @@ COMThread::~COMThread(void)
 	m_leaveThreadEvent.Signal();
 	m_functionAvailableEvent.Signal();
 
-	thread::Join(m_internalThread);
-	thread::Close(m_internalThread);
+	Thread::Join(m_internalThread);
+	Thread::Close(m_internalThread);
 }
 
 
-unsigned int COMThread::ThreadFunction(void)
+Thread::ReturnValue COMThread::ThreadFunction(void)
 {
 	const HRESULT result = ::CoInitialize(NULL);
 	if (result != S_OK)
@@ -179,7 +188,9 @@ unsigned int COMThread::ThreadFunction(void)
 
 	::CoUninitialize();
 
-	return 0u;
+	return Thread::ReturnValue(0u);
 }
 
+// BEGIN EPIC MOD
 #include "Windows/HideWindowsPlatformAtomics.h"
+// END EPIC MOD
\ No newline at end of file
diff --git a/Engine/Source/Developer/Windows/LiveCodingServer/Private/External/LC_COMThread.h b/Engine/Source/Developer/Windows/LiveCodingServer/Private/External/LC_COMThread.h
--- a/Engine/Source/Developer/Windows/LiveCodingServer/Private/External/LC_COMThread.h
+++ b/Engine/Source/Developer/Windows/LiveCodingServer/Private/External/LC_COMThread.h
@@ -1,10 +1,12 @@
-// Copyright 2011-2019 Molecular Matters GmbH, all rights reserved.
+// Copyright 2011-2020 Molecular Matters GmbH, all rights reserved.
 
 #pragma once
 
+// BEGIN EPIC MOD
 #include <type_traits>
+// END EPIC MOD
 #include "LC_Event.h"
-#include "LC_Thread.h"
+#include "LC_ThreadTypes.h"
 
 
 // helper class that is the only one calling COM-related functions and methods.
@@ -68,11 +70,11 @@ public:
 	}
 
 private:
-	unsigned int ThreadFunction(void);
+	Thread::ReturnValue ThreadFunction(void);
 
 	COMFunction m_function;
 	Event m_functionAvailableEvent;
 	Event m_functionFinishedExecutingEvent;
 	Event m_leaveThreadEvent;
-	thread::Handle m_internalThread;
+	Thread::Handle m_internalThread;
 };
diff --git a/Engine/Source/Developer/Windows/LiveCodingServer/Private/External/LC_ChangeNotification.cpp b/Engine/Source/Developer/Windows/LiveCodingServer/Private/External/LC_ChangeNotification.cpp
--- a/Engine/Source/Developer/Windows/LiveCodingServer/Private/External/LC_ChangeNotification.cpp
+++ b/Engine/Source/Developer/Windows/LiveCodingServer/Private/External/LC_ChangeNotification.cpp
@@ -1,9 +1,13 @@
-// Copyright 2011-2019 Molecular Matters GmbH, all rights reserved.
+// Copyright 2011-2020 Molecular Matters GmbH, all rights reserved.
 
+// BEGIN EPIC MOD
+//#include PCH_INCLUDE
+// END EPIC MOD
 #include "LC_ChangeNotification.h"
+// BEGIN EPIC MOD
 #include "LC_Logging.h"
 #include "Windows/WindowsHWrapper.h"
-
+// END EPIC MOD
 
 ChangeNotification::ChangeNotification(void)
 	: m_handle(INVALID_HANDLE_VALUE)
@@ -19,7 +23,9 @@ ChangeNotification::~ChangeNotification(void)
 
 void ChangeNotification::Create(const wchar_t* path)
 {
+	// BEGIN EPIC MOD
 	m_handle = ::FindFirstChangeNotificationW(path, Windows::TRUE, FILE_NOTIFY_CHANGE_LAST_WRITE);
+	// END EPIC MOD
 	if (m_handle == INVALID_HANDLE_VALUE)
 	{
 		const DWORD error = ::GetLastError();
diff --git a/Engine/Source/Developer/Windows/LiveCodingServer/Private/External/LC_CodeCave.cpp b/Engine/Source/Developer/Windows/LiveCodingServer/Private/External/LC_CodeCave.cpp
--- a/Engine/Source/Developer/Windows/LiveCodingServer/Private/External/LC_CodeCave.cpp
+++ b/Engine/Source/Developer/Windows/LiveCodingServer/Private/External/LC_CodeCave.cpp
@@ -1,10 +1,14 @@
-// Copyright 2011-2019 Molecular Matters GmbH, all rights reserved.
+// Copyright 2011-2020 Molecular Matters GmbH, all rights reserved.
 
+// BEGIN EPIC MOD
+//#include PCH_INCLUDE
+// END EPIC MOD
 #include "LC_CodeCave.h"
+#include "LC_Process.h"
 #include "LC_Thread.h"
 
 
-CodeCave::CodeCave(process::Handle processHandle, unsigned int processId, unsigned int commandThreadId, const void* jumpToSelf)
+CodeCave::CodeCave(Process::Handle processHandle, Process::Id processId, Thread::Id commandThreadId, const void* jumpToSelf)
 	: m_processHandle(processHandle)
 	, m_processId(processId)
 	, m_commandThreadId(commandThreadId)
@@ -17,10 +21,10 @@ CodeCave::CodeCave(process::Handle processHandle, unsigned int processId, unsign
 
 void CodeCave::Install(void)
 {
-	process::Suspend(m_processHandle);
+	Process::Suspend(m_processHandle);
 
 	// enumerate all threads of the process now that it's suspended
-	const std::vector<unsigned int>& threadIds = process::EnumerateThreads(m_processId);
+	const std::vector<Thread::Id>& threadIds = Process::EnumerateThreads(m_processId);
 	
 	const size_t threadCount = threadIds.size();
 	m_perThreadData.resize(threadCount);
@@ -31,7 +35,7 @@ void CodeCave::Install(void)
 	// are currently running.
 	for (size_t i = 0u; i < threadCount; ++i)
 	{
-		const unsigned id = threadIds[i];
+		const Thread::Id id = threadIds[i];
 		m_perThreadData[i].id = id;
 
 		if (id == m_commandThreadId)
@@ -40,51 +44,51 @@ void CodeCave::Install(void)
 			continue;
 		}
 
-		thread::Handle threadHandle = thread::Open(id);
-		thread::Context context = thread::GetContext(threadHandle);
-		m_perThreadData[i].priority = thread::GetPriority(threadHandle);
-		m_perThreadData[i].originalIp = thread::ReadInstructionPointer(context);
-		thread::SetPriority(threadHandle, THREAD_PRIORITY_IDLE);
-		thread::WriteInstructionPointer(context, m_jumpToSelf);
-		thread::SetContext(threadHandle, context);
-		thread::Close(threadHandle);
+		Thread::Handle threadHandle = Thread::Open(id);
+		Thread::Context context = Thread::GetContext(threadHandle);
+		m_perThreadData[i].priority = Thread::GetPriority(threadHandle);
+		m_perThreadData[i].originalIp = Thread::ReadInstructionPointer(&context);
+		Thread::SetPriority(threadHandle, THREAD_PRIORITY_IDLE);
+		Thread::WriteInstructionPointer(&context, m_jumpToSelf);
+		Thread::SetContext(threadHandle, &context);
+		Thread::Close(threadHandle);
 	}
 
 	// let the process resume. all threads except the Live++ thread will be held in the code cave
-	process::Resume(m_processHandle);
+	Process::Resume(m_processHandle);
 }
 
 
 void CodeCave::Uninstall(void)
 {
-	process::Suspend(m_processHandle);
+	Process::Suspend(m_processHandle);
 
 	// restore original thread instruction pointers
 	const size_t threadCount = m_perThreadData.size();
 	for (size_t i = 0u; i < threadCount; ++i)
 	{
-		const unsigned id = m_perThreadData[i].id;
+		const Thread::Id id = m_perThreadData[i].id;
 		if (id == m_commandThreadId)
 		{
 			// this is the Live++ command thread
 			continue;
 		}
 
-		thread::Handle threadHandle = thread::Open(id);
-		thread::Context context = thread::GetContext(threadHandle);
-		const void* currentIp = thread::ReadInstructionPointer(context);
+		Thread::Handle threadHandle = Thread::Open(id);
+		Thread::Context context = Thread::GetContext(threadHandle);
+		const void* currentIp = Thread::ReadInstructionPointer(&context);
 
 		// only set the original instruction pointer if the thread is really being held in the cave.
 		// in certain situations (e.g. after an exception), the debugger/OS already restored the context
 		// of all threads, and it would be fatal to interfere with this.
 		if (currentIp == m_jumpToSelf)
 		{
-			thread::SetPriority(threadHandle, m_perThreadData[i].priority);
-			thread::WriteInstructionPointer(context, m_perThreadData[i].originalIp);
-			thread::SetContext(threadHandle, context);
+			Thread::SetPriority(threadHandle, m_perThreadData[i].priority);
+			Thread::WriteInstructionPointer(&context, m_perThreadData[i].originalIp);
+			Thread::SetContext(threadHandle, &context);
 		}
-		thread::Close(threadHandle);
+		Thread::Close(threadHandle);
 	}
 
-	process::Resume(m_processHandle);
+	Process::Resume(m_processHandle);
 }
diff --git a/Engine/Source/Developer/Windows/LiveCodingServer/Private/External/LC_CodeCave.h b/Engine/Source/Developer/Windows/LiveCodingServer/Private/External/LC_CodeCave.h
--- a/Engine/Source/Developer/Windows/LiveCodingServer/Private/External/LC_CodeCave.h
+++ b/Engine/Source/Developer/Windows/LiveCodingServer/Private/External/LC_CodeCave.h
@@ -1,30 +1,36 @@
-// Copyright 2011-2019 Molecular Matters GmbH, all rights reserved.
+// Copyright 2011-2020 Molecular Matters GmbH, all rights reserved.
 
 #pragma once
 
+// BEGIN EPIC MOD
 #include "CoreTypes.h"
-#include "LC_Process.h"
+// END EPIC MOD
+#include "LC_ProcessTypes.h"
+#include "LC_ThreadTypes.h"
+// BEGIN EPIC MOD
 #include "LC_Types.h"
+// END EPIC MOD
+
 
 // nifty helper to let all threads of a process except one make "progress" by being held inside a jump-to-self cave.
 class CodeCave
 {
 public:
 	// the jump-to-self code needs to be available in the address space of the given process
-	CodeCave(process::Handle processHandle, unsigned int processId, unsigned int commandThreadId, const void* jumpToSelf);
+	CodeCave(Process::Handle processHandle, Process::Id processId, Thread::Id commandThreadId, const void* jumpToSelf);
 
 	void Install(void);
 	void Uninstall(void);
 
 private:
-	process::Handle m_processHandle;
-	unsigned int m_processId;
-	unsigned int m_commandThreadId;
+	Process::Handle m_processHandle;
+	Process::Id m_processId;
+	Thread::Id m_commandThreadId;
 	const void* m_jumpToSelf;
 
 	struct PerThreadData
 	{
-		unsigned int id;
+		Thread::Id id;
 		const void* originalIp;
 		int priority;
 	};
diff --git a/Engine/Source/Developer/Windows/LiveCodingServer/Private/External/LC_Coff.cpp b/Engine/Source/Developer/Windows/LiveCodingServer/Private/External/LC_Coff.cpp
--- a/Engine/Source/Developer/Windows/LiveCodingServer/Private/External/LC_Coff.cpp
+++ b/Engine/Source/Developer/Windows/LiveCodingServer/Private/External/LC_Coff.cpp
@@ -1,17 +1,23 @@
-// Copyright 2011-2019 Molecular Matters GmbH, all rights reserved.
+// Copyright 2011-2020 Molecular Matters GmbH, all rights reserved.
 
+// BEGIN EPIC MOD
+//#include PCH_INCLUDE
+// END EPIC MOD
 #include "LC_Coff.h"
 #include "LC_CoffDetail.h"
 #include "LC_StringUtil.h"
 #include "LC_PointerUtil.h"
-#include "LC_FileUtil.h"
+#include "LC_Filesystem.h"
 #include "LC_SymbolPatterns.h"
 #include "LC_Symbols.h"
 #include "LC_MemoryStream.h"
 #include "LC_UniqueId.h"
+#include "LC_MemoryMappedFile.h"
+// BEGIN EPIC MOD
 #include "LC_Allocators.h"
+#include "LC_Assert.h"
 #include <algorithm>
-
+// END EPIC MOD
 
 namespace
 {
@@ -22,36 +28,6 @@ namespace
 	static const wchar_t COFF_SUFFIX_WIDE = L'%';
 
 	static const ImmutableString TLS_SECTION(".tls");
-
-
-	// helper class to deal with converting unique IDs to hex strings.
-	// converts an uint32_t into 8 hex characters, e.g. 255 will be converted to "000000FF".
-	class HexUniqueId
-	{
-	public:
-		static const unsigned int SIZE = 8u;
-
-		explicit HexUniqueId(uint32_t uniqueId)
-		{
-			const char hexChars[17u] = "0123456789ABCDEF";
-			m_hex[0u] = hexChars[(uniqueId >> 28u) & 0x0Fu];
-			m_hex[1u] = hexChars[(uniqueId >> 24u) & 0x0Fu];
-			m_hex[2u] = hexChars[(uniqueId >> 20u) & 0x0Fu];
-			m_hex[3u] = hexChars[(uniqueId >> 16u) & 0x0Fu];
-			m_hex[4u] = hexChars[(uniqueId >> 12u) & 0x0Fu];
-			m_hex[5u] = hexChars[(uniqueId >> 8u) & 0x0Fu];
-			m_hex[6u] = hexChars[(uniqueId >> 4u) & 0x0Fu];
-			m_hex[7u] = hexChars[(uniqueId >> 0u) & 0x0Fu];
-		}
-
-		inline const char* GetHex(void) const
-		{
-			return m_hex;
-		}
-
-	private:
-		char m_hex[SIZE];
-	};
 }
 
 
@@ -61,7 +37,7 @@ namespace coff
 	{
 		ObjFile* objFile = LC_NEW(&g_objFileAllocator, ObjFile);
 		objFile->filename = string::ToUtf8String(filename);
-		objFile->memoryFile = file::Open(filename, file::OpenMode::READ_ONLY);
+		objFile->memoryFile = Filesystem::OpenMemoryMappedFile(filename, Filesystem::OpenMode::READ);
 
 		return objFile;
 	}
@@ -69,32 +45,13 @@ namespace coff
 
 	void CloseObj(ObjFile*& objFile)
 	{
-		file::Close(objFile->memoryFile);
+		Filesystem::CloseMemoryMappedFile(objFile->memoryFile);
 		
 		LC_DELETE(&g_objFileAllocator, objFile, sizeof(ObjFile));
 		objFile = nullptr;
 	}
 
 
-	LibFile* OpenLib(const wchar_t* filename)
-	{
-		LibFile* libFile = LC_NEW(&g_libFileAllocator, LibFile);
-		libFile->filename = string::ToUtf8String(filename);
-		libFile->memoryFile = file::Open(filename, file::OpenMode::READ_ONLY);
-
-		return libFile;
-	}
-
-
-	void CloseLib(LibFile*& libFile)
-	{
-		file::Close(libFile->memoryFile);
-
-		LC_DELETE(&g_libFileAllocator, libFile, sizeof(LibFile));
-		libFile = nullptr;
-	}
-
-
 	static size_t FindSymbolShortNameLength(const char* shortName)
 	{
 		size_t length = 0u;
@@ -196,7 +153,7 @@ namespace coff
 
 
 	template <typename T>
-	static ImmutableString GetSymbolName(const char* stringTable, const T* symbol, uint32_t uniqueId, const HexUniqueId& hexUniqueId, uint16_t uniqueCounter, coff::ReadFlags::Enum readFlags)
+	static ImmutableString GetSymbolName(const char* stringTable, const T* symbol, uint32_t uniqueId, uint16_t uniqueCounter)
 	{
 		/*
 		From the COFF spec:
@@ -228,63 +185,6 @@ namespace coff
 		{
 			// long name, points into string table
 			const char* str = stringTable + symbol->N.Name.Long;
-
-			// AMALGAMATION
-			if (readFlags == coff::ReadFlags::GENERATE_ANS_NAME_FROM_UNIQUE_ID)
-			{
-				// this could be the name of an anonymous namespace (ANS). an ANS symbol name is always of the form
-				// ?identifier@?A0x12345678, where the hex code following the "@?A0x" part is most likely a hash of the
-				// filename the ANS appears in, generated by the compiler.
-				// note that due to the structure of such a symbol name it can never be a short name (limited to 8 chars).
-
-				// when splitting amalgamated files, we need to make sure that symbols in anonymous namespaces compiled into
-				// those files are also found when compiled into single-part files.
-				// however, single-part files get assigned a different hash by the compiler, leading to different
-				// symbol names for symbols that reside in anonymous namespaces.
-				// in order to "correct" this, we generate our own hex identifier for ANS symbols, making sure that this
-				// identifier yields the same result for both amalgamated as well as single-part files.
-				// this is done by using the uniqueId as hex identifier, which is the same for amalgamated files as well
-				// as their split single-file counterparts.
-				const char* anonNamespace = string::Find(str, symbolPatterns::ANONYMOUS_NAMESPACE_PATTERN);
-				if (anonNamespace)
-				{
-					const size_t patternLength = strlen(symbolPatterns::ANONYMOUS_NAMESPACE_PATTERN);
-
-					// make a copy of the original symbol name
-					const size_t lengthWithoutNull = strlen(str);
-					char* newStr = new char[lengthWithoutNull + 1u];
-					memcpy(newStr, str, lengthWithoutNull + 1u);
-
-					// index of the first hex character found in the anonymous namespace name (skips the "@?A0x" pattern)
-					const ptrdiff_t indexOfHexId = anonNamespace - str + static_cast<ptrdiff_t>(patternLength);
-
-					// overwrite the compiler-generated hex ID with ours
-					memcpy(newStr + indexOfHexId, hexUniqueId.GetHex(), HexUniqueId::SIZE);
-
-					// the identifier could contain several more anonymous namespaces
-					char* found = newStr + indexOfHexId;
-					do
-					{
-						found = string::Find(found, symbolPatterns::ANONYMOUS_NAMESPACE_PATTERN);
-						if (found)
-						{
-							// overwrite the compiler-generated hex ID with ours
-							found += patternLength;
-							memcpy(found, hexUniqueId.GetHex(), HexUniqueId::SIZE);
-						}
-					}
-					while (found);
-
-					ImmutableString result = isStatic
-						? DisambiguateStaticSymbolName(newStr, lengthWithoutNull, uniqueId, uniqueCounter)
-						: ImmutableString(newStr);
-
-					delete[] newStr;
-
-					return result;
-				}
-			}
-
 			return isStatic
 				? DisambiguateStaticSymbolName(str, strlen(str), uniqueId, uniqueCounter)
 				: ImmutableString(str);
@@ -292,42 +192,6 @@ namespace coff
 	}
 
 
-	static ImmutableString GetArchiveMemberName(const IMAGE_ARCHIVE_MEMBER_HEADER* header, const char* longnamesMember)
-	{
-		// COFF Spec, 7.2. Archive Member Headers:
-		//	Field "Name" at offset 0:	The name of the archive member, with a slash (/) appended to terminate the name.
-		//								If the first character is a slash, the name has a special interpretation [...].
-		const char* name = reinterpret_cast<const char*>(header->Name);
-		if (name[0] == '/')
-		{
-			/*
-			From the COFF spec, 7.2. Archive Member Headers:
-				The name of the archive member is located at offset n within the longnames
-				member. The number n is the decimal representation of the offset.
-				For example : "/26" indicates that the name of the archive member is located 26
-				bytes beyond the beginning of the longnames member contents.
-			*/
-			uint32_t offset = 0u;
-			sscanf_s(name + 1u, "%" SCNu32, &offset);
-
-			// strings stored in the longnames member are null-terminated, so we can directly construct an
-			// immutable string from there
-			return ImmutableString(longnamesMember + offset);
-		}
-		else
-		{
-			// the name is terminated with '/'
-			size_t length = 0u;
-			while ((name[length] != '/') && (length < sizeof(header->Name)))
-			{
-				++length;
-			}
-
-			return ImmutableString(name, length);
-		}
-	}
-
-
 	static ImmutableString GetSectionName(const char* stringTable, const IMAGE_SECTION_HEADER* section)
 	{
 		/*
@@ -409,7 +273,7 @@ namespace coff
 
 
 	template <typename RawCoffType, coffDetail::CoffType::Enum Type>
-	static RawCoffType* ReadRaw(ObjFile* file, uint32_t uniqueId, coff::ReadFlags::Enum readFlags)
+	static RawCoffType* ReadRaw(ObjFile* file, uint32_t uniqueId)
 	{
 		using CoffHeader = typename coffDetail::CoffType::TypeMap<Type>::HeaderType;
 		using CoffSymbol = typename coffDetail::CoffType::TypeMap<Type>::SymbolType;
@@ -417,18 +281,20 @@ namespace coff
 
 		RawCoffType* rawCoff = new RawCoffType;
 		{
-			const file::Attributes& attributes = file::GetAttributes(string::ToWideString(file->filename).c_str());
-			rawCoff->size = file::GetSize(attributes);
+			const Filesystem::PathAttributes& attributes = Filesystem::GetAttributes(string::ToWideString(file->filename).c_str());
+			rawCoff->size = Filesystem::GetSize(attributes);
 		}
 		rawCoff->type = RawCoffType::TYPE;
 
+		const void* objFileBase = Filesystem::GetMemoryMappedFileData(file->memoryFile);
+
 		// section header is first
-		const CoffHeader* rawCoffHeader = pointer::As<const CoffHeader*>(file->memoryFile->base);
+		const CoffHeader* rawCoffHeader = pointer::As<const CoffHeader*>(objFileBase);
 		rawCoff->header = *rawCoffHeader;
 
 		// section headers follow after that
-		const IMAGE_SECTION_HEADER* rawSectionHeaders = coffDetail::GetSectionHeader<CoffHeader>(file->memoryFile->base);
-		const uint32_t sectionCount = coffDetail::GetNumberOfSections<CoffHeader>(file->memoryFile->base);
+		const IMAGE_SECTION_HEADER* rawSectionHeaders = coffDetail::GetSectionHeader<CoffHeader>(objFileBase);
+		const uint32_t sectionCount = coffDetail::GetNumberOfSections<CoffHeader>(objFileBase);
 		rawCoff->sections.reserve(sectionCount);
 		for (uint32_t i = 0u; i < sectionCount; ++i)
 		{
@@ -447,13 +313,13 @@ namespace coff
 				// some sections like .bss don't store (uninitialized) data
 				const DWORD rawDataSize = rawSection->SizeOfRawData;
 				section.data = LC_ALLOC(&g_rawCoffAllocator, rawDataSize, 8u);
-				memcpy(section.data, pointer::Offset<const void*>(file->memoryFile->base, rawSection->PointerToRawData), rawDataSize);
+				memcpy(section.data, pointer::Offset<const void*>(objFileBase, rawSection->PointerToRawData), rawDataSize);
 			}
 
 			// relocations
 			{
-				const IMAGE_RELOCATION* rawRelocations = pointer::Offset<const IMAGE_RELOCATION*>(file->memoryFile->base, rawSection->PointerToRelocations);
-				const size_t count = coffDetail::GetRelocationCount(file->memoryFile->base, rawSection);
+				const IMAGE_RELOCATION* rawRelocations = pointer::Offset<const IMAGE_RELOCATION*>(objFileBase, rawSection->PointerToRelocations);
+				const size_t count = coffDetail::GetRelocationCount(objFileBase, rawSection);
 
 				// if relocation count in section has overflown, ignore the first relocation
 				const size_t startRelocation = (count > 0xFFFFu) ? 1u : 0u;
@@ -466,7 +332,7 @@ namespace coff
 
 			// line numbers
 			{
-				const IMAGE_LINENUMBER* rawLineNumbers = pointer::Offset<const IMAGE_LINENUMBER*>(file->memoryFile->base, rawSection->PointerToLinenumbers);
+				const IMAGE_LINENUMBER* rawLineNumbers = pointer::Offset<const IMAGE_LINENUMBER*>(objFileBase, rawSection->PointerToLinenumbers);
 				const size_t count = rawSection->NumberOfLinenumbers;
 				section.lineNumbers.reserve(count);
 				for (size_t j = 0u; j < count; ++j)
@@ -479,8 +345,8 @@ namespace coff
 		}
 
 		// symbol table follows after section headers, but offset is stored in COFF header directly
-		const void* rawSymbolTable = coffDetail::GetSymbolTable<CoffHeader>(file->memoryFile->base);
-		const uint32_t symbolCount = coffDetail::GetNumberOfSymbols<CoffHeader>(file->memoryFile->base);
+		const void* rawSymbolTable = coffDetail::GetSymbolTable<CoffHeader>(objFileBase);
+		const uint32_t symbolCount = coffDetail::GetNumberOfSymbols<CoffHeader>(objFileBase);
 		rawCoff->symbols.reserve(symbolCount);
 		for (uint32_t i = 0u; i < symbolCount; ++i)
 		{
@@ -496,8 +362,6 @@ namespace coff
 		rawCoff->rawStringTable.data = static_cast<char*>(LC_ALLOC(&g_rawCoffAllocator, rawCoff->rawStringTable.size, 8u));
 		memcpy(rawCoff->rawStringTable.data, rawStringTable, rawCoff->rawStringTable.size);
 
-		const HexUniqueId hexUniqueId(uniqueId);
-
 		// construct string table
 		{
 			rawCoff->stringTable.resize(symbolCount);
@@ -523,7 +387,7 @@ namespace coff
 				}
 
 				const SymbolType::Enum type = DetermineSymbolType(symbol);
-				rawCoff->stringTable[i] = GetSymbolName(rawStringTable, symbol, uniqueId, hexUniqueId, 0u, readFlags);
+				rawCoff->stringTable[i] = GetSymbolName(rawStringTable, symbol, uniqueId, 0u);
 				const ImmutableString& name = rawCoff->stringTable[i];
 
 				if (type == SymbolType::STATIC_DATA)
@@ -539,7 +403,7 @@ namespace coff
 
 						// increase the counter associated with this name and generate a new name from it
 						++insertIterator.first->second;
-						rawCoff->stringTable[i] = GetSymbolName(rawStringTable, symbol, uniqueId, hexUniqueId, insertIterator.first->second, readFlags);
+						rawCoff->stringTable[i] = GetSymbolName(rawStringTable, symbol, uniqueId, insertIterator.first->second);
 					}
 				}
 			}
@@ -593,16 +457,16 @@ namespace coff
 	}
 
 
-	RawCoff* ReadRaw(ObjFile* file, uint32_t uniqueId, coff::ReadFlags::Enum readFlags)
+	RawCoff* ReadRaw(ObjFile* file, uint32_t uniqueId)
 	{
-		const coffDetail::CoffType::Enum type = coffDetail::GetCoffType(file->memoryFile->base);
+		const coffDetail::CoffType::Enum type = coffDetail::GetCoffType(Filesystem::GetMemoryMappedFileData(file->memoryFile));
 		if (type == coffDetail::CoffType::COFF)
 		{
-			return ReadRaw<RawCoffRegular, coffDetail::CoffType::COFF>(file, uniqueId, readFlags);
+			return ReadRaw<RawCoffRegular, coffDetail::CoffType::COFF>(file, uniqueId);
 		}
 		else if (type == coffDetail::CoffType::BIGOBJ)
 		{
-			return ReadRaw<RawCoffBigObj, coffDetail::CoffType::BIGOBJ>(file, uniqueId, readFlags);
+			return ReadRaw<RawCoffBigObj, coffDetail::CoffType::BIGOBJ>(file, uniqueId);
 		}
 
 		return nullptr;
@@ -614,7 +478,6 @@ namespace coff
 	{
 		using CoffHeader = typename coffDetail::CoffType::TypeMap<Type>::HeaderType;
 		using CoffSymbol = typename coffDetail::CoffType::TypeMap<Type>::SymbolType;
-		using CoffAuxSymbol = typename coffDetail::CoffType::TypeMap<Type>::AuxSymbolType;
 
 		// the file structure is as follows:
 		//	- COFF file header
@@ -766,7 +629,7 @@ namespace coff
 			}
 		}
 
-		file::CreateFileWithData(filename, outputData.GetData(), outputData.GetSize());
+		Filesystem::CreateFileWithData(filename, outputData.GetData(), outputData.GetSize());
 	}
 
 
@@ -872,6 +735,13 @@ namespace coff
 	}
 
 
+	ImmutableString GetSectionName(const RawCoff* rawCoff, size_t index)
+	{
+		const IMAGE_SECTION_HEADER* header = &rawCoff->sections[index].header;
+		return GetSectionName(rawCoff->rawStringTable.data, header);
+	}
+
+
 	uint32_t GetSymbolSectionIndex(const RawCoff* rawCoff, size_t index)
 	{
 		if (rawCoff->type == RawCoffRegular::TYPE)
@@ -1007,19 +877,21 @@ namespace coff
 
 
 	template <coffDetail::CoffType::Enum Type>
-	static types::vector<std::string> ExtractLinkerDirectives(const ObjFile* file)
+	static types::vector<std::string> ExtractLinkerDirectives(const ObjFile* objFile)
 	{
 		using CoffHeader = typename coffDetail::CoffType::TypeMap<Type>::HeaderType;
 
-		const uint32_t sectionCount = coffDetail::GetNumberOfSections<CoffHeader>(file->memoryFile->base);
-		const IMAGE_SECTION_HEADER* sectionHeader = coffDetail::GetSectionHeader<CoffHeader>(file->memoryFile->base);
+		const void* objFileBase = Filesystem::GetMemoryMappedFileData(objFile->memoryFile);
+
+		const uint32_t sectionCount = coffDetail::GetNumberOfSections<CoffHeader>(objFileBase);
+		const IMAGE_SECTION_HEADER* sectionHeader = coffDetail::GetSectionHeader<CoffHeader>(objFileBase);
 
 		for (unsigned int i = 0u; i < sectionCount; ++i)
 		{
 			const IMAGE_SECTION_HEADER* section = sectionHeader + i;
 			if (coffDetail::IsDirectiveSection(section))
 			{			
-				const char* rawData = pointer::Offset<const char*>(file->memoryFile->base, section->PointerToRawData);
+				const char* rawData = pointer::Offset<const char*>(objFileBase, section->PointerToRawData);
 				return TokenizeLinkerDirectives(rawData, section->SizeOfRawData);
 			}
 		}
@@ -1030,7 +902,7 @@ namespace coff
 
 	types::vector<std::string> ExtractLinkerDirectives(const ObjFile* file)
 	{
-		const coffDetail::CoffType::Enum type = coffDetail::GetCoffType(file->memoryFile->base);
+		const coffDetail::CoffType::Enum type = coffDetail::GetCoffType(Filesystem::GetMemoryMappedFileData(file->memoryFile));
 		if (type == coffDetail::CoffType::COFF)
 		{
 			return ExtractLinkerDirectives<coffDetail::CoffType::COFF>(file);
@@ -1193,7 +1065,7 @@ namespace coff
 
 
 	template <coffDetail::CoffType::Enum Type>
-	static CoffDB* GatherDatabase(ObjFile* file, uint32_t uniqueId, coff::ReadFlags::Enum readFlags)
+	static CoffDB* GatherDatabase(ObjFile* objFile, uint32_t uniqueId)
 	{
 		using CoffHeader = typename coffDetail::CoffType::TypeMap<Type>::HeaderType;
 		using CoffSymbol = typename coffDetail::CoffType::TypeMap<Type>::SymbolType;
@@ -1203,10 +1075,12 @@ namespace coff
 		coffDb->symbolAllocator = new PoolAllocator<PoolAllocatorSingleThreadPolicy>("COFF symbols", sizeof(coff::Symbol), alignof(coff::Symbol), 8192u);
 		coffDb->relocationAllocator = new PoolAllocator<PoolAllocatorSingleThreadPolicy>("Relocations", sizeof(coff::Relocation), alignof(coff::Relocation), 8192u);
 
-		const uint32_t sectionCount = coffDetail::GetNumberOfSections<CoffHeader>(file->memoryFile->base);
-		const uint32_t symbolCount = coffDetail::GetNumberOfSymbols<CoffHeader>(file->memoryFile->base);
-		const IMAGE_SECTION_HEADER* sectionHeader = coffDetail::GetSectionHeader<CoffHeader>(file->memoryFile->base);
-		const void* symbolTable = coffDetail::GetSymbolTable<CoffHeader>(file->memoryFile->base);
+		const void* objFileBase = Filesystem::GetMemoryMappedFileData(objFile->memoryFile);
+
+		const uint32_t sectionCount = coffDetail::GetNumberOfSections<CoffHeader>(objFileBase);
+		const uint32_t symbolCount = coffDetail::GetNumberOfSymbols<CoffHeader>(objFileBase);
+		const IMAGE_SECTION_HEADER* sectionHeader = coffDetail::GetSectionHeader<CoffHeader>(objFileBase);
+		const void* symbolTable = coffDetail::GetSymbolTable<CoffHeader>(objFileBase);
 		const char* stringTable = coffDetail::GetStringTable<CoffSymbol>(symbolTable, symbolCount);
 
 		// pre-allocate data structures
@@ -1241,14 +1115,14 @@ namespace coff
 				// add this section to the lookup-table
 				lutSectionIndexToCrtSection[i] = static_cast<uint32_t>(coffDb->crtSections.size());
 
-				const CrtSection crtSection { s.name, section->SizeOfRawData, section->PointerToRawData };
+				const CrtSection crtSection { s.name, section->SizeOfRawData, section->PointerToRawData, {} };
 				coffDb->crtSections.push_back(crtSection);
 			}
 
 			// we would like to reserve space for the symbols of each section to avoid allocations where possible.
 			// because we don't know yet how many symbols a section holds, using the number of relocations is a
 			// good approximation.
-			const unsigned int relocationCount = coffDetail::GetRelocationCount(file->memoryFile->base, section);
+			const unsigned int relocationCount = coffDetail::GetRelocationCount(objFileBase, section);
 			symbolsForSection[i].reserve(relocationCount);
 		}
 
@@ -1271,8 +1145,6 @@ namespace coff
 		types::StringMap<uint16_t> uniqueStaticDataSymbols;
 		uniqueStaticDataSymbols.reserve(16u);
 
-		const HexUniqueId hexUniqueId(uniqueId);
-
 		for (unsigned int i = 0u; i < symbolCount; ++i)
 		{
 			const CoffSymbol* symbol = coffDetail::GetSymbol<CoffSymbol>(symbolTable, i);
@@ -1286,7 +1158,7 @@ namespace coff
 			}
 			else if (coffDetail::IsUndefinedSymbol(symbol))
 			{
-				coffDb->stringTable[i] = GetSymbolName(stringTable, symbol, uniqueId, hexUniqueId, 0u, readFlags);
+				coffDb->stringTable[i] = GetSymbolName(stringTable, symbol, uniqueId, 0u);
 			}
 			else if (coffDetail::IsSectionSymbol(symbol))
 			{
@@ -1304,11 +1176,13 @@ namespace coff
 						coffDb->sections[sectionIndex].comdatSelection = auxSymbol->Section.Selection;
 					}
 				}
+
+				coffDb->stringTable[i] = GetSectionName(stringTable, section);
 			}
 			else
 			{
 				// this symbol is stored in the COFF.
-				coffDb->stringTable[i] = GetSymbolName(stringTable, symbol, uniqueId, hexUniqueId, 0u, readFlags);
+				coffDb->stringTable[i] = GetSymbolName(stringTable, symbol, uniqueId, 0u);
 				const ImmutableString& name = coffDb->stringTable[i];
 
 				// we are not interested in certain types of symbols.
@@ -1321,7 +1195,7 @@ namespace coff
 					const uint32_t rva = section->PointerToRawData + symbol->Value;
 					const SymbolType::Enum type = DetermineSymbolType(symbol);
 
-					Symbol* newSymbol = LC_NEW(coffDb->symbolAllocator, Symbol) { i, rva, sectionIndex, type };
+					Symbol* newSymbol = LC_NEW(coffDb->symbolAllocator, Symbol) { i, rva, sectionIndex, type, {} };
 					newSymbol->relocations.reserve(32u);
 					coffDb->symbols.push_back(newSymbol);
 
@@ -1350,13 +1224,13 @@ namespace coff
 								if (!string::StartsWith(symbolName.c_str(), "$unwind$??__F") &&
 									!string::StartsWith(symbolName.c_str(), "$pdata$??__F"))
 								{
-									LC_WARNING_USER("Non-unique symbol %s found in COFF file %s. Do not change the order of these variables while live coding, or consider upgrading to a newer compiler (VS 2015 or later)",
-										symbolName.c_str(), file->filename.c_str());
+								    LC_WARNING_USER("Non-unique symbol %s found in COFF file %s. Do not change the order of these variables while live coding, or consider upgrading to a newer compiler (VS 2015 or later)",
+									    symbolName.c_str(), objFile->filename.c_str());
 								}
 								else
 								{
 									LC_LOG_USER("Non-unique at-exit symbol %s found in COFF file %s. These sometimes appear in debug builds",
-										symbolName.c_str(), file->filename.c_str());
+										symbolName.c_str(), objFile->filename.c_str());
 								}
 								// END EPIC MOD
 							}
@@ -1367,7 +1241,7 @@ namespace coff
 
 							// increase the counter associated with this name and generate a new name from it
 							++insertIterator.first->second;
-							coffDb->stringTable[i] = GetSymbolName(stringTable, symbol, uniqueId, hexUniqueId, insertIterator.first->second, readFlags);
+							coffDb->stringTable[i] = GetSymbolName(stringTable, symbol, uniqueId, insertIterator.first->second);
 						}
 					}
 
@@ -1397,8 +1271,8 @@ namespace coff
 					continue;
 				}
 
-				const unsigned int relocationCount = coffDetail::GetRelocationCount(file->memoryFile->base, section);
-				const IMAGE_RELOCATION* relocations = pointer::Offset<const IMAGE_RELOCATION*>(file->memoryFile->base, section->PointerToRelocations);
+				const unsigned int relocationCount = coffDetail::GetRelocationCount(objFileBase, section);
+				const IMAGE_RELOCATION* relocations = pointer::Offset<const IMAGE_RELOCATION*>(objFileBase, section->PointerToRelocations);
 
 				const types::vector<Symbol*>& symbolsForCurrentSection = symbolsForSection[i];
 				if (symbolsForCurrentSection.size() == 0u)
@@ -1424,13 +1298,6 @@ namespace coff
 					}
 
 					const CoffSymbol* symbol = coffDetail::GetSymbol<CoffSymbol>(symbolTable, relocation->SymbolTableIndex);
-
-					// ignore relocations to sections
-					if (coffDetail::IsSectionSymbol(symbol))
-					{
-						continue;
-					}
-
 					const ImmutableString& dstSymbolName = coffDb->stringTable[relocation->SymbolTableIndex];
 
 					// ignore relocations to line numbers and string literals
@@ -1481,7 +1348,7 @@ namespace coff
 					Symbol* srcSymbol = symbolsForCurrentSection[currentSymbolIndex];
 					if (relocationRva < srcSymbol->rva)
 					{
-						LC_ERROR_DEV("Cannot find symbol that contains relocation at 0x%X for destination symbol %s in file %s", relocationRva, dstSymbolName.c_str(), file->filename.c_str());
+						LC_ERROR_DEV("Cannot find symbol that contains relocation at 0x%X for destination symbol %s in file %s", relocationRva, dstSymbolName.c_str(), objFile->filename.c_str());
 						continue;
 					}
 
@@ -1504,7 +1371,7 @@ namespace coff
 						case Relocation::Type::RELATIVE_OFFSET_5:
 #endif
 						{
-							const uint32_t* relocationAddress = pointer::Offset<const uint32_t*>(file->memoryFile->base, relocationRva);
+							const uint32_t* relocationAddress = pointer::Offset<const uint32_t*>(objFileBase, relocationRva);
 							const uint32_t destinationOffset = *relocationAddress;
 							Relocation* newRelocation = LC_NEW(coffDb->relocationAllocator, Relocation)
 							{
@@ -1514,7 +1381,8 @@ namespace coff
 								sectionIndex,
 								static_cast<Relocation::Type::Enum>(relocation->Type),
 								srcSymbol->type,
-								DetermineSymbolType(symbol)
+								DetermineSymbolType(symbol),
+								coffDetail::IsSectionSymbol(symbol)
 							};
 							srcSymbol->relocations.emplace_back(newRelocation);
 						}
@@ -1523,7 +1391,7 @@ namespace coff
 #if LC_64_BIT
 						case Relocation::Type::VA_64:
 						{
-							const uint64_t* relocationAddress = pointer::Offset<const uint64_t*>(file->memoryFile->base, relocationRva);
+							const uint64_t* relocationAddress = pointer::Offset<const uint64_t*>(objFileBase, relocationRva);
 
 							// read the destination offset as 64-bit, but convert it into a 32-bit offset.
 							// no symbol can ever be larger than 4 GB.
@@ -1536,7 +1404,8 @@ namespace coff
 								sectionIndex,
 								static_cast<Relocation::Type::Enum>(relocation->Type),
 								srcSymbol->type,
-								DetermineSymbolType(symbol)
+								DetermineSymbolType(symbol),
+								coffDetail::IsSectionSymbol(symbol)
 							};
 							srcSymbol->relocations.emplace_back(newRelocation);
 						}
@@ -1579,16 +1448,16 @@ namespace coff
 	}
 
 
-	CoffDB* GatherDatabase(ObjFile* file, uint32_t uniqueId, coff::ReadFlags::Enum readFlags)
+	CoffDB* GatherDatabase(ObjFile* file, uint32_t uniqueId)
 	{
-		const coffDetail::CoffType::Enum type = coffDetail::GetCoffType(file->memoryFile->base);
+		const coffDetail::CoffType::Enum type = coffDetail::GetCoffType(Filesystem::GetMemoryMappedFileData(file->memoryFile));
 		if (type == coffDetail::CoffType::COFF)
 		{
-			return GatherDatabase<coffDetail::CoffType::COFF>(file, uniqueId, readFlags);
+			return GatherDatabase<coffDetail::CoffType::COFF>(file, uniqueId);
 		}
 		else if (type == coffDetail::CoffType::BIGOBJ)
 		{
-			return GatherDatabase<coffDetail::CoffType::BIGOBJ>(file, uniqueId, readFlags);
+			return GatherDatabase<coffDetail::CoffType::BIGOBJ>(file, uniqueId);
 		}
 
 		return nullptr;
@@ -1618,147 +1487,19 @@ namespace coff
 	}
 
 
-	LibDB* GatherDatabase(LibFile* libFile)
-	{
-		LibDB* libDb = new LibDB;
-
-		const void* fileBase = file::GetData(libFile->memoryFile);
-
-		const file::Attributes& attributes = file::GetAttributes(string::ToWideString(libFile->filename).c_str());
-		const uint64_t fileSize = file::GetSize(attributes);
-
-		const void* fileEnd = pointer::Offset<const void*>(fileBase, fileSize);
-
-		// according to COFF spec 7. "Archive (Library) File Format", first comes the signature
-		const char* signature = pointer::As<const char*>(fileBase);
-		if (strncmp(signature, IMAGE_ARCHIVE_START, IMAGE_ARCHIVE_START_SIZE) != 0)
-		{
-			LC_ERROR_DEV("Unknown archive format");
-			return libDb;
-		}
-
-		// after that comes the "first linker member", which we ignore but check for correctness
-		const IMAGE_ARCHIVE_MEMBER_HEADER* firstLinkerMemberHeader = pointer::Offset<const IMAGE_ARCHIVE_MEMBER_HEADER*>(signature, IMAGE_ARCHIVE_START_SIZE);
-		if ((firstLinkerMemberHeader->Name[0] != IMAGE_ARCHIVE_LINKER_MEMBER[0]) || (firstLinkerMemberHeader->Name[1] != IMAGE_ARCHIVE_LINKER_MEMBER[1]))
-		{
-			LC_ERROR_DEV("First linker member in archive seems to be corrupt");
-			return libDb;
-		}
-		const uint32_t firstLinkerMemberSize = coffDetail::PadArchiveMemberSize(coffDetail::GetArchiveMemberSize(firstLinkerMemberHeader));
-
-		// next comes the "second linker member"
-		const IMAGE_ARCHIVE_MEMBER_HEADER* secondLinkerMemberHeader = pointer::Offset<const IMAGE_ARCHIVE_MEMBER_HEADER*>(firstLinkerMemberHeader, IMAGE_SIZEOF_ARCHIVE_MEMBER_HDR + firstLinkerMemberSize);
-		if ((secondLinkerMemberHeader->Name[0] != IMAGE_ARCHIVE_LINKER_MEMBER[0]) || (secondLinkerMemberHeader->Name[1] != IMAGE_ARCHIVE_LINKER_MEMBER[1]))
-		{
-			LC_ERROR_DEV("Second linker member in archive seems to be corrupt");
-			return libDb;
-		}
-
-		// COFF Spec: 7.4 Second Linker Member
-		{
-			const void* secondLinkerMember = pointer::Offset<const void*>(secondLinkerMemberHeader, IMAGE_SIZEOF_ARCHIVE_MEMBER_HDR);
-
-			// skip offsets and indices, go to string table directly
-			const uint32_t* ptrNumberOfMembers = pointer::Offset<const uint32_t*>(secondLinkerMember, 0u);
-			const uint32_t* ptrNumberOfSymbols = pointer::Offset<const uint32_t*>(secondLinkerMember, sizeof(uint32_t) + *ptrNumberOfMembers * sizeof(uint32_t));
-			const char* stringTable = pointer::Offset<const char*>(ptrNumberOfSymbols, sizeof(uint32_t) + *ptrNumberOfSymbols * sizeof(uint16_t));
-			const char* currentSymbolName = stringTable;
-
-			const uint32_t symbolCount = *ptrNumberOfSymbols;
-			libDb->exportedSymbols.reserve(symbolCount);
-
-			for (uint32_t i = 0u; i < symbolCount; ++i)
-			{
-				const ImmutableString symbolName(currentSymbolName);
-				currentSymbolName += symbolName.GetLength() + 1u;
-
-				libDb->exportedSymbols.emplace_back(symbolName);
-			}
-		}
-
-		const uint32_t secondLinkerMemberSize = coffDetail::PadArchiveMemberSize(coffDetail::GetArchiveMemberSize(secondLinkerMemberHeader));
-
-		// next comes the optional "longnames member"
-		const IMAGE_ARCHIVE_MEMBER_HEADER* header = pointer::Offset<const IMAGE_ARCHIVE_MEMBER_HEADER*>(secondLinkerMemberHeader, IMAGE_SIZEOF_ARCHIVE_MEMBER_HDR + secondLinkerMemberSize);
-
-		const char* longnamesMember = nullptr;
-		const bool hasLongnamesMember = ((header->Name[0] == IMAGE_ARCHIVE_LONGNAMES_MEMBER[0]) && (header->Name[1] == IMAGE_ARCHIVE_LONGNAMES_MEMBER[1]));
-		if (hasLongnamesMember)
-		{
-			longnamesMember = pointer::Offset<const char*>(header, IMAGE_SIZEOF_ARCHIVE_MEMBER_HDR);
-
-			// skip the longnames member's data to arrive at the first COFF file header
-			const uint32_t longnamesMemberSize = coffDetail::PadArchiveMemberSize(coffDetail::GetArchiveMemberSize(header));
-			header = pointer::Offset<const IMAGE_ARCHIVE_MEMBER_HEADER*>(header, IMAGE_SIZEOF_ARCHIVE_MEMBER_HDR + longnamesMemberSize);
-		}
-
-		// finally, all the COFF files follow
-		for (/* nothing */; header < fileEnd; /* nothing */)
-		{			
-			ImmutableString coffName = GetArchiveMemberName(header, longnamesMember);
-
-			// skip the header to arrive at the raw COFF data for this archive member
-			const void* coffFileStart = pointer::Offset<const void*>(header, IMAGE_SIZEOF_ARCHIVE_MEMBER_HDR);
-			const uint64_t offset = pointer::Displacement<uint64_t>(fileBase, coffFileStart);
-			libDb->libEntries.emplace_back(LibEntry { std::move(coffName), offset });
-
-			const uint32_t memberSize = coffDetail::PadArchiveMemberSize(coffDetail::GetArchiveMemberSize(header));
-			header = pointer::Offset<const IMAGE_ARCHIVE_MEMBER_HEADER*>(header, IMAGE_SIZEOF_ARCHIVE_MEMBER_HDR + memberSize);
-		}
-
-		return libDb;
-	}
-
-
-	void DestroyDatabase(LibDB* db)
-	{
-		delete db;
-	}
-
-
-	CoffDB* GatherDatabase(LibFile* libFile, const LibDB* libDb, const ImmutableString& objPath)
-	{
-		const size_t count = libDb->libEntries.size();
-		for (size_t i = 0u; i < count; ++i)
-		{
-			const LibEntry& entry = libDb->libEntries[i];
-			if (string::Matches(entry.objPath.c_str(), objPath.c_str()))
-			{
-				// found the COFF file we need in the library
-				file::MemoryFile memoryFile =
-				{
-					libFile->memoryFile->file,
-					libFile->memoryFile->memoryMappedFile,
-					pointer::Offset<void*>(libFile->memoryFile->base, entry.offset)
-				};
-
-				ObjFile objFile = { objPath, &memoryFile };
-				// BEGIN EPIC MOD - 
-				uint32_t uniqueId = symbols::GetCompilandIdFromPath(string::ToWideString(objPath));
-				return GatherDatabase(&objFile, uniqueId, ReadFlags::NONE);
-				// END EPIC MOD
-			}
-		}
-
-		LC_LOG_DEV("Cannot find COFF %s in archive %s", objPath.c_str(), libFile->filename.c_str());
-		return nullptr;
-	}
-
-
 	template <coffDetail::CoffType::Enum Type>
-	static UnresolvedSymbolDB* GatherUnresolvedSymbolDatabase(ObjFile* file, uint32_t uniqueId, coff::ReadFlags::Enum readFlags)
+	static UnresolvedSymbolDB* GatherUnresolvedSymbolDatabase(ObjFile* objFile, uint32_t uniqueId)
 	{
 		using CoffHeader = typename coffDetail::CoffType::TypeMap<Type>::HeaderType;
 		using CoffSymbol = typename coffDetail::CoffType::TypeMap<Type>::SymbolType;
-		using CoffAuxSymbol = typename coffDetail::CoffType::TypeMap<Type>::AuxSymbolType;
 
 		UnresolvedSymbolDB* symbolDb = new UnresolvedSymbolDB;
 
-		const uint32_t symbolCount = coffDetail::GetNumberOfSymbols<CoffHeader>(file->memoryFile->base);
-		const void* symbolTable = coffDetail::GetSymbolTable<CoffHeader>(file->memoryFile->base);
-		const char* stringTable = coffDetail::GetStringTable<CoffSymbol>(symbolTable, symbolCount);
+		const void* objFileBase = Filesystem::GetMemoryMappedFileData(objFile->memoryFile);
 
-		const HexUniqueId hexUniqueId(uniqueId);
+		const uint32_t symbolCount = coffDetail::GetNumberOfSymbols<CoffHeader>(objFileBase);
+		const void* symbolTable = coffDetail::GetSymbolTable<CoffHeader>(objFileBase);
+		const char* stringTable = coffDetail::GetStringTable<CoffSymbol>(symbolTable, symbolCount);
 
 		// pre-allocate data structures
 		symbolDb->symbols.reserve(symbolCount);
@@ -1777,7 +1518,7 @@ namespace coff
 			}
 			else if (coffDetail::IsUndefinedSymbol(symbol))
 			{
-				ImmutableString name = GetSymbolName(stringTable, symbol, uniqueId, hexUniqueId, 0u, readFlags);
+				ImmutableString name = GetSymbolName(stringTable, symbol, uniqueId, 0u);
 				symbolDb->symbols.emplace_back(std::move(name));
 				symbolDb->symbolIndex.emplace_back(i);
 			}
@@ -1794,16 +1535,16 @@ namespace coff
 	}
 
 
-	UnresolvedSymbolDB* GatherUnresolvedSymbolDatabase(ObjFile* file, uint32_t uniqueId, coff::ReadFlags::Enum readFlags)
+	UnresolvedSymbolDB* GatherUnresolvedSymbolDatabase(ObjFile* file, uint32_t uniqueId)
 	{
-		const coffDetail::CoffType::Enum type = coffDetail::GetCoffType(file->memoryFile->base);
+		const coffDetail::CoffType::Enum type = coffDetail::GetCoffType(Filesystem::GetMemoryMappedFileData(file->memoryFile));
 		if (type == coffDetail::CoffType::COFF)
 		{
-			return GatherUnresolvedSymbolDatabase<coffDetail::CoffType::COFF>(file, uniqueId, readFlags);
+			return GatherUnresolvedSymbolDatabase<coffDetail::CoffType::COFF>(file, uniqueId);
 		}
 		else if (type == coffDetail::CoffType::BIGOBJ)
 		{
-			return GatherUnresolvedSymbolDatabase<coffDetail::CoffType::BIGOBJ>(file, uniqueId, readFlags);
+			return GatherUnresolvedSymbolDatabase<coffDetail::CoffType::BIGOBJ>(file, uniqueId);
 		}
 
 		return nullptr;
@@ -1817,19 +1558,18 @@ namespace coff
 
 
 	template <coffDetail::CoffType::Enum Type>
-	static ExternalSymbolDB* GatherExternalSymbolDatabase(ObjFile* file, uint32_t uniqueId, coff::ReadFlags::Enum readFlags)
+	static ExternalSymbolDB* GatherExternalSymbolDatabase(ObjFile* objFile, uint32_t uniqueId)
 	{
 		using CoffHeader = typename coffDetail::CoffType::TypeMap<Type>::HeaderType;
 		using CoffSymbol = typename coffDetail::CoffType::TypeMap<Type>::SymbolType;
-		using CoffAuxSymbol = typename coffDetail::CoffType::TypeMap<Type>::AuxSymbolType;
 
 		ExternalSymbolDB* symbolDb = new ExternalSymbolDB;
 
-		const uint32_t symbolCount = coffDetail::GetNumberOfSymbols<CoffHeader>(file->memoryFile->base);
-		const void* symbolTable = coffDetail::GetSymbolTable<CoffHeader>(file->memoryFile->base);
-		const char* stringTable = coffDetail::GetStringTable<CoffSymbol>(symbolTable, symbolCount);
+		const void* objFileBase = Filesystem::GetMemoryMappedFileData(objFile->memoryFile);
 
-		const HexUniqueId hexUniqueId(uniqueId);
+		const uint32_t symbolCount = coffDetail::GetNumberOfSymbols<CoffHeader>(objFileBase);
+		const void* symbolTable = coffDetail::GetSymbolTable<CoffHeader>(objFileBase);
+		const char* stringTable = coffDetail::GetStringTable<CoffSymbol>(symbolTable, symbolCount);
 
 		// pre-allocate data structures
 		symbolDb->symbols.reserve(symbolCount);
@@ -1858,7 +1598,7 @@ namespace coff
 				if ((type == SymbolType::EXTERNAL_DATA) || (type == SymbolType::EXTERNAL_FUNCTION))
 				{
 					// this is an external symbol stored in the COFF
-					ImmutableString name = GetSymbolName(stringTable, symbol, uniqueId, hexUniqueId, 0u, readFlags);
+					ImmutableString name = GetSymbolName(stringTable, symbol, uniqueId, 0u);
 					symbolDb->symbols.emplace_back(std::move(name));
 					symbolDb->types.emplace_back(type);
 				}
@@ -1873,16 +1613,16 @@ namespace coff
 	}
 
 
-	ExternalSymbolDB* GatherExternalSymbolDatabase(ObjFile* file, uint32_t uniqueId, coff::ReadFlags::Enum readFlags)
+	ExternalSymbolDB* GatherExternalSymbolDatabase(ObjFile* file, uint32_t uniqueId)
 	{
-		const coffDetail::CoffType::Enum type = coffDetail::GetCoffType(file->memoryFile->base);
+		const coffDetail::CoffType::Enum type = coffDetail::GetCoffType(Filesystem::GetMemoryMappedFileData(file->memoryFile));
 		if (type == coffDetail::CoffType::COFF)
 		{
-			return GatherExternalSymbolDatabase<coffDetail::CoffType::COFF>(file, uniqueId, readFlags);
+			return GatherExternalSymbolDatabase<coffDetail::CoffType::COFF>(file, uniqueId);
 		}
 		else if (type == coffDetail::CoffType::BIGOBJ)
 		{
-			return GatherExternalSymbolDatabase<coffDetail::CoffType::BIGOBJ>(file, uniqueId, readFlags);
+			return GatherExternalSymbolDatabase<coffDetail::CoffType::BIGOBJ>(file, uniqueId);
 		}
 
 		return nullptr;
@@ -2012,7 +1752,7 @@ namespace coff
 
 	const void* GetBaseAddress(const ObjFile* file)
 	{
-		return file->memoryFile->base;
+		return Filesystem::GetMemoryMappedFileData(file->memoryFile);
 	}
 
 
@@ -2080,4 +1820,10 @@ namespace coff
 
 		return true;
 	}
+
+
+	bool IsMSVCJustMyCodeSection(const char* sectionName)
+	{
+		return string::Matches(sectionName, ".msvcjmc");
+	}
 }
diff --git a/Engine/Source/Developer/Windows/LiveCodingServer/Private/External/LC_Coff.h b/Engine/Source/Developer/Windows/LiveCodingServer/Private/External/LC_Coff.h
--- a/Engine/Source/Developer/Windows/LiveCodingServer/Private/External/LC_Coff.h
+++ b/Engine/Source/Developer/Windows/LiveCodingServer/Private/External/LC_Coff.h
@@ -1,38 +1,35 @@
-// Copyright 2011-2019 Molecular Matters GmbH, all rights reserved.
+// Copyright 2011-2020 Molecular Matters GmbH, all rights reserved.
 
 #pragma once
 
+// BEGIN EPIC MOD
 #include "CoreTypes.h"
+// END EPIC MOD
 #include "LC_ImmutableString.h"
-#include "LC_MemoryFile.h"
+#include "LC_MemoryMappedFile.h"
+// BEGIN EPIC MOD
 #include "LC_Types.h"
 #include "LC_Logging.h"
 #include "Windows/WindowsHWrapper.h"
 
 #undef RELATIVE
 #undef OPTIONAL
+// END EPIC MOD
 
 struct PoolAllocatorSingleThreadPolicy;
 template <typename T> class PoolAllocator;
 
+
 namespace coff
 {
 	struct ObjFile
 	{
 		ImmutableString filename;
-		file::MemoryFile* memoryFile;
+		Filesystem::MemoryMappedFile* memoryFile;
 
 		LC_DISABLE_ASSIGNMENT(ObjFile);
 	};
 
-	struct LibFile
-	{
-		ImmutableString filename;
-		file::MemoryFile* memoryFile;
-
-		LC_DISABLE_ASSIGNMENT(LibFile);
-	};
-
 	struct SymbolType
 	{
 		enum Enum : uint8_t
@@ -50,7 +47,7 @@ namespace coff
 	{
 		struct Type
 		{
-			enum Enum : uint16_t
+			enum Enum : uint8_t
 			{
 #if LC_64_BIT
 				/*
@@ -112,7 +109,7 @@ namespace coff
 					IMAGE_REL_I386_TOKEN		The CLR token.
 					IMAGE_REL_I386_SECREL7		A 7-bit offset from the base of the section that contains the target.
 					IMAGE_REL_I386_REL32		The 32-bit relative displacement to the target.
-												This supports the x86 relative branch and call instructions.				
+												This supports the x86 relative branch and call instructions.
 
 					This means that only a handful of relocation types need to be supported.
 				*/
@@ -122,7 +119,7 @@ namespace coff
 				VA_32 = IMAGE_REL_I386_DIR32,
 				RVA_32 = IMAGE_REL_I386_DIR32NB,
 #endif
-				UNKNOWN = 0xFFFFu
+				UNKNOWN = 0xFFu
 			};
 
 			static const char* ToString(Enum value)
@@ -195,6 +192,7 @@ namespace coff
 		Type::Enum type;						// type of the relocation
 		SymbolType::Enum srcSymbolType;			// symbol type of the source symbol, cached (does not increase struct size)
 		SymbolType::Enum dstSymbolType;			// symbol type of the destination symbol, cached (does not increase struct size)
+		bool dstIsSection;						// true if this relocation points to a section
 	};
 
 	struct Symbol
@@ -259,12 +257,6 @@ namespace coff
 		LC_DISABLE_ASSIGNMENT(LibEntry);
 	};
 
-	struct LibDB
-	{
-		types::vector<ImmutableString> exportedSymbols;		// all symbols exported by the library, alphabetically sorted
-		types::vector<LibEntry> libEntries;
-	};
-
 	struct UnresolvedSymbolDB
 	{
 		types::vector<ImmutableString> symbols;
@@ -351,16 +343,6 @@ namespace coff
 	};
 	
 
-	struct ReadFlags
-	{
-		enum Enum
-		{
-			NONE = 0,
-			GENERATE_ANS_NAME_FROM_UNIQUE_ID = 1
-		};
-	};
-
-
 	struct SymbolRemovalStrategy
 	{
 		enum Enum : SHORT
@@ -375,12 +357,9 @@ namespace coff
 	ObjFile* OpenObj(const wchar_t* filename);
 	void CloseObj(ObjFile*& file);
 
-	LibFile* OpenLib(const wchar_t* filename);
-	void CloseLib(LibFile*& file);
-
 
 	// UPDATE
-	RawCoff* ReadRaw(ObjFile* file, uint32_t uniqueId, ReadFlags::Enum readFlags);
+	RawCoff* ReadRaw(ObjFile* file, uint32_t uniqueId);
 	void WriteRaw(const wchar_t* filename, const RawCoff* rawCoff, SymbolRemovalStrategy::Enum removalStrategy);
 	void DestroyRaw(RawCoff* rawCoff);
 
@@ -391,6 +370,7 @@ namespace coff
 	size_t GetAuxSymbolCount(const RawCoff* rawCoff, size_t symbolIndex);
 	SymbolType::Enum GetSymbolType(const RawCoff* rawCoff, size_t index);
 	const ImmutableString& GetSymbolName(const RawCoff* rawCoff, size_t index);
+	ImmutableString GetSectionName(const RawCoff* rawCoff, size_t index);
 	uint32_t GetSymbolSectionIndex(const RawCoff* rawCoff, size_t index);
 
 	bool IsAbsoluteSymbol(const RawCoff* rawCoff, size_t index);
@@ -428,24 +408,18 @@ namespace coff
 
 
 	// unique ID must uniquely identify this ObjFile. each obj file with a unique name must have a unique ID
-	UnresolvedSymbolDB* GatherUnresolvedSymbolDatabase(ObjFile* file, uint32_t uniqueId, coff::ReadFlags::Enum readFlags);
+	UnresolvedSymbolDB* GatherUnresolvedSymbolDatabase(ObjFile* file, uint32_t uniqueId);
 	void DestroyDatabase(UnresolvedSymbolDB* db);
 
 	// unique ID must uniquely identify this ObjFile. each obj file with a unique name must have a unique ID
-	ExternalSymbolDB* GatherExternalSymbolDatabase(ObjFile* file, uint32_t uniqueId, coff::ReadFlags::Enum readFlags);
+	ExternalSymbolDB* GatherExternalSymbolDatabase(ObjFile* file, uint32_t uniqueId);
 	void DestroyDatabase(ExternalSymbolDB* db);
 
 
 	// unique ID must uniquely identify this ObjFile. each obj file with a unique name must have a unique ID
-	CoffDB* GatherDatabase(ObjFile* file, uint32_t uniqueId, ReadFlags::Enum readFlags);
+	CoffDB* GatherDatabase(ObjFile* file, uint32_t uniqueId);
 	void DestroyDatabase(CoffDB* db);
 
-	LibDB* GatherDatabase(LibFile* file);
-	void DestroyDatabase(LibDB* db);
-
-	// loads the COFF database from an .obj contained in a .lib
-	CoffDB* GatherDatabase(LibFile* file, const LibDB* libDb, const ImmutableString& objPath);
-
 
 	size_t GetIndexCount(const CoffDB* coffDb);
 	const Symbol* GetSymbolByIndex(const CoffDB* coffDb, size_t index);
@@ -484,4 +458,5 @@ namespace coff
 
 
 	bool IsInterestingSymbol(const ImmutableString& name);
+	bool IsMSVCJustMyCodeSection(const char* sectionName);
 }
diff --git a/Engine/Source/Developer/Windows/LiveCodingServer/Private/External/LC_CoffCache.cpp b/Engine/Source/Developer/Windows/LiveCodingServer/Private/External/LC_CoffCache.cpp
--- a/Engine/Source/Developer/Windows/LiveCodingServer/Private/External/LC_CoffCache.cpp
+++ b/Engine/Source/Developer/Windows/LiveCodingServer/Private/External/LC_CoffCache.cpp
@@ -1,8 +1,27 @@
-// Copyright 2011-2019 Molecular Matters GmbH, all rights reserved.
+// Copyright 2011-2020 Molecular Matters GmbH, all rights reserved.
 
+// BEGIN EPIC MOD
+//#include PCH_INCLUDE
+// END EPIC MOD
 #include "LC_CoffCache.h"
+#include "LC_Coff.h"
 
 
+template <typename T>
+void Destroy(T*&) {}
+
+template <>
+void Destroy(coff::CoffDB*& database)
+{
+	coff::DestroyDatabase(database);
+}
+
+template <>
+void Destroy(coff::ExternalSymbolDB*& database)
+{
+	coff::DestroyDatabase(database);
+}
+
 template <typename T>
 CoffCache<T>::CoffCache(void)
 {
@@ -16,7 +35,7 @@ CoffCache<T>::~CoffCache(void)
 	for (auto it = m_cache.begin(); it != m_cache.end(); ++it)
 	{
 		T* database = it->second;
-		coff::DestroyDatabase(database);
+		Destroy(database);
 	}
 }
 
@@ -28,13 +47,13 @@ void CoffCache<T>::Update(const ImmutableString& coffIdentifier, T* database)
 
 	// try to insert the element into the cache. if it exists, delete the old entry.
 	// if it doesn't exist, store the new element.
-	const std::pair<Cache::iterator, bool>& insertPair = m_cache.emplace(coffIdentifier, nullptr);
+	const std::pair<typename Cache::iterator, bool>& insertPair = m_cache.emplace(coffIdentifier, nullptr);
 	T*& data = insertPair.first->second;
 
 	if (!insertPair.second)
 	{
 		// value exists already, delete the old entry
-		coff::DestroyDatabase(data);
+		Destroy(data);
 	}
 
 	data = database;
diff --git a/Engine/Source/Developer/Windows/LiveCodingServer/Private/External/LC_CoffDetail.h b/Engine/Source/Developer/Windows/LiveCodingServer/Private/External/LC_CoffDetail.h
--- a/Engine/Source/Developer/Windows/LiveCodingServer/Private/External/LC_CoffDetail.h
+++ b/Engine/Source/Developer/Windows/LiveCodingServer/Private/External/LC_CoffDetail.h
@@ -1,12 +1,16 @@
-// Copyright 2011-2019 Molecular Matters GmbH, all rights reserved.
+// Copyright 2011-2020 Molecular Matters GmbH, all rights reserved.
 
 #pragma once
 
+// BEGIN EPIC MOD
 #include "CoreTypes.h"
+// END EPIC MOD
 #include "LC_PointerUtil.h"
+// BEGIN EPIC MOD
 #include "LC_Platform.h"
 #include <inttypes.h>
 #include "Windows/WindowsHWrapper.h"
+// END EPIC MOD
 
 namespace coffDetail
 {
@@ -57,7 +61,7 @@ namespace coffDetail
 		// stated in the spec.
 		// e.g. some non-bigobj COFF files that have more than 32767 sections will have symbols with
 		// section numbers like 0x8000, 0x8001, etc. which need to be treated as unsigned values.
-		typedef std::make_unsigned<decltype(symbol->SectionNumber)>::type UnsignedSectionNumberType;
+		typedef typename std::make_unsigned<decltype(symbol->SectionNumber)>::type UnsignedSectionNumberType;
 
 		const uint32_t sectionIndex = static_cast<uint32_t>(static_cast<UnsignedSectionNumberType>(symbol->SectionNumber) - 1u);
 		return sectionIndex;
