Integration of Live++ 1.6.3

diff --git a/Engine/Source/Developer/Windows/LiveCodingServer/Private/External/LC_Compiler.cpp b/Engine/Source/Developer/Windows/LiveCodingServer/Private/External/LC_Compiler.cpp
--- a/Engine/Source/Developer/Windows/LiveCodingServer/Private/External/LC_Compiler.cpp
+++ b/Engine/Source/Developer/Windows/LiveCodingServer/Private/External/LC_Compiler.cpp
@@ -1,16 +1,22 @@
-// Copyright 2011-2019 Molecular Matters GmbH, all rights reserved.
+// Copyright 2011-2020 Molecular Matters GmbH, all rights reserved.
 
+// BEGIN EPIC MOD
+//#include PCH_INCLUDE
+// END EPIC MOD
 #include "LC_Compiler.h"
 #include "LC_StringUtil.h"
-#include "LC_FileUtil.h"
+#include "LC_Filesystem.h"
 #include "LC_Process.h"
 #include "LC_Environment.h"
 #include "LC_CriticalSection.h"
-#include "LC_Logging.h"
 #include "LC_TimeStamp.h"
-#include "LC_MemoryFile.h"
+#include "LC_MemoryMappedFile.h"
 #include "LC_StringUtil.h"
-
+#include "LC_Thread.h"
+// BEGIN EPIC MOD
+#include "LC_Logging.h"
+#include "LC_AppSettings.h"
+// END EPIC MOD
 
 namespace
 {
@@ -27,18 +33,18 @@ namespace
 		{
 			for (auto it = m_cache.begin(); it != m_cache.end(); ++it)
 			{
-				process::Environment* env = it->second;
-				process::DestroyEnvironment(env);
+				Process::Environment env = it->second;
+				Process::DestroyEnvironment(env);
 			}
 		}
 
-		void Insert(const wchar_t* key, process::Environment* value)
+		void Insert(const wchar_t* key, const Process::Environment& value)
 		{
 	        // BEGIN EPIC MOD - Allow passing environment block for linker
 			auto it = m_cache.find(key);
 			if (it != m_cache.end())
 			{
-				process::DestroyEnvironment(it->second);
+				Process::DestroyEnvironment(it->second);
 				it->second = value;
 				return;
 			}
@@ -47,7 +53,7 @@ namespace
 			m_cache[key] = value;
 		}
 
-		const process::Environment* Fetch(const wchar_t* key)
+		Process::Environment Fetch(const wchar_t* key)
 		{
 			const auto it = m_cache.find(key);
 			if (it != m_cache.end())
@@ -55,11 +61,11 @@ namespace
 				return it->second;
 			}
 
-			return nullptr;
+			return Process::Environment { 0u, nullptr };
 		}
 
 	private:
-		types::unordered_map<std::wstring, process::Environment*> m_cache;
+		types::unordered_map<std::wstring, Process::Environment> m_cache;
 	};
 
 	static CompilerEnvironmentCache g_compilerEnvironmentCache;
@@ -130,33 +136,54 @@ namespace
 }
 
 
-const process::Environment* compiler::CreateEnvironmentCacheEntry(const wchar_t* absolutePathToCompilerExe)
+Process::Environment compiler::CreateEnvironmentCacheEntry(const wchar_t* absolutePathToCompilerExe)
 {
 	LC_LOG_DEV("Creating environment cache entry for %S", absolutePathToCompilerExe);
 
 	// COMPILER SPECIFIC: Visual Studio. other compilers and linkers don't need vcvars*.bat to be invoked.
 	{
 		// bail out early in case this is the LLVM/clang/lld toolchain
-		const std::wstring& toolFilename = file::GetFilename(absolutePathToCompilerExe);
-		if (string::Matches(toolFilename.c_str(), L"lld.exe"))
+		const Filesystem::Path toolFilename = Filesystem::GetFilename(absolutePathToCompilerExe);
+
+		// Clang
+		if (string::Matches(toolFilename.GetString(), L"clang.exe"))
 		{
-			return nullptr;
+			return Process::Environment { 0u, nullptr };
 		}
-		else if (string::Matches(toolFilename.c_str(), L"lld-link.exe"))
+		else if (string::Matches(toolFilename.GetString(), L"clang++.exe"))
 		{
-			return nullptr;
+			return Process::Environment { 0u, nullptr };
 		}
-		else if (string::Matches(toolFilename.c_str(), L"ld.lld.exe"))
+		else if (string::Matches(toolFilename.GetString(), L"clang-cl.exe"))
+		{
+			return Process::Environment { 0u, nullptr };
+		}
+		else if (string::Matches(toolFilename.GetString(), L"clang-cpp.exe"))
+		{
+			return Process::Environment { 0u, nullptr };
+		}
+
+		// LLD
+		else if (string::Matches(toolFilename.GetString(), L"lld.exe"))
 		{
-			return nullptr;
+			return Process::Environment { 0u, nullptr };
 		}
-		else if (string::Matches(toolFilename.c_str(), L"ld64.lld.exe"))
+		else if (string::Matches(toolFilename.GetString(), L"lld-link.exe"))
 		{
-			return nullptr;
+			return Process::Environment { 0u, nullptr };
+		}
+		else if (string::Matches(toolFilename.GetString(), L"ld.lld.exe"))
+		{
+			return Process::Environment { 0u, nullptr };
+		}
+		else if (string::Matches(toolFilename.GetString(), L"ld64.lld.exe"))
+		{
+			return Process::Environment { 0u, nullptr };
 		}
 	}
 
-	const std::wstring& path = file::GetDirectory(absolutePathToCompilerExe);
+	const double timeout = static_cast<double>(appSettings::g_prewarmTimeout->GetValue()) / 1000.0;
+	const std::wstring path = Filesystem::GetDirectory(absolutePathToCompilerExe).GetString();
 
 	// get all possible paths to vcvars*.bat files and check which one is available
 	const std::vector<const wchar_t*>& relativePathsToVcvarsFile = DetermineRelativePathToVcvarsFile(absolutePathToCompilerExe);
@@ -167,8 +194,8 @@ const process::Environment* compiler::CreateEnvironmentCacheEntry(const wchar_t*
 
 		LC_LOG_DEV("Trying vcvars*.bat at %S", pathToVcvars.c_str());
 
-		const file::Attributes& attributes = file::GetAttributes(pathToVcvars.c_str());
-		if (file::DoesExist(attributes))
+		const Filesystem::PathAttributes& attributes = Filesystem::GetAttributes(pathToVcvars.c_str());
+		if (Filesystem::DoesExist(attributes))
 		{
 			// this is the correct vcvars*.bat.
 			// we need to invoke the command shell, run the .bat file, and extract the process' environment to cache it for later use.
@@ -186,25 +213,25 @@ const process::Environment* compiler::CreateEnvironmentCacheEntry(const wchar_t*
 			// we can retrieve this from the environment later and check if there was an error.
 			commandLine += L"\" & call set LPP_TOOLCHAIN_EXIT_CODE=%^ERRORLEVEL% & call pause \"";
 
-			process::Context* vcvarsProcess = process::Spawn(cmdPath.c_str(), nullptr, commandLine.c_str(), nullptr, process::SpawnFlags::NO_WINDOW);
+			Process::Context* vcvarsProcess = Process::Spawn(cmdPath.c_str(), nullptr, commandLine.c_str(), nullptr, Process::SpawnFlags::NO_WINDOW);
 
 			// wait until LPP_TOOLCHAIN_EXIT_CODE shows up in the environment of the process.
 			// busy waiting like this is not very nice, but happens only once or twice during startup, and is called from a separate thread anyway.
 			const uint64_t startTimestamp = timeStamp::Get();
 			bool shownWarning = false;
 
-			process::Environment* environment = nullptr;
+			Process::Environment environment = {};
 			const wchar_t* toolchainExitCodeStr = nullptr;
 			for (;;)
 			{
 				// grab the environment from the process
-				process::Suspend(vcvarsProcess->pi.hProcess);
-				environment = process::CreateEnvironment(vcvarsProcess->pi.hProcess);
-				process::Resume(vcvarsProcess->pi.hProcess);
+				Process::Suspend(vcvarsProcess);
+				environment = Process::CreateEnvironment(vcvarsProcess);
+				Process::Resume(vcvarsProcess);
 
-				if (environment)
+				if (environment.data)
 				{
-					toolchainExitCodeStr = string::Find(static_cast<const wchar_t*>(environment->data), environment->size / sizeof(wchar_t), L"LPP_TOOLCHAIN_EXIT_CODE", wcslen(L"LPP_TOOLCHAIN_EXIT_CODE"));
+					toolchainExitCodeStr = string::Find(static_cast<const wchar_t*>(environment.data), environment.size / sizeof(wchar_t), L"LPP_TOOLCHAIN_EXIT_CODE", wcslen(L"LPP_TOOLCHAIN_EXIT_CODE"));
 					if (toolchainExitCodeStr)
 					{
 						const wchar_t* exitCodeStr = toolchainExitCodeStr + wcslen(L"LPP_TOOLCHAIN_EXIT_CODE") + 1u;
@@ -217,8 +244,8 @@ const process::Environment* compiler::CreateEnvironmentCacheEntry(const wchar_t*
 				}
 
 				// the batch file hasn't finished running yet, wait a bit
-				process::DestroyEnvironment(environment);
-				thread::Sleep(20u);
+				Process::DestroyEnvironment(environment);
+				Thread::Current::SleepMilliSeconds(20u);
 
 				// show a warning in case this takes longer than 5 seconds.
 				// this can happen for some users:
@@ -230,11 +257,11 @@ const process::Environment* compiler::CreateEnvironmentCacheEntry(const wchar_t*
 					shownWarning = true;
 				}
 
-				// safety net: bail out if this takes longer than 10 seconds
-				if (timeStamp::ToSeconds(delta) >= 10.0)
+				// safety net: bail out if this takes longer than the configured timeout
+				if (timeStamp::ToSeconds(delta) >= timeout)
 				{
 					LC_WARNING_USER("Prewarming compiler/linker environment for %S took too long and was aborted.", pathToVcvars.c_str());
-					return nullptr;
+					return Process::Environment { 0u, nullptr };
 				}
 			}
 
@@ -253,8 +280,8 @@ const process::Environment* compiler::CreateEnvironmentCacheEntry(const wchar_t*
 				}
 			}
 
-			process::Terminate(vcvarsProcess->pi.hProcess);
-			process::Destroy(vcvarsProcess);
+			Process::Terminate(vcvarsProcess);
+			Process::Destroy(vcvarsProcess);
 
 			return environment;
 		}
@@ -265,21 +292,21 @@ const process::Environment* compiler::CreateEnvironmentCacheEntry(const wchar_t*
 	}
 
 	LC_WARNING_USER("Cannot determine vcvars*.bat environment for compiler/linker %S", absolutePathToCompilerExe);
-	return nullptr;
+	return Process::Environment { 0u, nullptr };
 }
 
 
-const process::Environment* compiler::GetEnvironmentFromCache(const wchar_t* absolutePathToCompilerExe)
+Process::Environment compiler::GetEnvironmentFromCache(const wchar_t* absolutePathToCompilerExe)
 {
 	CriticalSection::ScopedLock lock(&g_compilerCacheCS);
 	return g_compilerEnvironmentCache.Fetch(absolutePathToCompilerExe);
 }
 
 
-const process::Environment* compiler::UpdateEnvironmentCache(const wchar_t* absolutePathToCompilerExe)
+Process::Environment compiler::UpdateEnvironmentCache(const wchar_t* absolutePathToCompilerExe)
 {
-	const process::Environment* environment = GetEnvironmentFromCache(absolutePathToCompilerExe);
-	if (environment)
+	Process::Environment environment = GetEnvironmentFromCache(absolutePathToCompilerExe);
+	if (environment.data)
 	{
 		return environment;
 	}
@@ -287,10 +314,60 @@ const process::Environment* compiler::UpdateEnvironmentCache(const wchar_t* abso
 	return CreateEnvironmentCacheEntry(absolutePathToCompilerExe);
 }
 
+
+compiler::CompilerType::Enum compiler::DetermineCompilerType(const wchar_t* compilerPath)
+{
+	// MSVC's cl.exe is much more common, so treat this as our default
+	const std::wstring lowerCaseCompilerPath = string::ToLower(Filesystem::GetFilename(compilerPath).GetString());
+	if (string::Contains(lowerCaseCompilerPath.c_str(), L"clang"))
+	{
+		return CompilerType::CLANG;
+	}
+	else if (string::Contains(lowerCaseCompilerPath.c_str(), L"clang++"))
+	{
+		return CompilerType::CLANG;
+	}
+	else if (string::Contains(lowerCaseCompilerPath.c_str(), L"clang-cl"))
+	{
+		return CompilerType::CLANG;
+	}
+	else if (string::Contains(lowerCaseCompilerPath.c_str(), L"clang-cpp"))
+	{
+		return CompilerType::CLANG;
+	}
+
+	return CompilerType::CL;
+}
+
+
+compiler::LinkerType::Enum compiler::DetermineLinkerType(const wchar_t* linkerPath)
+{
+	// MSVC's link.exe is much more common, so treat this as our default
+	const std::wstring lowerCaseLinkerPath = string::ToLower(Filesystem::GetFilename(linkerPath).GetString());
+	if (string::Contains(lowerCaseLinkerPath.c_str(), L"lld"))
+	{
+		return LinkerType::LLD;
+	}
+	else if (string::Contains(lowerCaseLinkerPath.c_str(), L"lld-link"))
+	{
+		return LinkerType::LLD;
+	}
+	else if (string::Contains(lowerCaseLinkerPath.c_str(), L"ld.lld"))
+	{
+		return LinkerType::LLD;
+	}
+	else if (string::Contains(lowerCaseLinkerPath.c_str(), L"ld64.lld"))
+	{
+		return LinkerType::LLD;
+	}
+
+	return LinkerType::LINK;
+}
+
 // BEGIN EPIC MOD - Allow passing environment block for linker
-void compiler::AddEnvironmentToCache(const wchar_t* absolutePathToCompilerExe, process::Environment* environment)
+void compiler::AddEnvironmentToCache(const wchar_t* absolutePathToCompilerExe, Process::Environment* environment)
 {
 	CriticalSection::ScopedLock lock(&g_compilerCacheCS);
-	g_compilerEnvironmentCache.Insert(absolutePathToCompilerExe, environment);
+	g_compilerEnvironmentCache.Insert(absolutePathToCompilerExe, *environment);
 }
-// END EPIC MOD
+// END EPIC MOD
\ No newline at end of file
diff --git a/Engine/Source/Developer/Windows/LiveCodingServer/Private/External/LC_Compiler.h b/Engine/Source/Developer/Windows/LiveCodingServer/Private/External/LC_Compiler.h
--- a/Engine/Source/Developer/Windows/LiveCodingServer/Private/External/LC_Compiler.h
+++ b/Engine/Source/Developer/Windows/LiveCodingServer/Private/External/LC_Compiler.h
@@ -1,29 +1,50 @@
-// Copyright 2011-2019 Molecular Matters GmbH, all rights reserved.
+// Copyright 2011-2020 Molecular Matters GmbH, all rights reserved.
 
 #pragma once
 
+// BEGIN EPIC MOD
 #include "CoreTypes.h"
-
-namespace process
-{
-	struct Environment;
-}
+// END EPIC MOD
+#include "LC_ProcessTypes.h"
 
 
+// TODO: rename to Toolchain or similar
 namespace compiler
 {
 	// creates a new entry in the cache for the given compiler .exe, and returns it
-	const process::Environment* CreateEnvironmentCacheEntry(const wchar_t* absolutePathToCompilerExe);
+	Process::Environment CreateEnvironmentCacheEntry(const wchar_t* absolutePathToCompilerExe);
 
 	// gets the environment for a given compiler .exe from the cache.
 	// returns nullptr if the environment is not yet in the cache.
-	const process::Environment* GetEnvironmentFromCache(const wchar_t* absolutePathToCompilerExe);
+	Process::Environment GetEnvironmentFromCache(const wchar_t* absolutePathToCompilerExe);
 
 	// helper function that either creates a new entry in the cache if none exists yet,
 	// or returns the one found in the cache.
-	const process::Environment* UpdateEnvironmentCache(const wchar_t* absolutePathToCompilerExe);
-	
+	Process::Environment UpdateEnvironmentCache(const wchar_t* absolutePathToCompilerExe);
+
 	// BEGIN EPIC MOD - Allow overriding environment block for tools
-	void AddEnvironmentToCache(const wchar_t* absolutePathToCompilerExe, process::Environment* environment);
+	void AddEnvironmentToCache(const wchar_t* absolutePathToCompilerExe, Process::Environment* environment);
 	// END EPIC MOD
+
+	struct CompilerType
+	{
+		enum Enum
+		{
+			CL,			// MSVC
+			CLANG		// LLVM/Clang
+		};
+	};
+
+	struct LinkerType
+	{
+		enum Enum
+		{
+			LINK,		// MSVC
+			LLD			// LLVM/LLD
+		};
+	};
+
+	CompilerType::Enum DetermineCompilerType(const wchar_t* compilerPath);
+
+	LinkerType::Enum DetermineLinkerType(const wchar_t* linkerPath);
 }
diff --git a/Engine/Source/Developer/Windows/LiveCodingServer/Private/External/LC_DiaSymbolName.h b/Engine/Source/Developer/Windows/LiveCodingServer/Private/External/LC_DiaSymbolName.h
--- a/Engine/Source/Developer/Windows/LiveCodingServer/Private/External/LC_DiaSymbolName.h
+++ b/Engine/Source/Developer/Windows/LiveCodingServer/Private/External/LC_DiaSymbolName.h
@@ -1,11 +1,14 @@
-// Copyright 2011-2019 Molecular Matters GmbH, all rights reserved.
+// Copyright 2011-2020 Molecular Matters GmbH, all rights reserved.
 
 #pragma once
 
+// BEGIN EPIC MOD
 #include "CoreTypes.h"
 #include "LC_Platform.h"
+#include LC_PLATFORM_INCLUDE(LC_Foundation)
 #include "Windows/WindowsHWrapper.h"
 #include <dia2.h>
+// END EPIC MOD
 
 namespace dia
 {
diff --git a/Engine/Source/Developer/Windows/LiveCodingServer/Private/External/LC_DiaUtil.cpp b/Engine/Source/Developer/Windows/LiveCodingServer/Private/External/LC_DiaUtil.cpp
--- a/Engine/Source/Developer/Windows/LiveCodingServer/Private/External/LC_DiaUtil.cpp
+++ b/Engine/Source/Developer/Windows/LiveCodingServer/Private/External/LC_DiaUtil.cpp
@@ -1,5 +1,8 @@
-// Copyright 2011-2019 Molecular Matters GmbH, all rights reserved.
+// Copyright 2011-2020 Molecular Matters GmbH, all rights reserved.
 
+// BEGIN EPIC MOD
+//#include PCH_INCLUDE
+// END EPIC MOD
 #include "LC_DiaUtil.h"
 
 
@@ -138,10 +141,14 @@ uint32_t dia::GetSymbolOffset(IDiaSymbol* symbol)
 
 bool dia::IsFunction(IDiaSymbol* symbol)
 {
+	// BEGIN EPIC MOD
 	BOOL isFunction = Windows::FALSE;
+	// END EPIC MOD
 	symbol->get_function(&isFunction);
 
+	// BEGIN EPIC MOD
 	return (isFunction != Windows::FALSE);
+	// END EPIC MOD
 }
 
 
@@ -174,7 +181,9 @@ IDiaSymbol* dia::GetSymbolById(IDiaSession* session, uint32_t id)
 
 bool dia::WasCompiledWithLTCG(IDiaSymbol* compilandDetail)
 {
+	// BEGIN EPIC MOD
 	BOOL isLTCG = Windows::FALSE;
+	// END EPIC MOD
 	compilandDetail->get_isLTCG(&isLTCG);
 
 	return (isLTCG != 0);
@@ -183,7 +192,9 @@ bool dia::WasCompiledWithLTCG(IDiaSymbol* compilandDetail)
 
 bool dia::WasCompiledWithHotpatch(IDiaSymbol* compilandDetail)
 {
+	// BEGIN EPIC MOD
 	BOOL isHotpatch = Windows::FALSE;
+	// END EPIC MOD
 	compilandDetail->get_isHotpatchable(&isHotpatch);
 
 	return (isHotpatch != 0);
diff --git a/Engine/Source/Developer/Windows/LiveCodingServer/Private/External/LC_DiaVariant.cpp b/Engine/Source/Developer/Windows/LiveCodingServer/Private/External/LC_DiaVariant.cpp
--- a/Engine/Source/Developer/Windows/LiveCodingServer/Private/External/LC_DiaVariant.cpp
+++ b/Engine/Source/Developer/Windows/LiveCodingServer/Private/External/LC_DiaVariant.cpp
@@ -1,7 +1,12 @@
-// Copyright 2011-2019 Molecular Matters GmbH, all rights reserved.
+// Copyright 2011-2020 Molecular Matters GmbH, all rights reserved.
 
+// BEGIN EPIC MOD
+//#include PCH_INCLUDE
+// END EPIC MOD
 #include "LC_DiaVariant.h"
+// BEGIN EPIC MOD
 #include <utility>
+// END EPIC MOD
 
 dia::Variant::Variant(IDiaSymbol* symbol)
 	: m_var{ VT_EMPTY }
diff --git a/Engine/Source/Developer/Windows/LiveCodingServer/Private/External/LC_DirectoryCache.cpp b/Engine/Source/Developer/Windows/LiveCodingServer/Private/External/LC_DirectoryCache.cpp
--- a/Engine/Source/Developer/Windows/LiveCodingServer/Private/External/LC_DirectoryCache.cpp
+++ b/Engine/Source/Developer/Windows/LiveCodingServer/Private/External/LC_DirectoryCache.cpp
@@ -1,8 +1,12 @@
-// Copyright 2011-2019 Molecular Matters GmbH, all rights reserved.
+// Copyright 2011-2020 Molecular Matters GmbH, all rights reserved.
 
+// BEGIN EPIC MOD
+//#include PCH_INCLUDE
+// END EPIC MOD
 #include "LC_DirectoryCache.h"
+// BEGIN EPIC MOD
 #include "LC_Allocators.h"
-
+// END EPIC MOD
 
 DirectoryCache::DirectoryCache(size_t expectedDirectoryCount)
 {
@@ -43,7 +47,9 @@ void DirectoryCache::PrimeNotifications(void)
 	for (auto it : m_directories)
 	{
 		Directory* dir = it.second;
-		dir->hadChange = dir->changeNotification.Check(0u);
+
+		// don't throw away any changes that might have been set already
+		dir->hadChange |= dir->changeNotification.Check(0u);
 	}
 }
 
diff --git a/Engine/Source/Developer/Windows/LiveCodingServer/Private/External/LC_Disassembler.cpp b/Engine/Source/Developer/Windows/LiveCodingServer/Private/External/LC_Disassembler.cpp
--- a/Engine/Source/Developer/Windows/LiveCodingServer/Private/External/LC_Disassembler.cpp
+++ b/Engine/Source/Developer/Windows/LiveCodingServer/Private/External/LC_Disassembler.cpp
@@ -1,84 +1,73 @@
-// Copyright 2011-2019 Molecular Matters GmbH, all rights reserved.
+// Copyright 2011-2020 Molecular Matters GmbH, all rights reserved.
 
+// BEGIN EPIC MOD
+//#include PCH_INCLUDE
+// END EPIC MOD
 #include "LC_Disassembler.h"
+#include "LC_Process.h"
 #include "LC_PointerUtil.h"
+#include "distorm.h"
+#include "mnemonics.h"
+// BEGIN EPIC MOD
 #include "LC_Platform.h"
 #include "LC_Logging.h"
-#include "distorm.h"
-
+// END EPIC MOD
 
 namespace
 {
-	// an x86/x64 instruction is at most 16 bytes long
-	static const size_t LONGEST_X86_INSTRUCTION = 16u;
-}
-
-
-size_t disassembler::FindInstructionSize(process::Handle processHandle, const void* address)
-{
-	uint8_t code[LONGEST_X86_INSTRUCTION] = {};
-	process::ReadProcessMemory(processHandle, address, code, LONGEST_X86_INSTRUCTION);
-
-	_CodeInfo codeInfo = {};
-	codeInfo.code = code;
-	codeInfo.codeLen = LONGEST_X86_INSTRUCTION;
-	codeInfo.codeOffset = pointer::AsInteger<uint64_t>(address);
+	static _DInst DisassembleInstruction(const uint8_t* instructions, uint64_t address)
+	{
+		_CodeInfo codeInfo = {};
+		codeInfo.code = instructions;
+		codeInfo.codeLen = Disassembler::Stream::Instruction::MAXIMUM_LENGTH;
+		codeInfo.codeOffset = address;
 
 #if LC_64_BIT
-	codeInfo.dt = Decode64Bits;
+		codeInfo.dt = Decode64Bits;
 #else
-	codeInfo.dt = Decode32Bits;
+		codeInfo.dt = Decode32Bits;
 #endif
 
-	_DInst result = {};
-	unsigned int instructionCount = 0u;
-	distorm_decompose(&codeInfo, &result, 1u, &instructionCount);
+		_DInst result = {};
+		unsigned int instructionCount = 0u;
+		distorm_decompose(&codeInfo, &result, 1u, &instructionCount);
 
-	if (instructionCount == 0u)
-	{
-		// something went horribly wrong
-		LC_ERROR_DEV("Could not disassemble instruction at 0x%p", address);
-		return 0u;
-	}
+		if (instructionCount == 0u)
+		{
+			// something went horribly wrong
+			LC_ERROR_DEV("Could not disassemble instruction at 0x%p", address);
+			return result;
+		}
 
-	if (result.flags == FLAG_NOT_DECODABLE)
-	{
-		// the opcode could not be decoded
-		LC_ERROR_DEV("Could not decode instruction at 0x%p", address);
-		return 0u;
-	}
+		if (result.flags == FLAG_NOT_DECODABLE)
+		{
+			// the opcode could not be decoded
+			LC_ERROR_DEV("Could not decode instruction at 0x%p", address);
+			return result;
+		}
 
-	return result.size;
+		return result;
+	}
 }
 
 
-const void* disassembler::FindPreviousInstructionAddress(process::Handle processHandle, const void* functionStart, const void* instructionAddress)
+uint8_t Disassembler::FindInstructionSize(Process::Handle processHandle, const void* address)
 {
-	// starting at the function, disassemble instructions until we arrive at the given address
-	const void* currentAddress = functionStart;
-	for (;;)
-	{
-		const size_t currentSize = FindInstructionSize(processHandle, currentAddress);
-		if (currentSize == 0u)
-		{
-			// something went wrong
-			break;
-		}
+	uint8_t code[Stream::Instruction::MAXIMUM_LENGTH] = {};
+	Process::ReadProcessMemory(processHandle, address, code, Stream::Instruction::MAXIMUM_LENGTH);
 
-		const void* previousAddress = currentAddress;
-		currentAddress = pointer::Offset<const void*>(currentAddress, currentSize);
-		if (currentAddress == instructionAddress)
-		{
-			// we just decoded the instruction right before the given address
-			return previousAddress;
-		}
+	const _DInst& instruction = DisassembleInstruction(code, pointer::AsInteger<uint64_t>(address));
+	return instruction.size;
+}
 
-		if (currentAddress > instructionAddress)
-		{
-			// something went wrong, we should arrive *exactly* at the given address
-			break;
-		}
-	}
 
-	return nullptr;
+bool Disassembler::IsMOV(uint16_t opcode)
+{
+	return (opcode == I_MOV);
+}
+
+
+bool Disassembler::IsLEA(uint16_t opcode)
+{
+	return (opcode == I_LEA);
 }
diff --git a/Engine/Source/Developer/Windows/LiveCodingServer/Private/External/LC_Disassembler.h b/Engine/Source/Developer/Windows/LiveCodingServer/Private/External/LC_Disassembler.h
--- a/Engine/Source/Developer/Windows/LiveCodingServer/Private/External/LC_Disassembler.h
+++ b/Engine/Source/Developer/Windows/LiveCodingServer/Private/External/LC_Disassembler.h
@@ -1,16 +1,39 @@
-// Copyright 2011-2019 Molecular Matters GmbH, all rights reserved.
+// Copyright 2011-2020 Molecular Matters GmbH, all rights reserved.
 
 #pragma once
 
+// BEGIN EPIC MOD
 #include "CoreTypes.h"
-#include "LC_Process.h"
+// END EPIC MOD
+#include "LC_ProcessTypes.h"
 
 
-namespace disassembler
+namespace Disassembler
 {
-	// returns the size of the first instruction found at the given address
-	size_t FindInstructionSize(process::Handle processHandle, const void* address);
+	// Returns the size of the first instruction found at the given address.
+	uint8_t FindInstructionSize(Process::Handle processHandle, const void* address);
 
-	// returns the address of the previous instruction inside a function
-	const void* FindPreviousInstructionAddress(process::Handle processHandle, const void* functionStart, const void* instructionAddress);
+	// Returns whether the given instruction opcode is a MOV of any kind.
+	bool IsMOV(uint16_t opcode);
+
+	// Returns whether the given instruction opcode is a LEA of any kind.
+	bool IsLEA(uint16_t opcode);
+
+
+	class Stream
+	{
+	public:
+		struct Instruction
+		{
+			// BEGIN EPIC MOD - removed copy non-ascii character
+			// From the Intel 64 and IA-32 Architectures Software Developer's Manual: The maximum length of an Intel 64 and IA-32 instruction remains 15 bytes.
+			// END EPIC MOD
+			static const unsigned int MAXIMUM_LENGTH = 15u;
+
+			uint8_t machineCode[MAXIMUM_LENGTH];
+			uint8_t size;
+			uint64_t address;
+			uint16_t opcode;
+		};
+	};
 }
diff --git a/Engine/Source/Developer/Windows/LiveCodingServer/Private/External/LC_ExecutablePatcher.cpp b/Engine/Source/Developer/Windows/LiveCodingServer/Private/External/LC_ExecutablePatcher.cpp
--- a/Engine/Source/Developer/Windows/LiveCodingServer/Private/External/LC_ExecutablePatcher.cpp
+++ b/Engine/Source/Developer/Windows/LiveCodingServer/Private/External/LC_ExecutablePatcher.cpp
@@ -1,10 +1,16 @@
-// Copyright 2011-2019 Molecular Matters GmbH, all rights reserved.
+// Copyright 2011-2020 Molecular Matters GmbH, all rights reserved.
 
+// BEGIN EPIC MOD
+//#include PCH_INCLUDE
+// END EPIC MOD
 #include "LC_ExecutablePatcher.h"
 #include "LC_Executable.h"
 #include "LC_PointerUtil.h"
+#include "LC_Process.h"
+// BEGIN EPIC MOD
+#include "LC_Assert.h"
 #include "LC_Logging.h"
-
+// END EPIC MOD
 
 namespace
 {
@@ -60,21 +66,21 @@ uint32_t ExecutablePatcher::DisableEntryPointInImage(executable::Image* image, e
 }
 
 
-void ExecutablePatcher::DisableEntryPoint(process::Handle processHandle, void* moduleBase, uint32_t entryPointRva)
+void ExecutablePatcher::DisableEntryPoint(Process::Handle processHandle, void* moduleBase, uint32_t entryPointRva)
 {
 	for (size_t i = 0u; i < INJECTED_CODE_SIZE; ++i)
 	{
 		uint8_t* address = pointer::Offset<uint8_t*>(moduleBase, entryPointRva + i);
-		process::WriteProcessMemory(processHandle, address, PATCH[i]);
+		Process::WriteProcessMemory(processHandle, address, PATCH[i]);
 	}
 }
 
 
-void ExecutablePatcher::RestoreEntryPoint(process::Handle processHandle, void* moduleBase, uint32_t entryPointRva)
+void ExecutablePatcher::RestoreEntryPoint(Process::Handle processHandle, void* moduleBase, uint32_t entryPointRva)
 {
 	for (size_t i = 0u; i < INJECTED_CODE_SIZE; ++i)
 	{
 		uint8_t* address = pointer::Offset<uint8_t*>(moduleBase, entryPointRva + i);
-		process::WriteProcessMemory(processHandle, address, m_originalCode[i]);
+		Process::WriteProcessMemory(processHandle, address, m_originalCode[i]);
 	}
 }
diff --git a/Engine/Source/Developer/Windows/LiveCodingServer/Private/External/LC_ExecutablePatcher.h b/Engine/Source/Developer/Windows/LiveCodingServer/Private/External/LC_ExecutablePatcher.h
--- a/Engine/Source/Developer/Windows/LiveCodingServer/Private/External/LC_ExecutablePatcher.h
+++ b/Engine/Source/Developer/Windows/LiveCodingServer/Private/External/LC_ExecutablePatcher.h
@@ -1,10 +1,12 @@
-// Copyright 2011-2019 Molecular Matters GmbH, all rights reserved.
+// Copyright 2011-2020 Molecular Matters GmbH, all rights reserved.
 
 #pragma once
 
+// BEGIN EPIC MOD
 #include "CoreTypes.h"
+// END EPIC MOD
 #include "LC_Executable.h"
-#include "LC_Process.h"
+#include "LC_ProcessTypes.h"
 
 
 class ExecutablePatcher
@@ -27,10 +29,10 @@ public:
 	uint32_t DisableEntryPointInImage(executable::Image* image, executable::ImageSectionDB* imageSections);
 
 	// disables the entry point in memory
-	void DisableEntryPoint(process::Handle processHandle, void* moduleBase, uint32_t entryPointRva);
+	void DisableEntryPoint(Process::Handle processHandle, void* moduleBase, uint32_t entryPointRva);
 
 	// restores the entry point of a loaded image that previously had its entry point disabled
-	void RestoreEntryPoint(process::Handle processHandle, void* moduleBase, uint32_t entryPointRva);
+	void RestoreEntryPoint(Process::Handle processHandle, void* moduleBase, uint32_t entryPointRva);
 
 
 	inline const uint8_t* GetEntryPointCode(void) const
diff --git a/Engine/Source/Developer/Windows/LiveCodingServer/Private/External/LC_FileAttributeCache.cpp b/Engine/Source/Developer/Windows/LiveCodingServer/Private/External/LC_FileAttributeCache.cpp
--- a/Engine/Source/Developer/Windows/LiveCodingServer/Private/External/LC_FileAttributeCache.cpp
+++ b/Engine/Source/Developer/Windows/LiveCodingServer/Private/External/LC_FileAttributeCache.cpp
@@ -1,7 +1,10 @@
-// Copyright 2011-2019 Molecular Matters GmbH, all rights reserved.
+// Copyright 2011-2020 Molecular Matters GmbH, all rights reserved.
 
+// BEGIN EPIC MOD
+//#include PCH_INCLUDE
+// END EPIC MOD
 #include "LC_FileAttributeCache.h"
-#include "LC_FileUtil.h"
+#include "LC_Filesystem.h"
 
 
 FileAttributeCache::FileAttributeCache(void)
@@ -21,9 +24,9 @@ FileAttributeCache::Data FileAttributeCache::UpdateCacheData(const std::wstring&
 	if (optional.second)
 	{
 		// value was inserted, update it with the correct data
-		const file::Attributes& attributes = file::GetAttributes(path.c_str());
-		data.exists = file::DoesExist(attributes);
-		data.lastModificationTime = file::GetLastModificationTime(attributes);
+		const Filesystem::PathAttributes& attributes = Filesystem::GetAttributes(path.c_str());
+		data.exists = Filesystem::DoesExist(attributes);
+		data.lastModificationTime = Filesystem::GetLastModificationTime(attributes);
 	}
 
 	return data;
diff --git a/Engine/Source/Developer/Windows/LiveCodingServer/Private/External/LC_FileAttributeCache.h b/Engine/Source/Developer/Windows/LiveCodingServer/Private/External/LC_FileAttributeCache.h
--- a/Engine/Source/Developer/Windows/LiveCodingServer/Private/External/LC_FileAttributeCache.h
+++ b/Engine/Source/Developer/Windows/LiveCodingServer/Private/External/LC_FileAttributeCache.h
@@ -1,10 +1,12 @@
-// Copyright 2011-2019 Molecular Matters GmbH, all rights reserved.
+// Copyright 2011-2020 Molecular Matters GmbH, all rights reserved.
 
 #pragma once
 
+// BEGIN EPIC MOD
 #include "CoreTypes.h"
 #include "LC_Types.h"
-#include "xxhash.h"
+// END EPIC MOD
+#include "LC_Hashing.h"
 
 
 class FileAttributeCache
@@ -13,7 +15,7 @@ class FileAttributeCache
 	{
 		inline size_t operator()(const std::wstring& key) const
 		{
-			return XXH32(key.c_str(), key.length() * sizeof(wchar_t), 0u);
+			return Hashing::Hash32(key.c_str(), key.length() * sizeof(wchar_t), 0u);
 		}
 	};
 
diff --git a/Engine/Source/Developer/Windows/LiveCodingServer/Private/External/LC_FunctionPatcher.cpp b/Engine/Source/Developer/Windows/LiveCodingServer/Private/External/LC_FunctionPatcher.cpp
--- a/Engine/Source/Developer/Windows/LiveCodingServer/Private/External/LC_FunctionPatcher.cpp
+++ b/Engine/Source/Developer/Windows/LiveCodingServer/Private/External/LC_FunctionPatcher.cpp
@@ -1,11 +1,17 @@
-// Copyright 2011-2019 Molecular Matters GmbH, all rights reserved.
+// Copyright 2011-2020 Molecular Matters GmbH, all rights reserved.
 
+// BEGIN EPIC MOD
+//#include PCH_INCLUDE
+// END EPIC MOD
 #include "LC_FunctionPatcher.h"
 #include "LC_Patch.h"
 #include "LC_Disassembler.h"
 #include "LC_NameMangling.h"
 #include "LC_AppSettings.h"
+#include "LC_Process.h"
+// BEGIN EPIC MOD
 #include "LC_Logging.h"
+// END EPIC MOD
 
 
 namespace
@@ -25,11 +31,11 @@ namespace
 
 	// helper function that checks whether a certain number of bytes in memory are available for patching
 	template <size_t N>
-	static inline bool AreBytesAvailableForPatching(process::Handle processHandle, const void* address)
+	static inline bool AreBytesAvailableForPatching(Process::Handle processHandle, const void* address)
 	{
 		// available bytes are either 0xCC (int 3) or 0x0 (page padding)
 		uint8_t memory[N] = {};
-		process::ReadProcessMemory(processHandle, address, memory, N);
+		Process::ReadProcessMemory(processHandle, address, memory, N);
 
 		for (size_t i = 0u; i < N; ++i)
 		{
@@ -52,14 +58,14 @@ functions::Record functions::PatchFunction
 	uint32_t patchFunctionRva,
 	const symbols::ThunkDB* thunkDb,
 	const symbols::Contribution* contribution,
-	process::Handle processHandle,
+	Process::Handle processHandle,
 	void* moduleBase,
 	uint16_t moduleIndex,
 	types::unordered_set<const void*>& patchedAddresses,
 	const types::vector<const void*>& threadIPs,
 
 	// debug only
-	unsigned int processId,
+	Process::Id processId,
 	const char* functionName
 )
 {
@@ -89,7 +95,7 @@ functions::Record functions::PatchFunction
 		// instead of the ILT.
 		for (auto thunkRva : thunkTableEntries)
 		{
-			LC_LOG_DEV("Patching ILT 0x%X of function %s at 0x%p (0x%X) (PID: %d)", thunkRva, functionName, moduleBase, functionRva, processId);
+			LC_LOG_DEV("Patching ILT 0x%X of function %s at 0x%p (0x%X) (PID: %d)", thunkRva, functionName, moduleBase, functionRva, +processId);
 
 			void* incrementalLinkingThunk = pointer::Offset<char*>(moduleBase, thunkRva);
 			patch::InstallRelativeNearJump(processHandle, incrementalLinkingThunk, patchAddress);
@@ -106,12 +112,12 @@ functions::Record functions::PatchFunction
 		// disassemble the first instructions to see how many bytes we can patch.
 		while (wholeInstructionSize < 5u)
 		{
-			const size_t instructionSize = disassembler::FindInstructionSize(processHandle, originalAddress + wholeInstructionSize);
+			const uint8_t instructionSize = Disassembler::FindInstructionSize(processHandle, originalAddress + wholeInstructionSize);
 			if (instructionSize == 0u)
 			{
 				// dump raw code in case it could not be decoded
-				LC_ERROR_DEV("Failed to disassemble code for function %s at 0x%p (0x%X) (PID: %d)", functionName, moduleBase, functionRva, processId);
-				process::DumpMemory(processHandle, originalAddress, contribution->size);
+				LC_ERROR_DEV("Failed to disassemble code for function %s at 0x%p (0x%X) (PID: %d)", functionName, moduleBase, functionRva, +processId);
+				Process::DumpMemory(processHandle, originalAddress, contribution->size);
 				break;
 			}
 
@@ -148,7 +154,7 @@ functions::Record functions::PatchFunction
 	// now install a patch using the selected technique
 	if (patchTechniqueUsed == PatchTechnique::DIRECT_RELATIVE_JUMP)
 	{
-		LC_LOG_DEV("Patching function %s directly at 0x%p (0x%X) (PID: %d)", functionName, moduleBase, functionRva, processId);
+		LC_LOG_DEV("Patching function %s directly at 0x%p (0x%X) (PID: %d)", functionName, moduleBase, functionRva, +processId);
 
 		patch::InstallRelativeNearJump(processHandle, originalAddress, patchAddress);
 
@@ -156,12 +162,12 @@ functions::Record functions::PatchFunction
 		if (wholeInstructionSize > 5u)
 		{
 			// yes, overwrite those with int 3
-			process::WriteProcessMemory(processHandle, originalAddress + 5u, INT3_PADDING, wholeInstructionSize - 5u);
+			Process::WriteProcessMemory(processHandle, originalAddress + 5u, INT3_PADDING, wholeInstructionSize - 5u);
 		}
 	}
 	else if (patchTechniqueUsed == PatchTechnique::HOTPATCH_INDIRECTION)
 	{
-		const size_t instructionSize = disassembler::FindInstructionSize(processHandle, originalAddress);
+		const uint8_t instructionSize = Disassembler::FindInstructionSize(processHandle, originalAddress);
 		if (instructionSize >= 2u)
 		{
 			// we need to go via an indirection, and install the relative jump to the patch
@@ -171,7 +177,7 @@ functions::Record functions::PatchFunction
 			const bool installedPatch = (patchedAddresses.find(originalAddress - 5u) != patchedAddresses.end());
 			if (isAvailable || installedPatch)
 			{
-				LC_LOG_DEV("Hot-patching function %s at 0x%p (0x%X) (PID: %d)", functionName, moduleBase, functionRva, processId);
+				LC_LOG_DEV("Hot-patching function %s at 0x%p (0x%X) (PID: %d)", functionName, moduleBase, functionRva, +processId);
 
 				// it is safe to install the relative jump right in front of the function
 				patch::InstallRelativeNearJump(processHandle, originalAddress - 5u, patchAddress);
@@ -179,7 +185,7 @@ functions::Record functions::PatchFunction
 				// note that in very, very rare cases, the memory region in front of the function might not
 				// be executable pages. this can only happen for the function right at the start of the
 				// code segment, but it can happen.
-				process::MakePagesExecutable(processHandle, originalAddress - 5u, 5u);
+				Process::MakePagesExecutable(processHandle, originalAddress - 5u, 5u);
 
 				// jump to the relative jump we just installed using a short jump, using 2 bytes.
 				// this memory region must always be executable already.
@@ -195,11 +201,11 @@ functions::Record functions::PatchFunction
 				{
 					if (appSettings::g_showUndecoratedNames->GetValue())
 					{
-						LC_WARNING_USER("Not enough space near function '%s' at 0x%X to install patch (PID: %d). Changes to this function will not be observable.", nameMangling::UndecorateSymbol(functionName, 0u).c_str(), originalAddress, processId);
+						LC_WARNING_USER("Not enough space near function '%s' at 0x%X to install patch (PID: %d). Changes to this function will not be observable.", nameMangling::UndecorateSymbol(functionName, 0u).c_str(), originalAddress, +processId);
 					}
 					else
 					{
-						LC_WARNING_USER("Not enough space near function '%s' at 0x%X to install patch (PID: %d). Changes to this function will not be observable.", functionName, originalAddress, processId);
+						LC_WARNING_USER("Not enough space near function '%s' at 0x%X to install patch (PID: %d). Changes to this function will not be observable.", functionName, originalAddress, +processId);
 					}
 				}
 			}
@@ -212,11 +218,11 @@ functions::Record functions::PatchFunction
 			{
 				if (appSettings::g_showUndecoratedNames->GetValue())
 				{
-					LC_WARNING_USER("Instruction in function '%s' at 0x%X is too short to install patch (PID: %d). Changes to this function will not be observable.", nameMangling::UndecorateSymbol(functionName, 0u).c_str(), originalAddress, processId);
+					LC_WARNING_USER("Instruction in function '%s' at 0x%X is too short to install patch (PID: %d). Changes to this function will not be observable.", nameMangling::UndecorateSymbol(functionName, 0u).c_str(), originalAddress, +processId);
 				}
 				else
 				{
-					LC_WARNING_USER("Instruction in function '%s' at 0x%X is too short to install patch (PID: %d). Changes to this function will not be observable.", functionName, originalAddress, processId);
+					LC_WARNING_USER("Instruction in function '%s' at 0x%X is too short to install patch (PID: %d). Changes to this function will not be observable.", functionName, originalAddress, +processId);
 				}
 			}
 		}
@@ -229,7 +235,7 @@ functions::Record functions::PatchFunction
 void functions::PatchFunction
 (
 	const Record& record,
-	process::Handle processHandle,
+	Process::Handle processHandle,
 	void* processModuleBases[],
 	void* newModuleBase,
 	types::unordered_set<const void*>& patchedAddresses,
@@ -299,12 +305,12 @@ void functions::PatchFunction
 		if (record.directJumpInstructionSize > 5u)
 		{
 			// yes, overwrite those with int 3
-			process::WriteProcessMemory(processHandle, originalAddress + 5u, INT3_PADDING, record.directJumpInstructionSize - 5u);
+			Process::WriteProcessMemory(processHandle, originalAddress + 5u, INT3_PADDING, record.directJumpInstructionSize - 5u);
 		}
 	}
 	else if (patchTechniqueUsed == PatchTechnique::HOTPATCH_INDIRECTION)
 	{
-		const size_t instructionSize = disassembler::FindInstructionSize(processHandle, originalAddress);
+		const uint8_t instructionSize = Disassembler::FindInstructionSize(processHandle, originalAddress);
 		if (instructionSize >= 2u)
 		{
 			// we need to go via an indirection, and install the relative jump to the patch
@@ -320,7 +326,7 @@ void functions::PatchFunction
 				// note that in very, very rare cases, the memory region in front of the function might not
 				// be executable pages. this can only happen for the function right at the start of the
 				// code segment, but it can happen.
-				process::MakePagesExecutable(processHandle, originalAddress - 5u, 5u);
+				Process::MakePagesExecutable(processHandle, originalAddress - 5u, 5u);
 
 				// jump to the relative jump we just installed using a short jump, using 2 bytes.
 				// this memory region must always be executable already.
@@ -340,7 +346,7 @@ functions::LibraryRecord functions::PatchLibraryFunction
 	uint32_t srcRva,
 	uint32_t destRva,
 	const symbols::Contribution* contribution,
-	process::Handle processHandle,
+	Process::Handle processHandle,
 	uint16_t moduleIndex
 )
 {
@@ -361,11 +367,11 @@ functions::LibraryRecord functions::PatchLibraryFunction
 	size_t wholeInstructionSize = 0u;
 	while (wholeInstructionSize < 5u)
 	{
-		const size_t instructionSize = disassembler::FindInstructionSize(processHandle, srcAddress + wholeInstructionSize);
+		const uint8_t instructionSize = Disassembler::FindInstructionSize(processHandle, srcAddress + wholeInstructionSize);
 		if (instructionSize == 0u)
 		{
 			// dump raw code in case it could not be decoded
-			process::DumpMemory(processHandle, srcAddress, contribution->size);
+			Process::DumpMemory(processHandle, srcAddress, contribution->size);
 			break;
 		}
 
@@ -384,7 +390,7 @@ functions::LibraryRecord functions::PatchLibraryFunction
 		if (wholeInstructionSize > 5u)
 		{
 			// yes, overwrite those with int 3
-			process::WriteProcessMemory(processHandle, srcAddress + 5u, INT3_PADDING, wholeInstructionSize - 5u);
+			Process::WriteProcessMemory(processHandle, srcAddress + 5u, INT3_PADDING, wholeInstructionSize - 5u);
 		}
 	}
 
@@ -395,7 +401,7 @@ functions::LibraryRecord functions::PatchLibraryFunction
 void functions::PatchLibraryFunction
 (
 	const LibraryRecord& record,
-	process::Handle processHandle,
+	Process::Handle processHandle,
 	void* processModuleBases[],
 	void* newModuleBase
 )
@@ -418,7 +424,7 @@ void functions::PatchLibraryFunction
 		if (record.wholeInstructionSize > 5u)
 		{
 			// yes, overwrite those with int 3
-			process::WriteProcessMemory(processHandle, srcAddress + 5u, INT3_PADDING, record.wholeInstructionSize - 5u);
+			Process::WriteProcessMemory(processHandle, srcAddress + 5u, INT3_PADDING, record.wholeInstructionSize - 5u);
 		}
 	}
 }
diff --git a/Engine/Source/Developer/Windows/LiveCodingServer/Private/External/LC_FunctionPatcher.h b/Engine/Source/Developer/Windows/LiveCodingServer/Private/External/LC_FunctionPatcher.h
--- a/Engine/Source/Developer/Windows/LiveCodingServer/Private/External/LC_FunctionPatcher.h
+++ b/Engine/Source/Developer/Windows/LiveCodingServer/Private/External/LC_FunctionPatcher.h
@@ -1,9 +1,11 @@
-// Copyright 2011-2019 Molecular Matters GmbH, all rights reserved.
+// Copyright 2011-2020 Molecular Matters GmbH, all rights reserved.
 
 #pragma once
 
+// BEGIN EPIC MOD
 #include "CoreTypes.h"
-#include "LC_Process.h"
+// END EPIC MOD
+#include "LC_ProcessTypes.h"
 #include "LC_Symbols.h"
 #include "LC_ModuleCache.h"
 
@@ -35,14 +37,14 @@ namespace functions
 		uint32_t patchFunctionRva,
 		const symbols::ThunkDB* thunkDb,
 		const symbols::Contribution* contribution,
-		process::Handle processHandle,
+		Process::Handle processHandle,
 		void* moduleBase,
 		uint16_t moduleIndex,
 		types::unordered_set<const void*>& patchedAddresses,
 		const types::vector<const void*>& threadIPs,
 
 		// debug only
-		unsigned int processId,
+		Process::Id processId,
 		const char* functionName
 	);
 
@@ -53,14 +55,14 @@ namespace functions
 		uint32_t srcRva,
 		uint32_t destRva,
 		const symbols::Contribution* contribution,
-		process::Handle processHandle,
+		Process::Handle processHandle,
 		uint16_t moduleIndex
 	);
 
 	void PatchFunction
 	(
 		const Record& record,
-		process::Handle processHandle,
+		Process::Handle processHandle,
 		void* processModuleBases[],
 		void* newModuleBase,
 		types::unordered_set<const void*>& patchedAddresses,
@@ -70,7 +72,7 @@ namespace functions
 	void PatchLibraryFunction
 	(
 		const LibraryRecord& record,
-		process::Handle processHandle,
+		Process::Handle processHandle,
 		void* processModuleBases[],
 		void* newModuleBase
 	);
diff --git a/Engine/Source/Developer/Windows/LiveCodingServer/Private/External/LC_GrowingMemoryBlock.cpp b/Engine/Source/Developer/Windows/LiveCodingServer/Private/External/LC_GrowingMemoryBlock.cpp
--- a/Engine/Source/Developer/Windows/LiveCodingServer/Private/External/LC_GrowingMemoryBlock.cpp
+++ b/Engine/Source/Developer/Windows/LiveCodingServer/Private/External/LC_GrowingMemoryBlock.cpp
@@ -1,7 +1,13 @@
-// Copyright 2011-2019 Molecular Matters GmbH, all rights reserved.
+// Copyright 2011-2020 Molecular Matters GmbH, all rights reserved.
 
+// BEGIN EPIC MOD
+//#include PCH_INCLUDE
+// END EPIC MOD
 #include "LC_GrowingMemoryBlock.h"
+// BEGIN EPIC MOD
+#include "LC_Assert.h"
 #include "LC_Logging.h"
+// END EPIC MOD
 
 
 GrowingMemoryBlock::GrowingMemoryBlock(size_t initialCapacity)
diff --git a/Engine/Source/Developer/Windows/LiveCodingServer/Private/External/LC_ImmutableString.cpp b/Engine/Source/Developer/Windows/LiveCodingServer/Private/External/LC_ImmutableString.cpp
--- a/Engine/Source/Developer/Windows/LiveCodingServer/Private/External/LC_ImmutableString.cpp
+++ b/Engine/Source/Developer/Windows/LiveCodingServer/Private/External/LC_ImmutableString.cpp
@@ -1,19 +1,24 @@
-// Copyright 2011-2019 Molecular Matters GmbH, all rights reserved.
+// Copyright 2011-2020 Molecular Matters GmbH, all rights reserved.
 
+// BEGIN EPIC MOD
+//#include PCH_INCLUDE
+// END EPIC MOD
 #include "LC_ImmutableString.h"
 #include "LC_StringUtil.h"
+#include "LC_Hashing.h"
+// BEGIN EPIC MOD
 #include "LC_Allocators.h"
+#include "LC_Assert.h"
 #include "LC_Platform.h"
 #include "LC_Logging.h"
 #include "Windows/WindowsHWrapper.h"
-#include "xxhash.h"
-
+// END EPIC MOD
 
 namespace
 {
 	static inline uint32_t Hash(const char* key, size_t length)
 	{
-		return XXH32(key, length * sizeof(char), 0u);
+		return Hashing::Hash32(key, length * sizeof(char), 0u);
 	}
 
 	static inline const char* Clone(const char* str, size_t length)
@@ -97,7 +102,7 @@ namespace string
 
 
 ImmutableString::ImmutableString(void)
-	: m_data()
+	: ImmutableString("", 0u)
 {
 }
 
@@ -192,4 +197,3 @@ void ImmutableString::Free(void)
 		LC_FREE(&g_immutableStringAllocator, const_cast<char*>(m_data.longString.str), (GetLength() + 1u) * sizeof(char));
 	}
 }
-
diff --git a/Engine/Source/Developer/Windows/LiveCodingServer/Private/External/LC_ImmutableString.h b/Engine/Source/Developer/Windows/LiveCodingServer/Private/External/LC_ImmutableString.h
--- a/Engine/Source/Developer/Windows/LiveCodingServer/Private/External/LC_ImmutableString.h
+++ b/Engine/Source/Developer/Windows/LiveCodingServer/Private/External/LC_ImmutableString.h
@@ -1,11 +1,16 @@
-// Copyright 2011-2019 Molecular Matters GmbH, all rights reserved.
+// Copyright 2011-2020 Molecular Matters GmbH, all rights reserved.
 
 #pragma once
 
+// BEGIN EPIC MOD
 #include "CoreTypes.h"
-#include "LC_Platform.h"
+// END EPIC MOD
 #include "LC_Memory.h"
+// BEGIN EPIC MOD
+#include "LC_Platform.h"
+#include LC_PLATFORM_INCLUDE(LC_Foundation)
 #include <string>
+// END EPIC MOD
 
 class ImmutableString;
 
@@ -179,13 +184,3 @@ private:
 
 	Data m_data;
 };
-
-namespace string
-{
-	ImmutableString ToUtf8String(const wchar_t* str);
-	ImmutableString ToUtf8String(const wchar_t* str, size_t count);
-	ImmutableString ToUtf8String(const std::wstring& str);
-
-	std::string ToAnsiString(const ImmutableString& utf8Str);
-	std::wstring ToWideString(const ImmutableString& utf8Str);
-}
