Integration of Live++ 1.6.3

diff --git a/Engine/Source/Developer/Windows/LiveCodingServer/Private/External/LC_LiveModule.cpp b/Engine/Source/Developer/Windows/LiveCodingServer/Private/External/LC_LiveModule.cpp
--- a/Engine/Source/Developer/Windows/LiveCodingServer/Private/External/LC_LiveModule.cpp
+++ b/Engine/Source/Developer/Windows/LiveCodingServer/Private/External/LC_LiveModule.cpp
@@ -1,14 +1,16 @@
-// Copyright 2011-2019 Molecular Matters GmbH, all rights reserved.
+// Copyright 2011-2020 Molecular Matters GmbH, all rights reserved.
 
+// BEGIN EPIC MOD
+//#include PCH_INCLUDE
+// END EPIC MOD
 #include "LC_LiveModule.h"
 #include "LC_Telemetry.h"
 #include "LC_StringUtil.h"
-#include "LC_FileUtil.h"
+#include "LC_Filesystem.h"
 #include "LC_Environment.h"
 #include "LC_SymbolReconstruction.h"
 #include "LC_Thread.h"
 #include "LC_Process.h"
-#include "LC_Thread.h"
 #include "LC_Compiler.h"
 #include "LC_SyncPoint.h"
 #include "LC_ExecutablePatcher.h"
@@ -28,12 +30,22 @@
 #include "LC_CompilerOptions.h"
 #include "LC_ModulePatch.h"
 #include "LC_Amalgamation.h"
-#include "LiveCodingServer.h"
-#include "LC_AppSettings.h"
+// BEGIN EPIC MOD
+//#include "LC_App.h"
+// END EPIC MOD
 #include "LC_Scheduler.h"
 #include "LC_LiveProcess.h"
 #include "LC_UniqueId.h"
+// BEGIN EPIC MOD
+//#include "LC_FASTBuild.h"
+// END EPIC MOD
+#include "LC_VirtualMemoryRange.h"
 #include "LPP_API.h"
+// BEGIN EPIC MOD
+#include "LiveCodingServer.h"
+#include "LC_AppSettings.h"
+#include <process.h> // needed for _PVFV
+// END EPIC MOD
 
 // BEGIN EPIC MOD - Support for UE4 debug visualizers
 #include "Misc/Paths.h"
@@ -41,6 +53,39 @@
 
 namespace
 {
+#if LC_64_BIT
+	static const void* GetLowerBoundIn4GBRange(const void* moduleBase)
+	{
+		// nothing can be loaded at address 0x0, 64 KB seems like a realistic minimum
+		const uint64_t LOWEST_ADDRESS = 64u * 1024u;
+		const uint64_t base = pointer::AsInteger<uint64_t>(moduleBase);
+
+		// make sure we don't underflow
+		if (base >= 0x80000000ull + LOWEST_ADDRESS)
+		{
+			return pointer::FromInteger<const void*>(base - 0x80000000ull);
+		}
+
+		// operation would underflow
+		return pointer::FromInteger<const void*>(LOWEST_ADDRESS);
+	}
+
+	static const void* GetUpperBoundIn4GBRange(const void* moduleBase)
+	{
+		const uint64_t base = pointer::AsInteger<uint64_t>(moduleBase);
+
+		// make sure we don't overflow
+		if (base <= 0xFFFFFFFF7FFFFFFFull)
+		{
+			return pointer::FromInteger<const void*>(base + 0x80000000ull);
+		}
+
+		// operation would overflow
+		return pointer::FromInteger<const void*>(0xFFFFFFFFFFFFFFFFull);
+	}
+#endif
+
+
 	// common linker options:
 	//	*) create x86/x64 code
 	//	*) don't echo command-line options
@@ -99,8 +144,8 @@ namespace
 			if (appSettings::g_useCompilerOverrideAsFallback->GetValue())
 			{
 				// yes, so test whether a compiler at the compiland's compiler path exists
-				const file::Attributes& attributes = file::GetAttributes(compilerPath.c_str());
-				if (file::DoesExist(attributes))
+				const Filesystem::PathAttributes& attributes = Filesystem::GetAttributes(compilerPath.c_str());
+				if (Filesystem::DoesExist(attributes))
 				{
 					// compiler exists, use it
 					return compilerPath;
@@ -137,8 +182,8 @@ namespace
 			if (appSettings::g_useLinkerOverrideAsFallback->GetValue())
 			{
 				// yes, so test whether a linker at the given path exists
-				const file::Attributes& attributes = file::GetAttributes(linkerPath.c_str());
-				if (file::DoesExist(attributes))
+				const Filesystem::PathAttributes& attributes = Filesystem::GetAttributes(linkerPath.c_str());
+				if (Filesystem::DoesExist(attributes))
 				{
 					// linker exists, use it
 					return linkerPath;
@@ -165,22 +210,10 @@ namespace
 	// helper function that determines the type of symbol removal strategy to use, depending on the linker
 	static coff::SymbolRemovalStrategy::Enum DetermineSymbolRemovalStrategy(const symbols::LinkerDB* linkerDb)
 	{
-		// MSVC's link.exe is much more common, so treat this as our default
 		const std::wstring linkerPath = GetLinkerPath(linkerDb);
-		const std::wstring lowerCaseLinkerPath = string::ToLower(file::GetFilename(linkerPath));
-		if (string::Contains(lowerCaseLinkerPath.c_str(), L"lld"))
-		{
-			return coff::SymbolRemovalStrategy::LLD_COMPATIBLE;
-		}
-		else if (string::Contains(lowerCaseLinkerPath.c_str(), L"lld-link"))
-		{
-			return coff::SymbolRemovalStrategy::LLD_COMPATIBLE;
-		}
-		else if (string::Contains(lowerCaseLinkerPath.c_str(), L"ld.lld"))
-		{
-			return coff::SymbolRemovalStrategy::LLD_COMPATIBLE;
-		}
-		else if (string::Contains(lowerCaseLinkerPath.c_str(), L"ld64.lld"))
+		compiler::LinkerType::Enum linkerType = compiler::DetermineLinkerType(linkerPath.c_str());
+
+		if (linkerType == compiler::LinkerType::LLD)
 		{
 			return coff::SymbolRemovalStrategy::LLD_COMPATIBLE;
 		}
@@ -206,6 +239,7 @@ namespace
 	static LiveModule::CompileResult Compile(ModuleCache* moduleCache, const symbols::ObjPath& normalizedObjPath, symbols::Compiland* compiland, const LiveModule::PerProcessData* processData, size_t processCount, unsigned int flags, LiveModule::UpdateType::Enum updateType)
 	{
 		const std::wstring compilerPath = GetCompilerPath(compiland);
+		const compiler::CompilerType::Enum compilerType = compiler::DetermineCompilerType(compilerPath.c_str());
 
 		// AMALGAMATION
 		// for files that are part of an amalgamation, check their current command-line options, file timestamps, etc. against
@@ -233,10 +267,24 @@ namespace
 		std::string compilerOptions;
 		compilerOptions.reserve(1u * 1024u * 1024u);
 
-		// add the "compile only" switch in any case. if it's already there, no harm done.
-		// for compilands that were compiled AND linked using cl.exe (which can call the linker internally!), this
-		// needs to be added.
-		compilerOptions += "-c ";
+		if (compilerType == compiler::CompilerType::CL)
+		{
+			// add the "compile only" switch in any case. if it's already there, no harm done.
+			// for compilands that were compiled AND linked using cl.exe (which can call the linker internally!), this
+			// needs to be added.
+			compilerOptions += "-c ";
+
+			// add compiler options based on flags
+			if (flags & CompileFlags::SERIALIZE_PDB_ACCESS)
+			{
+				compilerOptions += "-FS ";
+			}
+		}
+
+		// add the real command line for this compiland.
+		// note that this must always come first for Clang, because of the first "-cc1" command.
+		compilerOptions += compiland->commandLine.c_str();
+		compilerOptions += " ";
 
 		// add custom compiler options
 		{
@@ -248,36 +296,39 @@ namespace
 			}
 		}
 
-		// add compiler options based on flags
-		if (flags & CompileFlags::SERIALIZE_PDB_ACCESS)
+		if (compilerType == compiler::CompilerType::CL)
 		{
-			compilerOptions += "-FS ";
-		}
+			// add the command line that specifies the .pdb path in case its not contained in the compiland's command line.
+			// note that for builds using /Z7, the PDB path is optional and not needed.
+			const bool hasPdbPath = (compiland->pdbPath.GetLength() != 0u);
+			const bool hasPdbCommandLine = string::Contains(compiland->commandLine.c_str(), "-Fd");
+			if (hasPdbPath && !hasPdbCommandLine)
+			{
+				compilerOptions += "-Fd\"";
 
-		// add the real command line for this compiland
-		compilerOptions += compiland->commandLine.c_str();
-		compilerOptions += " ";
+				// the .PDB path could contain UTF8 characters, but the response file wants ANSI
+				compilerOptions += string::ToAnsiString(compiland->pdbPath);
+				compilerOptions += "\" ";
+			}
 
-		// add the command line that specifies the .pdb path in case its not contained in the compiland's command line.
-		// note that for builds using /Z7, the PDB path is optional and not needed.
-		const bool hasPdbPath = (compiland->pdbPath.GetLength() != 0u);
-		const bool hasPdbCommandLine = string::Contains(compiland->commandLine.c_str(), "-Fd");
-		if (hasPdbPath && !hasPdbCommandLine)
-		{
-			compilerOptions += "-Fd\"";
+			// add the command line that specifies the output .obj path in case its not contained in the compiland's command line
+			if (!string::Contains(compiland->commandLine.c_str(), "-Fo"))
+			{
+				compilerOptions += "-Fo\"";
 
-			// the .PDB path could contain UTF8 characters, but the response file wants ANSI
-			compilerOptions += string::ToAnsiString(compiland->pdbPath);
-			compilerOptions += "\" ";
+				// the .obj path could contain UTF8 characters, but the response file wants ANSI
+				compilerOptions += string::ToAnsiString(compiland->originalObjPath);
+				compilerOptions += "\" ";
+			}
 		}
-
-		// add the command line that specifies the output .obj path in case its not contained in the compiland's command line
-		if (!string::Contains(compiland->commandLine.c_str(), "-Fo"))
+		else if (compilerType == compiler::CompilerType::CLANG)
 		{
-			compilerOptions += "-Fo\"";
+			// for Clang, the output .obj path must always be specified
+			compilerOptions += "-o\"";
 
-			// the .obj path could contain UTF8 characters, but the response file wants ANSI
-			compilerOptions += string::ToAnsiString(compiland->originalObjPath);
+			// the .obj path could contain UTF8 characters, but the response file wants ANSI.
+			// additionally, Clang wants double backslashes.
+			compilerOptions += string::ToAnsiString(ImmutableString(string::ReplaceAll(string::ReplaceAll(compiland->originalObjPath.c_str(), "\\", "/"), "/", "\\\\").c_str()));
 			compilerOptions += "\" ";
 		}
 
@@ -288,8 +339,17 @@ namespace
 		// normalizing is NOT allowed, we don't want to follow reparse points!
 		{
 			const std::wstring wideSrcPath = string::ToWideString(compiland->srcPath);
-			const std::wstring prettyPath = file::NormalizePathWithoutLinks(wideSrcPath.c_str());
-			compilerOptions += string::ToAnsiString(string::ToUtf8String(prettyPath));
+			const Filesystem::Path prettyPath = Filesystem::NormalizePathWithoutResolvingLinks(wideSrcPath.c_str());
+
+			if (compilerType == compiler::CompilerType::CL)
+			{
+				compilerOptions += string::ToAnsiString(string::ToUtf8String(prettyPath.GetString()));
+			}
+			else if (compilerType == compiler::CompilerType::CLANG)
+			{
+				// Clang wants double backslashes
+				compilerOptions += string::ToAnsiString(string::ToUtf8String(string::ReplaceAll(string::ReplaceAll(prettyPath.GetString(), L"\\", L"/"), L"/", L"\\\\")));
+			}
 		}
 
 		compilerOptions += "\"";
@@ -297,8 +357,8 @@ namespace
 		// create a temporary file that acts as a so-called response file for the compiler, and contains
 		// the whole compiler command-line. this is done because the latter can get very long, longer
 		// than the limit of 32k characters.
-		const std::wstring responseFilePath = file::GenerateTempFilename();
-		file::CreateFileWithData(responseFilePath.c_str(), compilerOptions.c_str(), compilerOptions.size() * sizeof(char));
+		const Filesystem::Path responseFilePath = Filesystem::GenerateTempFilename();
+		Filesystem::CreateFileWithData(responseFilePath.GetString(), compilerOptions.c_str(), compilerOptions.size() * sizeof(char));
 
 		std::wstring compilerCommandLine;
 		compilerCommandLine.reserve(256u);	
@@ -310,43 +370,45 @@ namespace
 
 		// add response file to command line
 		compilerCommandLine += L"@\"";
-		compilerCommandLine += responseFilePath;
+		compilerCommandLine += responseFilePath.GetString();
 		compilerCommandLine += L"\"";
 
-		const process::Environment* environment = compiler::GetEnvironmentFromCache(compilerPath.c_str());
-		const void* environmentData = environment ? environment->data : nullptr;
+		const Process::Environment environment = compiler::GetEnvironmentFromCache(compilerPath.c_str());
+		const void* environmentData = environment.data;
 		std::wstring workingDirectory = string::ToWideString(compiland->workingDirectory);
 
 		// if the working directory does not exist, use the compiler's directory instead.
 		// otherwise, remote/distributed builds would use working directories on remote machines.
 		{
-			const file::Attributes& attributes = file::GetAttributes(workingDirectory.c_str());
-			if (!file::DoesExist(attributes))
+			const Filesystem::PathAttributes& attributes = Filesystem::GetAttributes(workingDirectory.c_str());
+			if (!Filesystem::DoesExist(attributes))
 			{
-				workingDirectory = file::GetDirectory(compilerPath);
+				workingDirectory = Filesystem::GetDirectory(compilerPath.c_str()).GetString();
 			}
 		}
 
 		LC_LOG_USER("Compiling %s %s", isPartOfAmalgamation ? "split file" : "file", normalizedObjPath.c_str());
 
-		process::Context* processContext = process::Spawn(compilerPath.c_str(), workingDirectory.c_str(), compilerCommandLine.c_str(), environmentData, process::SpawnFlags::REDIRECT_STDOUT | process::SpawnFlags::NO_WINDOW);
-		const unsigned int exitCode = process::Wait(processContext);
-		const wchar_t* compilerOutput = processContext->stdoutData.c_str();
+		Process::Context* processContext = Process::Spawn(compilerPath.c_str(), workingDirectory.c_str(), compilerCommandLine.c_str(), environmentData, Process::SpawnFlags::REDIRECT_STDOUT | Process::SpawnFlags::NO_WINDOW);
+		const unsigned int exitCode = Process::Wait(processContext);
+
+		const std::wstring processStdout = Process::GetStdOutData(processContext);
+		const wchar_t* compilerOutput = processStdout.c_str();
 
 		// log the complete command-line into the DEV log
 		{
 			LC_LOG_DEV("Compiler command-line: ");
-			logging::LogNoFormat<logging::Channel::DEV>(compilerOptions.c_str());
-			logging::LogNoFormat<logging::Channel::DEV>("\n");
+			Logging::LogNoFormat<Logging::Channel::DEV>(compilerOptions.c_str());
+			Logging::LogNoFormat<Logging::Channel::DEV>("\n");
 		}
 
 		{
 			CriticalSection::ScopedLock lock(&g_compileOutputCS);
 
 			// log to local UI
-			logging::LogNoFormat<logging::Channel::USER>(compilerOutput);
+			Logging::LogNoFormat<Logging::Channel::USER>(compilerOutput);
 
-			const size_t outputLength = processContext->stdoutData.length();
+			const size_t outputLength = processStdout.length();
 			if (outputLength != 0u)
 			{
 				if (updateType != LiveModule::UpdateType::NO_CLIENT_COMMUNICATION)
@@ -382,46 +444,13 @@ namespace
 			}
 		}
 	
-		process::Destroy(processContext);
+		Process::Destroy(processContext);
 
-		file::Delete(responseFilePath.c_str());
+		Filesystem::Delete(responseFilePath.GetString());
 
 		return LiveModule::CompileResult { exitCode, true };
 	}
 
-	// helper function that returns or generates the unique ID of an optional compiland.
-	// for files split off from amalgamated files, we need to use the original object path of the amalgamated file here,
-	// otherwise names of symbols would differ, leading to constructors of global instances being called again.
-	static inline uint32_t GetCompilandId(const symbols::Compiland* compiland, const wchar_t* const objPath, const types::vector<LiveModule::ModifiedObjFile>& modifiedObjFiles)
-	{
-		// try to find the given .obj path in the array of modified object files to check if there's an original amalgamated object path for it
-		for (size_t i = 0u; i < modifiedObjFiles.size(); ++i)
-		{
-			const LiveModule::ModifiedObjFile& objFile = modifiedObjFiles[i];
-
-			// don't bother checking strings if the amalgamated object path is empty anyway
-			if (!objFile.amalgamatedObjPath.empty())
-			{
-				if (string::Matches(objPath, objFile.objPath.c_str()))
-				{
-					return uniqueId::Generate(file::NormalizePath(objFile.amalgamatedObjPath.c_str()));
-				}
-			}
-		}
-
-		if (compiland)
-		{
-			// the compiland already exists
-			return compiland->uniqueId;
-		}
-		else
-		{
-			// BEGIN EPIC MOD - Static map of object file to unity blob
-			return symbols::GetCompilandIdFromPath(objPath);
-			// BEGIN EPIC MOD - Static map of object file to unity blob
-		}
-	}
-
 
 	struct SymbolAndRelocation
 	{
@@ -457,6 +486,10 @@ namespace
 			return nullptr;
 		}
 
+		// TODO: all relocations point to the same destination symbol, hence all relocations also have the same destination section index.
+		// if a relocation wouldn't be patched because of certain characteristics of the section, we can do that ONCE here, and don't need to
+		// check every relocation.
+
 		const size_t relocationCount = cache.size();
 		for (size_t i = 0u; i < relocationCount; ++i)
 		{
@@ -503,7 +536,7 @@ namespace
 
 
 	template <typename T>
-	static types::vector<symbols::ObjPath> UpdateCoffCache(const T& compilands, CoffCache<coff::CoffDB>* coffCache, CacheUpdate::Enum updateType, coff::ReadFlags::Enum coffReadFlags, const types::vector<LiveModule::ModifiedObjFile>& modifiedOrNewObjFiles)
+	static types::vector<symbols::ObjPath> UpdateCoffCache(const T& compilands, CoffCache<coff::CoffDB>* coffCache, CacheUpdate::Enum updateType, const types::vector<symbols::ModifiedObjFile>& modifiedOrNewObjFiles)
 	{
 		LC_LOG_INDENT_DEV;
 
@@ -520,7 +553,7 @@ namespace
 			symbols::ObjPath objPath = it->first;
 			const std::wstring& wideObjPath = string::ToWideString(objPath);
 			const symbols::Compiland* compiland = it->second;
-			const uint32_t compilandUniqueId = GetCompilandId(compiland, wideObjPath.c_str(), modifiedOrNewObjFiles);
+			const uint32_t compilandUniqueId = symbols::GetCompilandId(compiland, wideObjPath.c_str(), modifiedOrNewObjFiles);
 
 			const bool shouldUpdate = (updateType == CacheUpdate::NON_EXISTANT)
 				? (coffCache->Lookup(objPath) == nullptr)				// NON-EXISTANT: update cache only for files which don't have an entry yet
@@ -530,14 +563,14 @@ namespace
 			{
 				updatedCoffs.push_back(objPath);
 
-				auto task = scheduler::CreateTask(taskRoot, [objPath, wideObjPath, compiland, coffCache, compilandUniqueId, coffReadFlags]()
+				auto task = scheduler::CreateTask(taskRoot, [objPath, wideObjPath, coffCache, compilandUniqueId]()
 				{
 					LC_LOG_DEV("Updating COFF cache for file %s", objPath.c_str());
 
 					coff::ObjFile* objFile = coff::OpenObj(wideObjPath.c_str());
 					if (objFile && objFile->memoryFile)
 					{
-						coff::CoffDB* database = coff::GatherDatabase(objFile, compilandUniqueId, coffReadFlags);
+						coff::CoffDB* database = coff::GatherDatabase(objFile, compilandUniqueId);
 						if (database)
 						{
 							coffCache->Update(objPath, database);
@@ -567,24 +600,24 @@ namespace
 
 
 #if LC_64_BIT
-	static executable::PreferredBase FindPreferredImageBase(uint32_t imageSize, unsigned int processId, process::Handle processHandle, void* moduleBase)
+	static executable::PreferredBase FindPreferredImageBase(uint32_t imageSize, Process::Id processId, Process::Handle processHandle, void* moduleBase)
 	{
 		// work out the lower and upper bound of the memory region into which a patch could be loaded
-		const uint32_t exeSize = process::GetImageSize(processHandle, moduleBase);
+		const uint32_t exeSize = Process::GetModuleSize(processHandle, moduleBase);
 		const uint32_t patchSize = imageSize;
 
-		const void* lowerBound = pointer::Offset<const void*>(moduleBase, static_cast<int64_t>(exeSize) - 0x80000000ll);
-		const void* upperBound = pointer::Offset<const void*>(moduleBase, 0x7FFFFFFFull);
+		const void* lowerBound = GetLowerBoundIn4GBRange(pointer::Offset<const void*>(moduleBase, exeSize));
+		const void* upperBound = GetUpperBoundIn4GBRange(moduleBase);
 
 		LC_LOG_DEV("Scanning memory range from 0x%p to 0x%p (base: 0x%p, exeSize: 0x%X, patchSize: 0x%X, PID: %d)",
-			lowerBound, upperBound, moduleBase, exeSize, patchSize, processId);
+			lowerBound, upperBound, moduleBase, exeSize, patchSize, +processId);
 
 		// modules can only be loaded at 64KB boundaries, so we should scan memory only at aligned addresses
 		const size_t MODULE_ALIGNMENT = 64u * 1024u;
-		void* preferredBase = process::ScanMemoryRange(processHandle, lowerBound, upperBound, patchSize, MODULE_ALIGNMENT);
+		void* preferredBase = Process::ScanMemoryRange(processHandle, lowerBound, upperBound, patchSize, MODULE_ALIGNMENT);
 
 		const executable::PreferredBase preferredImageBase = pointer::AsInteger<executable::PreferredBase>(preferredBase);
-		LC_LOG_DEV("Preferred base address for image: 0x%" PRIX64 " (PID: %d)", preferredImageBase, processId);
+		LC_LOG_DEV("Preferred base address for image: 0x%" PRIX64 " (PID: %d)", preferredImageBase, +processId);
 
 		return preferredImageBase;
 	}
@@ -592,9 +625,9 @@ namespace
 
 
 	// helper function that returns the instruction pointers of all threads of a process
-	static types::vector<const void*> EnumerateInstructionPointers(unsigned int processId)
+	static types::vector<const void*> EnumerateInstructionPointers(Process::Id processId)
 	{
-		const std::vector<unsigned int>& threadIds = process::EnumerateThreads(processId);
+		const std::vector<Thread::Id>& threadIds = Process::EnumerateThreads(processId);
 		const size_t threadCount = threadIds.size();
 
 		types::vector<const void*> instructionPointers;
@@ -602,15 +635,15 @@ namespace
 
 		for (size_t i=0u; i < threadCount; ++i)
 		{
-			const unsigned int threadId = threadIds[i];
-			thread::Handle threadHandle = thread::Open(threadId);
+			const Thread::Id threadId = threadIds[i];
+			Thread::Handle threadHandle = Thread::Open(threadId);
 
-			const thread::Context& context = thread::GetContext(threadHandle);
-			const void* ip = thread::ReadInstructionPointer(context);
+			const Thread::Context context = Thread::GetContext(threadHandle);
+			const void* ip = Thread::ReadInstructionPointer(&context);
 
 			instructionPointers.push_back(ip);
 
-			thread::Close(threadHandle);
+			Thread::Close(threadHandle);
 		}
 
 		return instructionPointers;
@@ -618,7 +651,7 @@ namespace
 
 
 	// helper function that checks whether a patch was loaded at a valid address
-	static bool CheckPatchAddressValidity(void* originalModuleBase, void* patchBase, process::Handle processHandle)
+	static bool CheckPatchAddressValidity(void* originalModuleBase, void* patchBase, Process::Handle processHandle)
 	{
 		if (!patchBase)
 		{
@@ -634,7 +667,7 @@ namespace
 		{
 			if (patchBase >= originalModuleBase)
 			{
-				const uint32_t patchSize = process::GetImageSize(processHandle, patchBase);
+				const uint32_t patchSize = Process::GetModuleSize(processHandle, patchBase);
 				const uint64_t displacement = pointer::Displacement<uint64_t>(originalModuleBase, pointer::Offset<const char*>(patchBase, patchSize));
 				if (displacement > 0x80000000ull)
 				{
@@ -645,7 +678,7 @@ namespace
 			}
 			else
 			{
-				const uint32_t exeSize = process::GetImageSize(processHandle, originalModuleBase);
+				const uint32_t exeSize = Process::GetModuleSize(processHandle, originalModuleBase);
 				const uint64_t displacement = pointer::Displacement<uint64_t>(patchBase, pointer::Offset<const char*>(originalModuleBase, exeSize));
 				if (displacement > 0x80000000ull)
 				{
@@ -680,18 +713,18 @@ namespace
 #endif
 
 	// helper function to patch UE4-specific symbols
-	static void PatchUE4NatVisSymbols(void* originalModuleBase, void* patchBase, uint32_t originalRva, uint32_t patchRva, process::Handle processHandle)
+	static void PatchUE4NatVisSymbols(void* originalModuleBase, void* patchBase, uint32_t originalRva, uint32_t patchRva, Process::Handle processHandle)
 	{
 		const void* originalAddr = pointer::Offset<const void*>(originalModuleBase, originalRva);
 		void* patchAddr = pointer::Offset<void*>(patchBase, patchRva);
 
-		const uintptr_t value = process::ReadProcessMemory<uintptr_t>(processHandle, originalAddr);
-		process::WriteProcessMemory(processHandle, patchAddr, value);
+		const uintptr_t value = Process::ReadProcessMemory<uintptr_t>(processHandle, originalAddr);
+		Process::WriteProcessMemory(processHandle, patchAddr, value);
 	}
 
 
 	// helper function to patch security cookies
-	static void PatchSecurityCookie(void* originalModuleBase, void* patchBase, uint32_t originalRva, uint32_t patchRva, process::Handle processHandle)
+	static void PatchSecurityCookie(void* originalModuleBase, void* patchBase, uint32_t originalRva, uint32_t patchRva, Process::Handle processHandle)
 	{
 		const void* cookieAddr = pointer::Offset<const void*>(originalModuleBase, originalRva);
 		void* newCookieAddr = pointer::Offset<void*>(patchBase, patchRva);
@@ -702,13 +735,13 @@ namespace
 		typedef uint32_t CookieType;
 #endif
 
-		const CookieType cookie = process::ReadProcessMemory<CookieType>(processHandle, cookieAddr);
-		process::WriteProcessMemory(processHandle, newCookieAddr, cookie);
+		const CookieType cookie = Process::ReadProcessMemory<CookieType>(processHandle, cookieAddr);
+		Process::WriteProcessMemory(processHandle, newCookieAddr, cookie);
 	}
 
 
 	// helper function to patch DllMain
-	static void PatchDllMain(void* patchBase, uint32_t dllMainRva, process::Handle processHandle)
+	static void PatchDllMain(void* patchBase, uint32_t dllMainRva, Process::Handle processHandle)
 	{
 		LC_LOG_DEV("Disabling optional DLL entry point");
 
@@ -722,16 +755,22 @@ namespace
 		// the code to inject on x64 is:
 		//		B0 01		mov al, 1
 		//		C3			ret				different calling convention than x86
+		// BEGIN EPIC MOD
 		const uint8_t PatchData[3u] = { 0xB0, 0x01, 0xC3 };
+		// END EPIC MOD
 #else
 		// the code to inject on x86 is:
 		//		B0 01		mov al, 1
 		//		C2 0C 00	ret 0Ch			different calling convention than x64
+		// BEGIN EPIC MOD
 		const uint8_t PatchData[5u] = { 0xB0, 0x01, 0xC2, 0x0C, 0x00 };
+		// END EPIC MOD
 #endif
 
 		uint8_t* address = pointer::Offset<uint8_t*>(patchBase, dllMainRva);
-		process::WriteProcessMemory(processHandle, address, PatchData, sizeof(PatchData));
+		// BEGIN EPIC MOD
+		Process::WriteProcessMemory(processHandle, address, PatchData, sizeof(PatchData));
+		// END EPIC MOD
 	}
 
 
@@ -773,11 +812,11 @@ namespace
 			{
 				const ModuleCache::ProcessData& processData = hookData.data->processes[p];
 
-				const unsigned int pid = processData.processId;
+				const Process::Id pid = processData.processId;
 				void* moduleBase = processData.moduleBase;
 				const DuplexPipe* pipe = processData.pipe;
 
-				LC_LOG_USER("Calling compile start hooks (PID: %d)", pid);
+				LC_LOG_USER("Calling compile start hooks (PID: %d)", +pid);
 				pipe->SendCommandAndWaitForAck(MakeCallHooksCommand(hook::Type::COMPILE_START, moduleBase, hookData), nullptr, 0u);
 			}
 		}
@@ -800,11 +839,11 @@ namespace
 			{
 				const ModuleCache::ProcessData& processData = hookData.data->processes[p];
 
-				const unsigned int pid = processData.processId;
+				const Process::Id pid = processData.processId;
 				void* moduleBase = processData.moduleBase;
 				const DuplexPipe* pipe = processData.pipe;
 
-				LC_LOG_USER("Calling compile success hooks (PID: %d)", pid);
+				LC_LOG_USER("Calling compile success hooks (PID: %d)", +pid);
 				pipe->SendCommandAndWaitForAck(MakeCallHooksCommand(hook::Type::COMPILE_SUCCESS, moduleBase, hookData), nullptr, 0u);
 			}
 		}
@@ -827,11 +866,11 @@ namespace
 			{
 				const ModuleCache::ProcessData& processData = hookData.data->processes[p];
 
-				const unsigned int pid = processData.processId;
+				const Process::Id pid = processData.processId;
 				void* moduleBase = processData.moduleBase;
 				const DuplexPipe* pipe = processData.pipe;
 
-				LC_LOG_USER("Calling compile error hooks (PID: %d)", pid);
+				LC_LOG_USER("Calling compile error hooks (PID: %d)", +pid);
 				pipe->SendCommandAndWaitForAck(MakeCallHooksCommand(hook::Type::COMPILE_ERROR, moduleBase, hookData), nullptr, 0u);
 			}
 		}
@@ -932,7 +971,29 @@ void LiveModule::Load(symbols::Provider* provider, symbols::DiaCompilandDB* diaC
 			options |= symbols::CompilandOptions::TRACK_OBJ_ONLY;
 		}
 
-		auto db = symbols::GatherCompilands(localProvider, localDiaCompilandDb, GetAmalgamatedSplitThreshold(), options);
+		symbols::CompilandDB* db = nullptr;
+		// BEGIN EPIC MOD
+#if 1
+		db = symbols::GatherCompilands(localProvider, localDiaCompilandDb, GetAmalgamatedSplitThreshold(), options);
+#else
+		const std::wstring fastBuildDatabasePath = appSettings::g_fastBuildDatabasePath->GetValue();
+		const std::wstring fastBuildDllName = appSettings::g_fastBuildDllName->GetValue();
+		if (fastBuildDatabasePath.length() > 0u)
+		{
+			db = FASTBuild::GatherCompilands(fastBuildDllName.c_str(), string::ToUtf8String(fastBuildDatabasePath).c_str(), string::ToUtf8String(m_moduleName).c_str(), GetAmalgamatedSplitThreshold(), options);
+
+			// safety net
+			if (!db)
+			{
+				db = new symbols::CompilandDB;
+			}
+		}
+		else
+		{
+			db = symbols::GatherCompilands(localProvider, localDiaCompilandDb, GetAmalgamatedSplitThreshold(), options);
+		}
+#endif
+		// END EPIC MOD
 
 		symbols::DestroyDiaCompilandDB(localDiaCompilandDb);
 		symbols::Close(localProvider);
@@ -1073,20 +1134,6 @@ void LiveModule::Load(symbols::Provider* provider, symbols::DiaCompilandDB* diaC
 					// yes, store it in our database
 					m_pchSymbolToCompilandName.emplace(symbol->name, compilandName);
 				}
-
-				// is this a weak symbol from a compiland that is part of a library?
-				if (symbols::IsWeakSymbol(symbol->name))
-				{
-					// if there is no compiland associated with this symbol, then it must have originated from a library.
-					// if there is a compiland, we need to check if the compiland is part of a static library.
-					const symbols::Compiland* compiland = symbols::FindCompiland(m_compilandDB, compilandName);
-					const bool isWeakSymbolInLibrary = compiland ? compiland->isPartOfLibrary : true;
-					if (isWeakSymbolInLibrary)
-					{
-						LC_LOG_DEV("Weak symbol %s in library compiland %s", symbol->name.c_str(), compilandName.c_str());
-						m_weakSymbolsInLibs.push_back(symbol->name);
-					}
-				}
 			}
 		}
 	}
@@ -1127,11 +1174,7 @@ void LiveModule::Load(symbols::Provider* provider, symbols::DiaCompilandDB* diaC
 			{
 				const std::wstring& wideObjPath = string::ToWideString(objPath);
 				coff::ObjFile* objFile = coff::OpenObj(wideObjPath.c_str());
-				// BEGIN EPIC MOD - Support for unity files s with external build system
-				const bool splitAmalgamatedFiles = appSettings::g_amalgamationSplitIntoSingleParts->GetValue();
-				const coff::ReadFlags::Enum coffReadFlags = splitAmalgamatedFiles ? coff::ReadFlags::GENERATE_ANS_NAME_FROM_UNIQUE_ID : coff::ReadFlags::NONE;
-				coff::CoffDB* database = coff::GatherDatabase(objFile, compiland->uniqueId, coffReadFlags);
-				// END EPIC MOD
+				coff::CoffDB* database = coff::GatherDatabase(objFile, compiland->uniqueId);
 				coff::CloseObj(objFile);
 
 				return GatherResult { database, objPath };
@@ -1188,7 +1231,7 @@ void LiveModule::Unload(void)
 		for (size_t p = 0u; p < processCount; ++p)
 		{
 			const ModuleCache::ProcessData& process = entry.processes[p];
-			if (!process::IsActive(process.processHandle))
+			if (!Process::IsActive(process.processHandle))
 			{
 				// this process is no longer valid, ignore it
 				continue;
@@ -1205,23 +1248,35 @@ void LiveModule::RegisterProcess(LiveProcess* liveProcess, void* moduleBase, con
 {
 	m_moduleCache->RegisterProcess(m_mainModuleToken, liveProcess, moduleBase);
 
-	PerProcessData perProcessData = { liveProcess, moduleBase, modulePath };
+#if LC_64_BIT
+	const void* lowerBound = GetLowerBoundIn4GBRange(moduleBase);
+	const void* upperBound = GetUpperBoundIn4GBRange(moduleBase);
+	VirtualMemoryRange* virtualMemoryRange = new VirtualMemoryRange(liveProcess->GetProcessHandle(), lowerBound, upperBound, 64u * 1024u);
+	virtualMemoryRange->ReservePages();
+#else
+	VirtualMemoryRange* virtualMemoryRange = new VirtualMemoryRange(liveProcess->GetProcessHandle(), moduleBase, moduleBase, 64u * 1024u);
+#endif
+
+	PerProcessData perProcessData = { liveProcess, moduleBase, modulePath, virtualMemoryRange };
 	m_perProcessData.emplace_back(perProcessData);
 }
 
 
 void LiveModule::UnregisterProcess(LiveProcess* liveProcess)
 {
-	const unsigned int processId = liveProcess->GetProcessId();
+	const Process::Id processId = liveProcess->GetProcessId();
 
 	m_moduleCache->UnregisterProcess(liveProcess);
-	m_patchedAddressesPerProcess.erase(liveProcess->GetProcessId());
+	m_patchedAddressesPerProcess.erase(processId);
 
 	for (auto it = m_perProcessData.begin(); it != m_perProcessData.end(); ++it)
 	{
 		const PerProcessData& data = *it;
 		if (data.liveProcess == liveProcess)
 		{
+			data.virtualMemoryRange->FreeReservedPages();
+			delete data.virtualMemoryRange;
+
 			m_perProcessData.erase(it);
 			break;
 		}
@@ -1231,7 +1286,7 @@ void LiveModule::UnregisterProcess(LiveProcess* liveProcess)
 
 void LiveModule::DisableControlFlowGuard(LiveProcess* liveProcess, void* moduleBase)
 {
-	process::Handle processHandle = liveProcess->GetProcessHandle();
+	Process::Handle processHandle = liveProcess->GetProcessHandle();
 
 	// disable control flow guard (CFG) checks
 	// https://msdn.microsoft.com/en-us/library/windows/desktop/mt637065(v=vs.85).aspx
@@ -1263,13 +1318,13 @@ void LiveModule::DisableControlFlowGuard(LiveProcess* liveProcess, void* moduleB
 				// make sure the process gets suspended while writing to its memory.
 				// otherwise, writing could change the page protection of an executable page while code is currently executing
 				// (when using the lpp*Async API), which would lead to a crash.
-				process::Suspend(processHandle);
+				Process::Suspend(processHandle);
 
-				void* addr = process::ReadProcessMemory<void*>(processHandle, pointer::Offset<const void*>(moduleBase, cfgFuncPtr->rva));
+				void* addr = Process::ReadProcessMemory<void*>(processHandle, pointer::Offset<const void*>(moduleBase, cfgFuncPtr->rva));
 				const uint8_t OPCODE_RET = 0xC3;
-				process::WriteProcessMemory(processHandle, addr, OPCODE_RET);
+				Process::WriteProcessMemory(processHandle, addr, OPCODE_RET);
 
-				process::Resume(processHandle);
+				Process::Resume(processHandle);
 			}
 		}
 	}
@@ -1282,26 +1337,32 @@ void LiveModule::UpdateDirectoryCache(DirectoryCache* cache)
 	for (auto it : m_compilandDB->dependencies)
 	{
 		symbols::Dependency* dependency = it.second;
-		if (dependency->parentDirectory)
+		if (!dependency->parentDirectory)
 		{
-			// dependency has a valid parent directory entry already
-			continue;
+			// dependency does not have a valid parent directory entry yet, create a new one
+			const ImmutableString& path = it.first;
+			UpdateDirectoryCache(path, dependency, cache);
 		}
 
-		const ImmutableString& path = it.first;
-		UpdateDirectoryCache(path, dependency, cache);
+		// merge initial changes of dependencies into the parent directories.
+		// this ensures that modifications to source files that were done before a module was loaded,
+		// will automatically be picked up by us.
+		dependency->parentDirectory->hadChange |= dependency->hadInitialChange;
 	}
 }
 
 
-LiveModule::ErrorType::Enum LiveModule::Update(FileAttributeCache* fileCache, DirectoryCache* directoryCache, UpdateType::Enum updateType, const types::vector<ModifiedObjFile>& modifiedOrNewObjFiles)
+LiveModule::ErrorType::Enum LiveModule::Update(FileAttributeCache* fileCache, DirectoryCache* directoryCache, UpdateType::Enum updateType, const types::vector<symbols::ModifiedObjFile>& modifiedOrNewObjFiles)
 {
 	telemetry::Scope updateScope("Update live module");
 
 	LC_LOG_DEV("LiveModule Update: %S", m_moduleName.c_str());
 
+	// BEGIN EPIC MOD
 	GLiveCodingServer->GetStatusChangeDelegate().ExecuteIfBound(L"Checking modified files...");
+	// END EPIC MOD
 
+	bool usesClang = false;
 	bool forceAmalgamationPartsLinkage = false;
 
 	// only check for modifications if no files have been handed to us
@@ -1326,8 +1387,8 @@ LiveModule::ErrorType::Enum LiveModule::Update(FileAttributeCache* fileCache, Di
 			{
 				dependency->lastModification = currentTime;
 				{
-					const std::wstring prettyPath = file::NormalizePathWithoutLinks(filePath.c_str());
-					LC_LOG_USER("File %S was modified", prettyPath.c_str());
+					const Filesystem::Path prettyPath = Filesystem::NormalizePathWithoutResolvingLinks(filePath.c_str());
+					LC_LOG_USER("File %S was modified", prettyPath.GetString());
 				}
 
 				// AMALGAMATION
@@ -1418,7 +1479,9 @@ LiveModule::ErrorType::Enum LiveModule::Update(FileAttributeCache* fileCache, Di
 			}
 			else
 			{
-				LC_LOG_USER("Detected %zu file(s) to be compiled for Live++ module %S", m_modifiedFiles.size(), m_moduleName.c_str());
+				// BEGIN EPIC MOD
+				LC_LOG_USER("Detected %zu file(s) to be compiled for Live coding module %S", m_modifiedFiles.size(), m_moduleName.c_str());
+				// END EPIC MOD
 			}
 		}
 		else if (m_runMode == RunMode::EXTERNAL_BUILD_SYSTEM)
@@ -1437,17 +1500,17 @@ LiveModule::ErrorType::Enum LiveModule::Update(FileAttributeCache* fileCache, Di
 			LC_LOG_USER("File %S was modified or is new", modifiedOrNewObjFiles[i].objPath.c_str());
 		}
 
-		LC_LOG_USER("Building patch from %zu file(s) for Live Coding module %S", modifiedOrNewObjFiles.size(), m_moduleName.c_str());
+		// BEGIN EPIC MOD
+		LC_LOG_USER("Building patch from %zu file(s) for Live coding module %S", modifiedOrNewObjFiles.size(), m_moduleName.c_str());
+		// END EPIC MOD
 	}
 
 	// let the user know that we're about to compile
 	CallCompileStartHooks(m_moduleCache, updateType);
 
-	// AMALGAMATION
-	const bool splitAmalgamatedFiles = appSettings::g_amalgamationSplitIntoSingleParts->GetValue();
-	const coff::ReadFlags::Enum coffReadFlags = splitAmalgamatedFiles ? coff::ReadFlags::GENERATE_ANS_NAME_FROM_UNIQUE_ID : coff::ReadFlags::NONE;
-
+	// BEGIN EPIC MOD
 	GLiveCodingServer->GetStatusChangeDelegate().ExecuteIfBound(L"Updating first time COFF cache...");
+	// END EPIC MOD
 
 	// before starting to compile, update the COFF cache for files that have been touched for the first time
 	struct ModifiedFile
@@ -1526,11 +1589,11 @@ LiveModule::ErrorType::Enum LiveModule::Update(FileAttributeCache* fileCache, Di
 						}
 
 						// do the loading and gathering concurrently
-						auto task = scheduler::CreateTask(taskRoot, [fileIndex, amalgamatedObjPath, compiland, coffReadFlags]()
+						auto task = scheduler::CreateTask(taskRoot, [fileIndex, amalgamatedObjPath, compiland]()
 						{
 							const std::wstring& wideObjPath = string::ToWideString(amalgamatedObjPath);
 							coff::ObjFile* objFile = coff::OpenObj(wideObjPath.c_str());
-							coff::CoffDB* database = coff::GatherDatabase(objFile, compiland->uniqueId, coffReadFlags);
+							coff::CoffDB* database = coff::GatherDatabase(objFile, compiland->uniqueId);
 							coff::CloseObj(objFile);
 
 							return GatherResult { fileIndex, database };
@@ -1576,7 +1639,9 @@ LiveModule::ErrorType::Enum LiveModule::Update(FileAttributeCache* fileCache, Di
 	// recompile changed files
 	if (m_runMode == RunMode::DEFAULT)
 	{
+		// BEGIN EPIC MOD
 		GLiveCodingServer->GetStatusChangeDelegate().ExecuteIfBound(L"Compiling...");
+		// END EPIC MOD
 
 		struct LocalCompileResult
 		{
@@ -1920,7 +1985,9 @@ LiveModule::ErrorType::Enum LiveModule::Update(FileAttributeCache* fileCache, Di
 
 
 
+	// BEGIN EPIC MOD
 	GLiveCodingServer->GetStatusChangeDelegate().ExecuteIfBound(L"Updating compilands...");
+	// END EPIC MOD
 
 	// we want to link a minimal .DLL file that contains all modified .OBJ files and only those required for resolving symbols.
 	// because we require users to use /OPT:NOREF and /OPT:NOICF, finding the set of files that need to be linked in is
@@ -1950,14 +2017,14 @@ LiveModule::ErrorType::Enum LiveModule::Update(FileAttributeCache* fileCache, Di
 		{
 			static void UpdateExternalSymbolsAndNeededFiles
 			(
-				const symbols::ObjPath& objPath, const symbols::Compiland* compiland, uint32_t compilandUniqueId, coff::ReadFlags::Enum coffReadFlags, const types::StringMap<ImmutableString>& pchSymbolToCompilandName,
+				const symbols::ObjPath& objPath, const symbols::Compiland* compiland, uint32_t compilandUniqueId, const types::StringMap<ImmutableString>& pchSymbolToCompilandName,
 				types::StringMap<CompilandInfo>& externalSymbols, types::StringSet& neededCompilands
 			)
 			{
 				coff::ObjFile* coffFile = coff::OpenObj(string::ToWideString(objPath).c_str());
 				if (coffFile && coffFile->memoryFile)
 				{
-					coff::ExternalSymbolDB* externalSymbolDb = coff::GatherExternalSymbolDatabase(coffFile, compilandUniqueId, coffReadFlags);
+					coff::ExternalSymbolDB* externalSymbolDb = coff::GatherExternalSymbolDatabase(coffFile, compilandUniqueId);
 					types::vector<std::string> linkerDirectives = coff::ExtractLinkerDirectives(coffFile);
 					coff::CloseObj(coffFile);
 
@@ -2023,7 +2090,7 @@ LiveModule::ErrorType::Enum LiveModule::Update(FileAttributeCache* fileCache, Di
 					LC_LOG_DEV("%s is recompiled", objPath.c_str());
 					neededCompilands.emplace(objPath);
 
-					Helper::UpdateExternalSymbolsAndNeededFiles(objPath, compiland, compiland->uniqueId, coffReadFlags, m_pchSymbolToCompilandName, externalSymbols, neededCompilands);
+					Helper::UpdateExternalSymbolsAndNeededFiles(objPath, compiland, compiland->uniqueId, m_pchSymbolToCompilandName, externalSymbols, neededCompilands);
 				}
 				else
 				{
@@ -2053,7 +2120,7 @@ LiveModule::ErrorType::Enum LiveModule::Update(FileAttributeCache* fileCache, Di
 				const symbols::Compiland* compiland = symbols::FindCompiland(m_compilandDB, objPath);
 
 				// new compilands won't be found in the database, so there's no unique ID yet that we can use
-				const uint32_t compilandUniqueId = GetCompilandId(compiland, it->objPath.c_str(), modifiedOrNewObjFiles);
+				const uint32_t compilandUniqueId = symbols::GetCompilandId(compiland, it->objPath.c_str(), modifiedOrNewObjFiles);
 
 				// this file was either modified or is new. in any case, the new .OBJ needs to be linked in, even
 				// though the file might be contained in a library.
@@ -2061,12 +2128,14 @@ LiveModule::ErrorType::Enum LiveModule::Update(FileAttributeCache* fileCache, Di
 				LC_LOG_DEV("%s %s", objPath.c_str(), compiland ? "was recompiled" : "is new");
 				neededCompilands.emplace(objPath);
 
-				Helper::UpdateExternalSymbolsAndNeededFiles(objPath, compiland, compilandUniqueId, coffReadFlags, m_pchSymbolToCompilandName, externalSymbols, neededCompilands);
+				Helper::UpdateExternalSymbolsAndNeededFiles(objPath, compiland, compilandUniqueId, m_pchSymbolToCompilandName, externalSymbols, neededCompilands);
 			}
 		}
 	}
 
+	// BEGIN EPIC MOD
 	GLiveCodingServer->GetStatusChangeDelegate().ExecuteIfBound(L"Reconstructing symbols...");
+	// END EPIC MOD
 
 	// we now have a list of all .obj files that are going to be part of the next patch.
 	// reconstruct symbols lazily for those object files that have not been reconstructed yet from the initial main executable.
@@ -2112,7 +2181,7 @@ LiveModule::ErrorType::Enum LiveModule::Update(FileAttributeCache* fileCache, Di
 		const size_t count = objToReconstruct.size();
 		if (count > 0u)
 		{
-			executable::Image* image = executable::OpenImage(m_moduleName.c_str(), file::OpenMode::READ_ONLY);
+			executable::Image* image = executable::OpenImage(m_moduleName.c_str(), Filesystem::OpenMode::READ);
 			executable::ImageSectionDB* imageSections = executable::GatherImageSectionDB(image);
 
 			// load and cache all .obj not in the cache yet concurrently
@@ -2130,29 +2199,29 @@ LiveModule::ErrorType::Enum LiveModule::Update(FileAttributeCache* fileCache, Di
 						// there is no entry yet for this COFF in the cache.
 						// this means that this .obj was not recompiled (otherwise it would have an entry already),
 						// but has been pulled in for the first time due to unresolved symbols.
-						auto task = scheduler::CreateTask(taskRoot, [this, objPath, coffReadFlags, &modifiedOrNewObjFiles]()
+						auto task = scheduler::CreateTask(taskRoot, [this, objPath, &modifiedOrNewObjFiles]()
 						{
 							const symbols::Compiland* compiland = symbols::FindCompiland(m_compilandDB, objPath);
 							const std::wstring& wideObjPath = string::ToWideString(objPath);
-							const uint32_t compilandUniqueId = GetCompilandId(compiland, wideObjPath.c_str(), modifiedOrNewObjFiles);
+							const uint32_t compilandUniqueId = symbols::GetCompilandId(compiland, wideObjPath.c_str(), modifiedOrNewObjFiles);
 
-								LC_LOG_DEV("Need %s for the first time, updating COFF cache", objPath.c_str());
+							LC_LOG_DEV("Need %s for the first time, updating COFF cache", objPath.c_str());
 
-								coff::ObjFile* objFile = coff::OpenObj(wideObjPath.c_str());
-								if (objFile && objFile->memoryFile)
+							coff::ObjFile* objFile = coff::OpenObj(wideObjPath.c_str());
+							if (objFile && objFile->memoryFile)
+							{
+								// note that even though we might be dealing with a single-part .obj of an amalgamated .obj
+								// here, the symbols will be disambiguated using the same uniqueId as the original amalgamated file.
+								coff::CoffDB* database = coff::GatherDatabase(objFile, compilandUniqueId);
+								if (database)
 								{
-									// note that even though we might be dealing with a single-part .obj of an amalgamated .obj
-									// here, the symbols will be disambiguated using the same uniqueId as the original amalgamated file.
-								coff::CoffDB* database = coff::GatherDatabase(objFile, compilandUniqueId, coffReadFlags);
-									if (database)
-									{
-										m_coffCache->Update(objPath, database);
-									}
-
-									coff::CloseObj(objFile);
+									m_coffCache->Update(objPath, database);
 								}
 
-								return true;
+								coff::CloseObj(objFile);
+							}
+
+							return true;
 						});
 						scheduler::RunTask(task);
 
@@ -2196,10 +2265,16 @@ LiveModule::ErrorType::Enum LiveModule::Update(FileAttributeCache* fileCache, Di
 		}
 	}
 
+	// now that everything is loaded and reconstructed, transform the symbols containing anonymous namespace names
+	symbols::TransformAnonymousNamespaceSymbols(m_symbolDB, m_contributionDB, m_compilandDB, modifiedOrNewObjFiles);
+
+
 	// update the COFF cache for all compiled files
-	UpdateCoffCache(m_compiledCompilands, m_coffCache, CacheUpdate::ALL, coffReadFlags, modifiedOrNewObjFiles);
+	UpdateCoffCache(m_compiledCompilands, m_coffCache, CacheUpdate::ALL, modifiedOrNewObjFiles);
 
+	// BEGIN EPIC MOD
 	GLiveCodingServer->GetStatusChangeDelegate().ExecuteIfBound(L"Stripping COFFs...");
+	// END EPIC MOD
 
 	// strip symbols which are already part of any of the modules
 	typedef types::StringSet StrippedSymbols;
@@ -2226,12 +2301,12 @@ LiveModule::ErrorType::Enum LiveModule::Update(FileAttributeCache* fileCache, Di
 			symbols::Compiland* compiland = symbols::FindCompiland(m_compilandDB, objPath);
 			const std::wstring wideObjPath = string::ToWideString(objPath);
 
-			const uint32_t compilandUniqueId = GetCompilandId(compiland, wideObjPath.c_str(), modifiedOrNewObjFiles);
+			const uint32_t compilandUniqueId = symbols::GetCompilandId(compiland, wideObjPath.c_str(), modifiedOrNewObjFiles);
 
 			coff::ObjFile* objFile = coff::OpenObj(wideObjPath.c_str());
 			if (objFile && objFile->memoryFile)
 			{
-				coff::RawCoff* rawCoff = coff::ReadRaw(objFile, compilandUniqueId, coffReadFlags);
+				coff::RawCoff* rawCoff = coff::ReadRaw(objFile, compilandUniqueId);
 				coff::CloseObj(objFile);
 
 				if (rawCoff)
@@ -2260,7 +2335,7 @@ LiveModule::ErrorType::Enum LiveModule::Update(FileAttributeCache* fileCache, Di
 			// we need it after linking has finished  
 			{
 				const std::wstring bakPath = wideObjPath + L".bak";
-				file::Move(wideObjPath.c_str(), bakPath.c_str());
+				Filesystem::Move(wideObjPath.c_str(), bakPath.c_str());
 			}
 
 			// remove linker directives which we don't want or need.
@@ -2270,6 +2345,7 @@ LiveModule::ErrorType::Enum LiveModule::Update(FileAttributeCache* fileCache, Di
 			// *) /INCLUDE can cause symbols we already have to be pulled in again from .lib files.
 			// this leads to code and data duplication, so it must be removed for symbols which are
 			// already known to us.
+			// *) /ENTRY sets a different entry point in the patch DLL, which would cause all kinds of problems.
 			{
 				types::vector<std::string> linkerDirectives = coff::ExtractLinkerDirectives(rawCoff);
 				for (auto it = linkerDirectives.begin(); it != linkerDirectives.end(); /*nothing*/)
@@ -2301,6 +2377,11 @@ LiveModule::ErrorType::Enum LiveModule::Update(FileAttributeCache* fileCache, Di
 							continue;
 						}
 					}
+					else if (string::Contains(upperCaseDirective.c_str(), "ENTRY:"))
+					{
+						it = linkerDirectives.erase(it);
+						continue;
+					}
 
 					++it;
 				}
@@ -2477,7 +2558,17 @@ LiveModule::ErrorType::Enum LiveModule::Update(FileAttributeCache* fileCache, Di
 					// functions are always kept.
 					if ((type == coff::SymbolType::EXTERNAL_DATA) || (type == coff::SymbolType::STATIC_DATA))
 					{
-						tryStrip = true;
+						// don't try stripping MSVC JustMyCode symbols
+						const uint32_t symbolSectionIndex = coff::GetSymbolSectionIndex(rawCoff, i);
+						const ImmutableString& sectionName = coff::GetSectionName(rawCoff, symbolSectionIndex);
+						if (coff::IsMSVCJustMyCodeSection(sectionName.c_str()))
+						{
+							LC_LOG_DEV("Symbol is an MSVC JustMyCode symbol");
+						}
+						else
+						{
+							tryStrip = true;
+						}
 					}
 					else
 					{
@@ -2638,6 +2729,60 @@ LiveModule::ErrorType::Enum LiveModule::Update(FileAttributeCache* fileCache, Di
 				sectionsWithMeaningfulSymbols.insert(symbolSectionIndex);
 			}
 
+			// COFF compiled with Clang often have sections that don't contain any symbol in the symbol table, but relocations that refer to these sections.
+			// we must make sure that those sections are not stripped in case any non-stripped symbol still refers to them.
+			// TODO: we should not be doing this for each compiland, this is overhead just for Clang.
+			const symbols::Compiland* compiland = symbols::FindCompiland(m_compilandDB, objPath);
+
+			// in external build system mode, we don't have a compiland for completely new files.
+			// in this case, we just assume that the file was built with MSVC.
+			compiler::CompilerType::Enum compilerType = compiler::CompilerType::CL;
+			if (compiland)
+			{
+				const std::wstring compilerPath = GetCompilerPath(compiland);
+				compilerType = compiler::DetermineCompilerType(compilerPath.c_str());
+
+				if (compilerType == compiler::CompilerType::CLANG)
+				{
+					usesClang = true;
+
+					for (size_t i = 0u; i < relocationsPerDstSymbol.size(); ++i)
+					{
+						const auto& relocationsForSymbol = relocationsPerDstSymbol[i];
+						if (relocationsForSymbol.size() == 0u)
+						{
+							continue;
+						}
+
+						const size_t relocationCount = relocationsForSymbol.size();
+						for (size_t j = 0u; j < relocationCount; ++j)
+						{
+							const coff::Relocation* relocation = relocationsForSymbol[j].relocation;
+							if (relocation->dstIsSection)
+							{
+								// this relocation refers to a section.
+								// if the source symbol was not stripped by us, then this section still contains meaningful data and must be kept.
+								const coff::Symbol* symbol = relocationsForSymbol[j].symbol;
+								const ImmutableString& srcSymbolName = coff::GetSymbolName(coffDb, symbol);
+
+								// TODO: can this be made better?
+								if (strippedSymbolsPerCompiland.find(srcSymbolName) == strippedSymbolsPerCompiland.end())
+								{
+									// the source symbol was not stripped, so we must keep this section
+									if (relocation->dstSectionIndex >= 0)
+									{
+										sectionsWithMeaningfulSymbols.insert(static_cast<uint32_t>(relocation->dstSectionIndex));
+									}
+
+									// all relocations refer to the same destination symbol, so we can stop checking more source symbols
+									break;
+								}
+							}
+						}
+					}
+				}
+			}
+
 			const size_t sectionCount = coff::GetSectionCount(rawCoff);
 			for (size_t i = 0u; i < sectionCount; ++i)
 			{
@@ -2665,6 +2810,28 @@ LiveModule::ErrorType::Enum LiveModule::Update(FileAttributeCache* fileCache, Di
 					continue;
 				}
 
+				// don't strip sections implicitly needed by the linker.
+				// on Clang, they don't contain any symbols, but relocations to meaningful symbols.
+				if (compilerType == compiler::CompilerType::CLANG)
+				{
+					const ImmutableString& sectionName = coff::GetSectionName(rawCoff, i);
+
+					// dynamic initializers
+					if (string::StartsWith(sectionName.c_str(), ".CRT$"))
+					{
+						continue;
+					}
+					// exception data
+					else if (string::StartsWith(sectionName.c_str(), ".xdata"))
+					{
+						continue;
+					}
+					else if (string::StartsWith(sectionName.c_str(), ".pdata"))
+					{
+						continue;
+					}
+				}
+
 				if (sectionsWithMeaningfulSymbols.find(i) == sectionsWithMeaningfulSymbols.end())
 				{
 					// this section has no more meaningful symbols, remove it
@@ -2720,7 +2887,9 @@ LiveModule::ErrorType::Enum LiveModule::Update(FileAttributeCache* fileCache, Di
 		}
 	}
 
+	// BEGIN EPIC MOD
 	GLiveCodingServer->GetStatusChangeDelegate().ExecuteIfBound(L"Generating linker command line...");
+	// END EPIC MOD
 
 	telemetry::Scope generateLinkerCommandLine("Generate linker command line");
 
@@ -2747,18 +2916,12 @@ LiveModule::ErrorType::Enum LiveModule::Update(FileAttributeCache* fileCache, Di
 		const symbols::ObjPath& objPath = compilandIt->first;
 		const std::wstring& wideObjPath = string::ToWideString(objPath);
 		const symbols::Compiland* compiland = compilandIt->second;
-		const uint32_t compilandUniqueId = GetCompilandId(compiland, wideObjPath.c_str(), modifiedOrNewObjFiles);
+		const uint32_t compilandUniqueId = symbols::GetCompilandId(compiland, wideObjPath.c_str(), modifiedOrNewObjFiles);
 
 		coff::ObjFile* coffFile = coff::OpenObj(wideObjPath.c_str());
 		if (coffFile && coffFile->memoryFile)
 		{
-			// it is crucial to use coff::ReadFlags::NONE here!
-			// otherwise, we would potentially alter the names of anonymous namespaces.
-			// in VS 2015 and earlier, some symbols (e.g. templates) that use code/data in anonymous namespaces are marked
-			// as being external, and those symbols would then be forced to /INCLUDE by the linker with their *altered* name,
-			// leading to unresolved external symbols.
-			// in VS 2017 this would be no problem, because such symbols are marked static.
-			coff::ExternalSymbolDB* externalSymbolDb = coff::GatherExternalSymbolDatabase(coffFile, compilandUniqueId, coff::ReadFlags::NONE);
+			coff::ExternalSymbolDB* externalSymbolDb = coff::GatherExternalSymbolDatabase(coffFile, compilandUniqueId);
 			coff::CloseObj(coffFile);
 
 			// force the linker to include references to all external functions which we're going to hook,
@@ -2787,37 +2950,6 @@ LiveModule::ErrorType::Enum LiveModule::Update(FileAttributeCache* fileCache, Di
 		}
 	}
 
-	// weak symbols coming from libraries need special treatment.
-	// the reason for this is that due to how MSVC's linker resolves symbols, we can run into a "multiply defined symbols" error
-	// in case operator new or delete are overwritten in a translation unit that is part of static library.
-	// the dependency chain for this to happen goes roughly as follows:
-	//  OBJ: main.cpp
-	//  LIB A: operators.cpp other.cpp
-	//  LIB B: extern.cpp
-	//  LIB C: something.cpp
-	// after changing extern.cpp and linking a patch, extern.cpp needs a symbol that cannot be stripped and is contained in LIB C.
-	// LIB A is ignored because no symbols are needed right now, LIB B gets processed, the object file pulled in from LIB C needs
-	// operator new. further scanning remaining libraries, this operator gets pulled in from the runtime, but LIB C also needs
-	// a symbol from LIB A.
-	// because there are still unresolved symbols, the linker begins looking for symbols *from the start of the list* again!
-	// it now finds other.cpp in LIB A, pulls it in, but that also needs something from operators.cpp, which now introduces
-	// operator new and delete which were already pulled in from the runtime, leading to a linker error.
-	// in order to never run in any problems in this case and always pull in the correct operator new and delete from user code,
-	// we simply /INCLUDE all weak symbols found in static libraries.
-	// this works because static libraries containing overwritten operators new and delete must come first in the list of libraries,
-	// otherwise the main executable would not have linked.
-	if (appSettings::g_forceLinkWeakSymbols->GetValue())
-	{
-		const size_t symbolCount = m_weakSymbolsInLibs.size();
-		for (size_t i = 0u; i < symbolCount; ++i)
-		{
-			const ImmutableString& symbolName = m_weakSymbolsInLibs[i];
-			linkerOptions += L"/INCLUDE:";
-			linkerOptions += string::ToWideString(symbolName);
-			linkerOptions += L"\n";
-		}
-	}
-
 	// generate path for .pdb and .exe file with monotonically increasing counter
 	std::wstring pdbPath;
 	std::wstring exePath;
@@ -2840,26 +2972,26 @@ LiveModule::ErrorType::Enum LiveModule::Update(FileAttributeCache* fileCache, Di
 		// safety net: in case the PDB path does not exist, we generate the PDB filename from the module name
 		const std::wstring originalPdbPath = (m_linkerDB->pdbPath.GetLength() != 0u)
 			? string::ToWideString(m_linkerDB->pdbPath)			// we have a valid, original PDB path for the module
-			: file::RemoveExtension(m_moduleName) + L".pdb";	// create our own PDB path based on the module's name
+			: std::wstring(Filesystem::RemoveExtension(m_moduleName.c_str()).GetString()) + L".pdb";	// create our own PDB path based on the module's name
 
 		isExeOrPdbFileStillThere = false;
-		pdbPath = string::Replace(originalPdbPath, L".pdb", std::wstring(L".pdb") + patchInstanceStr);
-		exePath = string::Replace(originalPdbPath, L".pdb", std::wstring(L".exe") + patchInstanceStr);
-		const file::Attributes& pdbAttributes = file::GetAttributes(pdbPath.c_str());
-		const file::Attributes& exeAttributes = file::GetAttributes(exePath.c_str());
+		pdbPath = string::Replace(originalPdbPath, L".pdb", patchInstanceStr + std::wstring(L".pdb"));
+		exePath = string::Replace(originalPdbPath, L".pdb", patchInstanceStr + std::wstring(L".exe"));
+		const Filesystem::PathAttributes& pdbAttributes = Filesystem::GetAttributes(pdbPath.c_str());
+		const Filesystem::PathAttributes& exeAttributes = Filesystem::GetAttributes(exePath.c_str());
 
-		if (file::DoesExist(pdbAttributes))
+		if (Filesystem::DoesExist(pdbAttributes))
 		{
-			if (!file::DeleteIfExists(pdbPath.c_str()))
+			if (!Filesystem::DeleteIfExists(pdbPath.c_str()))
 			{
 				// PDB file could not be deleted
 				isExeOrPdbFileStillThere = true;
 			}
 		}
 		
-		if (file::DoesExist(exeAttributes))
+		if (Filesystem::DoesExist(exeAttributes))
 		{
-			if (!file::DeleteIfExists(exePath.c_str()))
+			if (!Filesystem::DeleteIfExists(exePath.c_str()))
 			{
 				// EXE file could not be deleted
 				isExeOrPdbFileStillThere = true;
@@ -2913,14 +3045,14 @@ LiveModule::ErrorType::Enum LiveModule::Update(FileAttributeCache* fileCache, Di
 	// add UE4-specific helper library to support NatVis visualizers when debugging
 	if (appSettings::g_ue4EnableNatVisSupport->GetValue())
 	{
-		const std::wstring& lppExePath = process::GetImagePath();
-		std::wstring libPath = file::GetDirectory(lppExePath);
+		const std::wstring& lppExePath = Process::Current::GetImagePath().GetString();
+		std::wstring libPath = Filesystem::GetDirectory(lppExePath.c_str()).GetString();
 		libPath += L"\\UE4_NatVisHelper.lib";
 
 		// check whether the .lib file exists and output a warning if not
 		{
-			const file::Attributes& attributes = file::GetAttributes(libPath.c_str());
-			if (!file::DoesExist(attributes))
+			const Filesystem::PathAttributes& attributes = Filesystem::GetAttributes(libPath.c_str());
+			if (!Filesystem::DoesExist(attributes))
 			{
 				LC_WARNING_USER("Cannot find UE4 NatVis helper library at %S", libPath.c_str());
 			}
@@ -2958,26 +3090,28 @@ LiveModule::ErrorType::Enum LiveModule::Update(FileAttributeCache* fileCache, Di
 	const std::wstring linkerPath = GetLinkerPath(m_linkerDB);
 	const std::wstring linkerWorkingDirectory = (m_linkerDB->workingDirectory.GetLength() != 0u)
 		? string::ToWideString(m_linkerDB->workingDirectory)	// we have a valid working directory
-		: file::GetDirectory(linkerPath);						// no valid working directory, take the linker directory instead
+		: Filesystem::GetDirectory(linkerPath.c_str()).GetString();						// no valid working directory, take the linker directory instead
 
 	// create a temporary file that acts as a so-called response file for the linker, and contains
 	// the whole linker command-line. this is done because the latter can get very long, longer
 	// than the limit of 32k characters.
-	const std::wstring responseFilePath = file::GenerateTempFilename();
-	file::CreateFileWithData(responseFilePath.c_str(), linkerOptions.c_str(), linkerOptions.size() * sizeof(wchar_t));
+	const Filesystem::Path responseFilePath = Filesystem::GenerateTempFilename();
+	Filesystem::CreateFileWithData(responseFilePath.GetString(), linkerOptions.c_str(), linkerOptions.size() * sizeof(wchar_t));
 
-	std::wstring linkerCommandLine = file::GetFilename(linkerPath);
+	std::wstring linkerCommandLine = Filesystem::GetFilename(linkerPath.c_str()).GetString();
 	linkerCommandLine += L" @\"";
-	linkerCommandLine += responseFilePath;
+	linkerCommandLine += responseFilePath.GetString();
 	linkerCommandLine += L"\"";
 
-	const process::Environment* linkerEnvironment = compiler::GetEnvironmentFromCache(linkerPath.c_str());
-	const void* linkerEnvironmentData = linkerEnvironment ? linkerEnvironment->data : nullptr;
+	const Process::Environment linkerEnvironment = compiler::GetEnvironmentFromCache(linkerPath.c_str());
+	const void* linkerEnvironmentData = linkerEnvironment.data;
 
+	// BEGIN EPIC MOD
 	GLiveCodingServer->GetStatusChangeDelegate().ExecuteIfBound(L"Linking patch...");
+	// END EPIC MOD
 
-	process::Context* linkerProcessContext = process::Spawn(linkerPath.c_str(), linkerWorkingDirectory.c_str(), linkerCommandLine.c_str(), linkerEnvironmentData, process::SpawnFlags::REDIRECT_STDOUT | process::SpawnFlags::NO_WINDOW);
-	const unsigned int linkerExitCode = process::Wait(linkerProcessContext);
+	Process::Context* linkerProcessContext = Process::Spawn(linkerPath.c_str(), linkerWorkingDirectory.c_str(), linkerCommandLine.c_str(), linkerEnvironmentData, Process::SpawnFlags::REDIRECT_STDOUT | Process::SpawnFlags::NO_WINDOW);
+	const unsigned int linkerExitCode = Process::Wait(linkerProcessContext);
 
 	const double linkerTime = linkScope.ReadSeconds();
 
@@ -2988,21 +3122,22 @@ LiveModule::ErrorType::Enum LiveModule::Update(FileAttributeCache* fileCache, Di
 		const std::wstring originalPath(string::ToWideString(objPath));
 		const std::wstring bakPath = originalPath + L".bak";
 
-		const file::Attributes& attributes = file::GetAttributes(bakPath.c_str());
-		if (file::DoesExist(attributes))
+		const Filesystem::PathAttributes& attributes = Filesystem::GetAttributes(bakPath.c_str());
+		if (Filesystem::DoesExist(attributes))
 		{
-			file::Delete(originalPath.c_str());
-			file::Move(bakPath.c_str(), originalPath.c_str());
+			Filesystem::Delete(originalPath.c_str());
+			Filesystem::Move(bakPath.c_str(), originalPath.c_str());
 		}
 	}
 
-	const wchar_t* linkerOutput = linkerProcessContext->stdoutData.c_str();
+	const std::wstring linkerProcessStdout = Process::GetStdOutData(linkerProcessContext);
+	const wchar_t* linkerOutput = linkerProcessStdout.c_str();
 
 	// send linker output to main executable
 	{
-		logging::LogNoFormat<logging::Channel::USER>(linkerOutput);
+		Logging::LogNoFormat<Logging::Channel::USER>(linkerOutput);
 
-		const size_t outputLength = linkerProcessContext->stdoutData.length();
+		const size_t outputLength = linkerProcessStdout.length();
 		if (outputLength != 0u)
 		{
 			if (updateType != LiveModule::UpdateType::NO_CLIENT_COMMUNICATION)
@@ -3018,9 +3153,9 @@ LiveModule::ErrorType::Enum LiveModule::Update(FileAttributeCache* fileCache, Di
 		}
 	}
 
-	process::Destroy(linkerProcessContext);
+	Process::Destroy(linkerProcessContext);
 
-	file::Delete(responseFilePath.c_str());
+	Filesystem::Delete(responseFilePath.GetString());
 
 	linkScope.End();
 
@@ -3046,10 +3181,12 @@ LiveModule::ErrorType::Enum LiveModule::Update(FileAttributeCache* fileCache, Di
 	}
 	++m_patchCounter;
 
+	// BEGIN EPIC MOD
 	GLiveCodingServer->GetStatusChangeDelegate().ExecuteIfBound(L"Preparing patch image...");
+	// END EPIC MOD
 
 	// try to load patch image
-	executable::Image* image = executable::OpenImage(exePath.c_str(), file::OpenMode::READ_AND_WRITE);
+	executable::Image* image = executable::OpenImage(exePath.c_str(), Filesystem::OpenMode::READ_WRITE);
 	if (!image)
 	{
 		LC_ERROR_USER("Cannot load patch executable %S", exePath.c_str());
@@ -3101,7 +3238,10 @@ LiveModule::ErrorType::Enum LiveModule::Update(FileAttributeCache* fileCache, Di
 			LC_LOG_DEV("Scanning memory for suitable patch location (PID: %d)", data.liveProcess->GetProcessId());
 
 			// disable the main process before scanning its memory to ensure that no operation allocates/frees virtual memory concurrently
-			process::Suspend(data.liveProcess->GetProcessHandle());
+			Process::Suspend(data.liveProcess->GetProcessHandle());
+
+			// free our page reservations in the +-2GB range of the main module
+			data.virtualMemoryRange->FreeReservedPages();
 
 			const executable::PreferredBase preferredImageBase = FindPreferredImageBase(patchImageSize, data.liveProcess->GetProcessId(), data.liveProcess->GetProcessHandle(), data.originalModuleBase);
 
@@ -3115,16 +3255,15 @@ LiveModule::ErrorType::Enum LiveModule::Update(FileAttributeCache* fileCache, Di
 			if (imageNeedsToBeCopied)
 			{
 				// this image needs to be copied. create a new name based on the process ID, which must be unique
-				rebasedExePath += L"_";
-				rebasedExePath += std::to_wstring(data.liveProcess->GetProcessId());
-				file::Copy(exePath.c_str(), rebasedExePath.c_str());
+				rebasedExePath = string::Replace(rebasedExePath, L".exe", std::wstring(L"_PID_") + std::to_wstring(+data.liveProcess->GetProcessId()) + L".exe");
+				Filesystem::Copy(exePath.c_str(), rebasedExePath.c_str());
 				wcscpy_s(cmd.path, rebasedExePath.c_str());
 			}
 			
 			if (imageNeedsToBeRebased)
 			{
 				// rebase the patch image to its preferred base address
-				executable::Image* rebasedImage = executable::OpenImage(rebasedExePath.c_str(), file::OpenMode::READ_AND_WRITE);
+				executable::Image* rebasedImage = executable::OpenImage(rebasedExePath.c_str(), Filesystem::OpenMode::READ_WRITE);
 				LC_LOG_DEV("Rebasing patch executable to image base 0x%" PRIX64 " (PID: %d)", preferredImageBase, data.liveProcess->GetProcessId());
 				executable::RebaseImage(rebasedImage, preferredImageBase);
 				executable::CloseImage(rebasedImage);
@@ -3136,7 +3275,7 @@ LiveModule::ErrorType::Enum LiveModule::Update(FileAttributeCache* fileCache, Di
 			// will allocate virtual memory at the patch's preferred image base, possibly rendering the patch unusable because
 			// it cannot be loaded.
 			// the chances of that happening are *very* rare though, and we can always load the next patch then.
-			process::Resume(data.liveProcess->GetProcessHandle());
+			Process::Resume(data.liveProcess->GetProcessHandle());
 #endif
 
 			data.liveProcess->GetPipe()->SendCommandAndWaitForAck(cmd, nullptr, 0u);
@@ -3148,6 +3287,15 @@ LiveModule::ErrorType::Enum LiveModule::Update(FileAttributeCache* fileCache, Di
 		}
 	}
 
+#if LC_64_BIT
+	// immediately reserve pages in the +-2GB range of the main module again
+	for (size_t i = 0u; i < processCount; ++i)
+	{
+		const PerProcessData& data = processData[i];
+		data.virtualMemoryRange->ReservePages();
+	}
+#endif
+
 	if (processCount != loadedPatches.size())
 	{
 		// communication with the client broke down while trying to load the patch, bail out
@@ -3209,7 +3357,9 @@ LiveModule::ErrorType::Enum LiveModule::Update(FileAttributeCache* fileCache, Di
 	}
 
 
+	// BEGIN EPIC MOD
 	GLiveCodingServer->GetStatusChangeDelegate().ExecuteIfBound(L"Loading patch PDB...");
+	// END EPIC MOD
 
 	LC_LOG_DEV("Loading patch PDB");
 
@@ -3270,7 +3420,11 @@ LiveModule::ErrorType::Enum LiveModule::Update(FileAttributeCache* fileCache, Di
 			options |= symbols::CompilandOptions::TRACK_OBJ_ONLY;
 		}
 
-		auto db = symbols::GatherCompilands(localProvider, localDiaCompilandDb, GetAmalgamatedSplitThreshold(), options);
+		// TODO: what if somebody uses FASTBuild? how to get the new dependencies?
+		// does this step even have to be done? seems like this might not even be needed in ANY CASE.
+		// -> this has to be done to figure out new #include files for existing compilands!
+		// => but it should be split in two: figuring out compilands, and figuring out include files.
+		symbols::CompilandDB* db = symbols::GatherCompilands(localProvider, localDiaCompilandDb, GetAmalgamatedSplitThreshold(), options);
 
 		symbols::DestroyDiaCompilandDB(localDiaCompilandDb);
 		symbols::Close(localProvider);
@@ -3352,7 +3506,9 @@ LiveModule::ErrorType::Enum LiveModule::Update(FileAttributeCache* fileCache, Di
 
 
 
+	// BEGIN EPIC MOD
 	GLiveCodingServer->GetStatusChangeDelegate().ExecuteIfBound(L"Updating COFF cache...");
+	// END EPIC MOD
 
 	{
 		LC_LOG_DEV("Updating COFF cache for new patch compilands");
@@ -3360,7 +3516,7 @@ LiveModule::ErrorType::Enum LiveModule::Update(FileAttributeCache* fileCache, Di
 		// update the COFF cache for new patch compilands.
 		// there may be files for which we don't have a database yet, even though we updated the database for all compiled files.
 		// this can happen when a new .obj that is part of a library is linked in for the first time.
-		types::vector<symbols::ObjPath> updatedCoffs = UpdateCoffCache(patch_compilandDB->compilands, m_coffCache, CacheUpdate::NON_EXISTANT, coffReadFlags, modifiedOrNewObjFiles);
+		types::vector<symbols::ObjPath> updatedCoffs = UpdateCoffCache(patch_compilandDB->compilands, m_coffCache, CacheUpdate::NON_EXISTANT, modifiedOrNewObjFiles);
 
 
 		// similarly, reconstruct symbols and dynamic initializers for new .obj that have been pulled in for the first time.
@@ -3370,7 +3526,7 @@ LiveModule::ErrorType::Enum LiveModule::Update(FileAttributeCache* fileCache, Di
 		{
 			LC_LOG_INDENT_DEV;
 
-			executable::Image* originalImage = executable::OpenImage(m_moduleName.c_str(), file::OpenMode::READ_ONLY);
+			executable::Image* originalImage = executable::OpenImage(m_moduleName.c_str(), Filesystem::OpenMode::READ);
 			executable::ImageSectionDB* originalImageSections = executable::GatherImageSectionDB(originalImage);
 
 			types::StringSet noSymbolsToIgnore;
@@ -3410,12 +3566,16 @@ LiveModule::ErrorType::Enum LiveModule::Update(FileAttributeCache* fileCache, Di
 		}
 	}
 
+	// new symbols have been reconstructed, so transform the symbols containing anonymous namespace names
+	symbols::TransformAnonymousNamespaceSymbols(m_symbolDB, m_contributionDB, m_compilandDB, modifiedOrNewObjFiles);
 
 
+	// BEGIN EPIC MOD
 	GLiveCodingServer->GetStatusChangeDelegate().ExecuteIfBound(L"Reconstructing patch symbols...");
+	// END EPIC MOD
 
 	// reconstruct symbols for all compilands that are part of the new patch executable
-	executable::Image* patchImage = executable::OpenImage(exePath.c_str(), file::OpenMode::READ_ONLY);
+	executable::Image* patchImage = executable::OpenImage(exePath.c_str(), Filesystem::OpenMode::READ);
 	executable::ImageSectionDB* patchImageSections = executable::GatherImageSectionDB(patchImage);
 
 	// gather the dynamic initializers and remaining symbols by walking the module
@@ -3476,6 +3636,11 @@ LiveModule::ErrorType::Enum LiveModule::Update(FileAttributeCache* fileCache, Di
 	symbols::Close(patchSymbolProvider);
 
 
+	// now that everything is loaded and reconstructed in the patch, transform the symbols containing anonymous namespace names.
+	// note that we need to used the merged main compiland database here.
+	TransformAnonymousNamespaceSymbols(patch_symbolDB, patch_contributionDB, m_compilandDB, modifiedOrNewObjFiles);
+
+
 	// store the new databases into the module cache
 	ModulePatch* compiledModulePatch = nullptr;
 	const size_t token = m_moduleCache->Insert(patch_symbolDB, patch_contributionDB, patch_compilandDB, patch_thunkDB, patch_imageSectionDB);
@@ -3512,11 +3677,11 @@ LiveModule::ErrorType::Enum LiveModule::Update(FileAttributeCache* fileCache, Di
 				{
 					const ModuleCache::ProcessData& hookProcessData = hookData.data->processes[p];
 
-					const unsigned int pid = hookProcessData.processId;
+					const Process::Id pid = hookProcessData.processId;
 					void* moduleBase = hookProcessData.moduleBase;
 					const DuplexPipe* pipe = hookProcessData.pipe;
 
-					LC_LOG_USER("Calling pre-patch hooks (PID: %d)", pid);
+					LC_LOG_USER("Calling pre-patch hooks (PID: %d)", +pid);
 					pipe->SendCommandAndWaitForAck(MakeCallHooksCommand(hook::Type::PREPATCH, moduleBase, hookData), nullptr, 0u);
 				}
 
@@ -3527,7 +3692,9 @@ LiveModule::ErrorType::Enum LiveModule::Update(FileAttributeCache* fileCache, Di
 
 
 
+	// BEGIN EPIC MOD
 	GLiveCodingServer->GetStatusChangeDelegate().ExecuteIfBound(L"Patching relocations...");
+	// END EPIC MOD
 
 	LC_LOG_DEV("Patching relocations before calling entry point");
 
@@ -3558,6 +3725,13 @@ LiveModule::ErrorType::Enum LiveModule::Update(FileAttributeCache* fileCache, Di
 				continue;
 			}
 
+			const std::wstring wideObjPath = string::ToWideString(objPath);
+
+			// note that we need to take the original compiland here, to make sure that single split-files get assigned
+			// the same unique ID as the corresponding amalgamated file.
+			symbols::Compiland* compiland = symbols::FindCompiland(m_compilandDB, objPath);
+			const uint32_t compilandUniqueId = symbols::GetCompilandId(compiland, wideObjPath.c_str(), modifiedOrNewObjFiles);
+
 			const types::StringSet& strippedSymbols = strippedSymbolsPerCompiland[objPath];
 			const types::StringSet& forceStrippedSymbols = forceStrippedSymbolsPerCompiland[objPath];
 
@@ -3569,7 +3743,7 @@ LiveModule::ErrorType::Enum LiveModule::Update(FileAttributeCache* fileCache, Di
 
 				// check if the patch knows this symbol.
 				// if not, it has probably been stripped and there is no need to walk all its relocations.
-				const ImmutableString& symbolName = coff::GetSymbolName(coffDb, symbol);
+				const ImmutableString& symbolName = symbols::TransformAnonymousNamespacePattern(coff::GetSymbolName(coffDb, symbol), compilandUniqueId);
 				const symbols::Symbol* realSymbol = symbols::FindSymbolByName(patch_symbolDB, symbolName);
 				if (!realSymbol)
 				{
@@ -3600,13 +3774,30 @@ LiveModule::ErrorType::Enum LiveModule::Update(FileAttributeCache* fileCache, Di
 				for (size_t j = 0u; j < relocationCount; ++j)
 				{
 					const coff::Relocation* relocation = symbol->relocations[j];
+					if (relocation->dstIsSection)
+					{
+						// don't patch relocations to sections
+						continue;
+					}
 
-					const ImmutableString& dstSymbolName = coff::GetRelocationDstSymbolName(coffDb, relocation);
+					// ignore relocations to symbols in .msvcjmc (MSVC JustMyCode) sections
+					if (relocation->dstSectionIndex >= 0)
+					{
+						const uint32_t index = static_cast<uint32_t>(relocation->dstSectionIndex);
+						const coff::Section& section = coffDb->sections[index];
+						if (coff::IsMSVCJustMyCodeSection(section.name.c_str()))
+						{
+							LC_LOG_DEV("Ignoring relocation to symbol in section %s", section.name.c_str());
+							continue;
+						}
+					}
+
+					ImmutableString dstSymbolName = symbols::TransformAnonymousNamespacePattern(GetRelocationDstSymbolName(coffDb, relocation), compilandUniqueId);
 					const bool refersToDataSymbol = !coff::IsFunctionSymbol(coff::GetRelocationDstSymbolType(relocation));
 					const bool refersToStrippedSymbol = (strippedSymbols.find(dstSymbolName) != strippedSymbols.end());
 					if (refersToDataSymbol || refersToStrippedSymbol)
 					{
-						const relocations::Record& relocationRecord = relocations::PatchRelocation(relocation, coffDb, forceStrippedSymbols, m_moduleCache, symbolName, realSymbol, token, &loadedPatches[0]);
+						const relocations::Record& relocationRecord = relocations::PatchRelocation(relocation, coffDb, forceStrippedSymbols, m_moduleCache, symbolName, dstSymbolName, realSymbol, token, &loadedPatches[0]);
 						if (relocations::IsValidRecord(relocationRecord))
 						{
 							compiledModulePatch->RegisterPreEntryPointRelocation(relocationRecord);
@@ -3624,8 +3815,121 @@ LiveModule::ErrorType::Enum LiveModule::Update(FileAttributeCache* fileCache, Di
 
 
 
+	// BEGIN EPIC MOD
+	GLiveCodingServer->GetStatusChangeDelegate().ExecuteIfBound(L"Patching public functions in lib compilands...");
+	// END EPIC MOD
+
+	LC_LOG_DEV("Patching public functions in lib compilands");
+	{
+		telemetry::Scope patchingFunctionsScope("Patching functions");
+
+		uint32_t functionsPatchedCount = 0u;
+		size_t functionsCount = patch_symbolDB->patchableFunctionSymbols.size();
+
+		// functions in lib compilands cannot have changed, per definition. but there can be code linked in from libraries
+		// that calls these functions, therefore they need to be patched to their original function, otherwise
+		// there would be functions working on new data. this can be the case when linking with static libraries we don't have
+		// the .obj files for.
+		LC_LOG_INDENT_DEV;
+
+		for (auto it : patch_symbolDB->patchableFunctionSymbols)
+		{
+			const symbols::Symbol* symbol = it;
+			const ImmutableString& functionName = symbol->name;
+
+			// never patch the entry point, we need to call that later on
+			if (symbol->rva == entryPointRva)
+			{
+				LC_LOG_DEV("Ignoring entry point function %s", functionName.c_str());
+				continue;
+			}
+
+			// note that when patching new functions to original ones, the same rules as for patching relocations apply,
+			// i.e. not all functions should be patched.
+			if (symbols::IsExceptionRelatedSymbol(functionName))
+			{
+				LC_LOG_DEV("Ignoring exception-related function %s", functionName.c_str());
+				continue;
+			}
+			else if (symbols::IsRuntimeCheckRelatedSymbol(functionName))
+			{
+				LC_LOG_DEV("Ignoring runtime check function %s", functionName.c_str());
+				continue;
+			}
+			else if (symbols::IsSdlCheckRelatedSymbol(functionName))
+			{
+				LC_LOG_DEV("Ignoring SDL check function %s", functionName.c_str());
+				continue;
+			}
+
+			// check whether the function is at least 5 bytes long to consider it for patching
+			const symbols::Contribution* contribution = symbols::FindContributionByRVA(patch_contributionDB, symbol->rva);
+			if (!contribution)
+			{
+				LC_ERROR_DEV("Ignoring function %s because its contribution cannot be found", functionName.c_str());
+				continue;
+			}
+
+			if (contribution->size < 5u)
+			{
+				LC_LOG_DEV("Ignoring function %s that is only %d bytes long", functionName.c_str(), contribution->size);
+				continue;
+			}
+
+			const ModuleCache::FindSymbolData& originalData = m_moduleCache->FindSymbolByName(token, functionName);
+			if (!originalData.symbol)
+			{
+				LC_LOG_DEV("Ignoring new function %s", functionName.c_str());
+				continue;
+			}
+
+			// if the original function to be patched did not come from a compiland, it cannot possibly have changed and
+			// therefore needs to be patched.
+			const symbols::Contribution* originalContribution = symbols::FindContributionByRVA(originalData.data->contributionDb, originalData.symbol->rva);
+			if (originalContribution)
+			{
+				const ImmutableString& compilandName = symbols::GetContributionCompilandName(originalData.data->contributionDb, originalContribution);
+				const symbols::Compiland* originalCompiland = symbols::FindCompiland(originalData.data->compilandDb, compilandName);
+				if (!originalCompiland)
+				{
+					const size_t moduleProcessCount = originalData.data->processes.size();
+					for (size_t p = 0u; p < moduleProcessCount; ++p)
+					{
+						++functionsPatchedCount;
+
+						const Process::Id pid = originalData.data->processes[p].processId;
+						void* moduleBase = originalData.data->processes[p].moduleBase;
+						Process::Handle processHandle = originalData.data->processes[p].processHandle;
+
+						const symbols::Symbol* patchSymbol = symbol;
+
+						char* srcAddress = pointer::Offset<char*>(loadedPatches[p], patchSymbol->rva);
+						char* destAddress = pointer::Offset<char*>(moduleBase, originalData.symbol->rva);
+
+						LC_LOG_DEV("Patching library function %s at 0x%p (0x%X) (PID: %d)", functionName.c_str(), moduleBase, patchSymbol->rva, +pid);
+
+						const functions::LibraryRecord& record = functions::PatchLibraryFunction(srcAddress, destAddress,
+							patchSymbol->rva, originalData.symbol->rva,
+							contribution, processHandle, originalData.data->index);
+						if (functions::IsValidRecord(record))
+						{
+							compiledModulePatch->RegisterLibraryFunctionPatch(record);
+						}
+					}
+				}
+			}
+		}
+
+		LC_LOG_TELEMETRY("Patched %d of %d functions in %.3fms (avg: %.3fus)", functionsPatchedCount, functionsCount,
+			patchingFunctionsScope.ReadMilliSeconds(), (patchingFunctionsScope.ReadMicroSeconds() / functionsPatchedCount));
+	}
+
+
+
 	// now that the .dll is loaded and symbols have been relocated, finally patch the dynamic initializers
+	// BEGIN EPIC MOD
 	GLiveCodingServer->GetStatusChangeDelegate().ExecuteIfBound(L"Patching dynamic initializers...");
+	// END EPIC MOD
 	{
 		const size_t count = initializerDb.dynamicInitializers.size();
 
@@ -3649,7 +3953,7 @@ LiveModule::ErrorType::Enum LiveModule::Update(FileAttributeCache* fileCache, Di
 					LC_LOG_DEV("Patching dynamic initializer symbol %s at RVA 0x%X (PID: %d)", name.c_str(), rva, data.liveProcess->GetProcessId());
 
 					void* initializerAddress = pointer::Offset<void*>(loadedPatches[p], rva);
-					process::WriteProcessMemory(data.liveProcess->GetProcessHandle(), initializerAddress, nullptr);
+					Process::WriteProcessMemory(data.liveProcess->GetProcessHandle(), initializerAddress, nullptr);
 				}
 
 				compiledModulePatch->RegisterPatchedDynamicInitializer(rva);
@@ -3659,6 +3963,74 @@ LiveModule::ErrorType::Enum LiveModule::Update(FileAttributeCache* fileCache, Di
 				LC_WARNING_DEV("Cannot find symbol %s in original executable", name.c_str());
 			}
 		}
+
+		// on Clang, CRT sections that store individual dynamic initializers don't expose any symbols in the COFF,
+		// but only relocations to the initializers. therefore, the initializer database is empty.
+		// however, the dynamic initializers should have been reconstructed at this point, so we know their RVA.
+		// in this case, we scan all symbols from __xc_a to __xc_z, check to which symbol RVA they refer, fetch the symbol at this RVA,
+		// and check if this symbol is already known to us. if so, we need to disable the corresponding initializer.
+		if (usesClang)
+		{
+			const symbols::Symbol* firstInitializerSymbol = symbols::FindSymbolByName(patch_symbolDB, ImmutableString(LC_IDENTIFIER("__xc_a")));
+			const symbols::Symbol* lastInitializerSymbol = symbols::FindSymbolByName(patch_symbolDB, ImmutableString(LC_IDENTIFIER("__xc_z")));
+			if (firstInitializerSymbol && lastInitializerSymbol)
+			{
+				executable::Image* thisImage = executable::OpenImage(exePath.c_str(), Filesystem::OpenMode::READ);
+				executable::ImageSectionDB* sectionDb = executable::GatherImageSectionDB(thisImage);
+
+				// this is defined in the CRT, which also defines all the special sections
+				// TODO: move this to somewhere else, e.g. Windows Internals
+				typedef _PVFV DynamicInitializer;
+
+				// the first symbol is always __xc_a, which we are not interested in.
+				// similarly, the last symbol is always __xc_z, which we are also not interested in.
+				const uint32_t firstRva = firstInitializerSymbol->rva + sizeof(DynamicInitializer);
+				const uint32_t lastRva = lastInitializerSymbol->rva - sizeof(DynamicInitializer);
+
+				for (uint32_t rva = firstRva; rva <= lastRva; rva += sizeof(DynamicInitializer))
+				{
+#if LC_64_BIT
+					const uint64_t initializerAddress = executable::ReadFromImage<uint64_t>(thisImage, sectionDb, rva);
+#else
+					const uint32_t initializerAddress = executable::ReadFromImage<uint32_t>(thisImage, sectionDb, rva);
+#endif
+
+					// the relocations from "$initializer$" to a dynamic initializer are always absolute, so its
+					// easy to reconstruct the dynamic initializer's RVA.
+					const uint32_t dynamicInitializerRva = static_cast<uint32_t>(initializerAddress - executable::GetPreferredBase(thisImage));
+					const symbols::Symbol* symbol = symbols::FindSymbolByRVA(patch_symbolDB, dynamicInitializerRva);
+					if (symbol)
+					{
+						// check if we already know a symbol with this name
+						const ModuleCache::FindSymbolData& findData = m_moduleCache->FindSymbolByName(token, symbol->name);
+
+						const bool symbolFound = (findData.symbol != nullptr);
+
+						// special case for amalgamated files: Clang assigns dynamic initializer functions based on filenames.
+						// for amalgamated files, the initializers will already have executed, but the initializers in the
+						// split files will get a different name.
+						const bool isClangDynamicInitializer = (string::Contains(symbol->name.c_str(), "_GLOBAL__sub_I_"));
+						if (symbolFound || isClangDynamicInitializer)
+						{
+							// this symbol already exists, so it must have been called during initialization earlier.
+							// patch the dynamic initializer in all processes.
+							for (size_t p = 0u; p < processCount; ++p)
+							{
+								const PerProcessData& data = processData[p];
+
+								void* address = pointer::Offset<void*>(loadedPatches[p], rva);
+								Process::WriteProcessMemory(data.liveProcess->GetProcessHandle(), address, nullptr);
+							}
+
+							compiledModulePatch->RegisterPatchedDynamicInitializer(rva);
+						}
+					}
+				}
+
+				executable::DestroyImageSectionDB(sectionDb);
+				executable::CloseImage(thisImage);
+			}
+		}
 	}
 
 	
@@ -3728,7 +4100,9 @@ LiveModule::ErrorType::Enum LiveModule::Update(FileAttributeCache* fileCache, Di
 
 	// now that relocations are done, it is safe to call the entry point.
 	// restore the original entry point and tell the process to call it.
+	// BEGIN EPIC MOD
 	GLiveCodingServer->GetStatusChangeDelegate().ExecuteIfBound(L"Restoring and calling entry point...");
+	// END EPIC MOD
 	{
 		// disable user entry point DllMain (if it exists).
 		// the DllMain function is named differently depending on the architecture.
@@ -3787,7 +4161,9 @@ LiveModule::ErrorType::Enum LiveModule::Update(FileAttributeCache* fileCache, Di
 
 
 	// dynamic initializers have run, patch the remaining relocations
+	// BEGIN EPIC MOD
 	GLiveCodingServer->GetStatusChangeDelegate().ExecuteIfBound(L"Patching remaining relocations...");
+	// END EPIC MOD
 	{
 		telemetry::Scope patchingRelocationsScope("Patching remaining relocations");
 
@@ -3810,6 +4186,13 @@ LiveModule::ErrorType::Enum LiveModule::Update(FileAttributeCache* fileCache, Di
 				continue;
 			}
 
+			const std::wstring wideObjPath = string::ToWideString(objPath);
+
+			// note that we need to take the original compiland here, to make sure that single split-files get assigned
+			// the same unique ID as the corresponding amalgamated file.
+			symbols::Compiland* compiland = symbols::FindCompiland(m_compilandDB, objPath);
+			const uint32_t compilandUniqueId = symbols::GetCompilandId(compiland, wideObjPath.c_str(), modifiedOrNewObjFiles);
+
 			const types::StringSet& strippedSymbols = strippedSymbolsPerCompiland[objPath];
 			const types::StringSet& forceStrippedSymbols = forceStrippedSymbolsPerCompiland[objPath];
 
@@ -3852,14 +4235,32 @@ LiveModule::ErrorType::Enum LiveModule::Update(FileAttributeCache* fileCache, Di
 				for (size_t j = 0u; j < relocationCount; ++j)
 				{
 					const coff::Relocation* relocation = symbol->relocations[j];
-					const ImmutableString& dstSymbolName = coff::GetRelocationDstSymbolName(coffDb, relocation);
+					if (relocation->dstIsSection)
+					{
+						// don't patch relocations to sections
+						continue;
+					}
+
+					// ignore relocations to symbols in .msvcjmc (MSVC JustMyCode) sections
+					if (relocation->dstSectionIndex >= 0)
+					{
+						const uint32_t index = static_cast<uint32_t>(relocation->dstSectionIndex);
+						const coff::Section& section = coffDb->sections[index];
+						if (coff::IsMSVCJustMyCodeSection(section.name.c_str()))
+						{
+							LC_LOG_DEV("Ignoring relocation to symbol in section %s", section.name.c_str());
+							continue;
+						}
+					}
+
+					ImmutableString dstSymbolName = symbols::TransformAnonymousNamespacePattern(GetRelocationDstSymbolName(coffDb, relocation), compilandUniqueId);
 
 					// relocations to data symbols and stripped symbols have already been done
 					const bool refersToFunctionSymbol = coff::IsFunctionSymbol(coff::GetRelocationDstSymbolType(relocation));
 					const bool refersToStrippedSymbol = (strippedSymbols.find(dstSymbolName) != strippedSymbols.end());
 					if (refersToFunctionSymbol && !refersToStrippedSymbol)
 					{
-						const relocations::Record& relocationRecord = relocations::PatchRelocation(relocation, coffDb, forceStrippedSymbols, m_moduleCache, symbolName, realSymbol, token, &loadedPatches[0]);
+						const relocations::Record& relocationRecord = relocations::PatchRelocation(relocation, coffDb, forceStrippedSymbols, m_moduleCache, symbolName, dstSymbolName, realSymbol, token, &loadedPatches[0]);
 						if (relocations::IsValidRecord(relocationRecord))
 						{
 							compiledModulePatch->RegisterPostEntryPointRelocation(relocationRecord);
@@ -3877,13 +4278,15 @@ LiveModule::ErrorType::Enum LiveModule::Update(FileAttributeCache* fileCache, Di
 
 
 
+	// BEGIN EPIC MOD
 	GLiveCodingServer->GetStatusChangeDelegate().ExecuteIfBound(L"Patching functions...");
+	// END EPIC MOD
 
 	// suspend the main processes before patching functions, because they might not use synchronization points.
 	for (size_t p = 0u; p < processCount; ++p)
 	{
 		const PerProcessData& data = processData[p];
-		process::Suspend(data.liveProcess->GetProcessHandle());
+		Process::Suspend(data.liveProcess->GetProcessHandle());
 	}
 
 
@@ -3908,8 +4311,6 @@ LiveModule::ErrorType::Enum LiveModule::Update(FileAttributeCache* fileCache, Di
 		uint32_t functionsPatchedCount = 0u;
 		size_t functionsCount = 0u;
 
-		types::StringSet patchedFunctions;
-
 		// we deliberately do not hook functions in lib compilands because they cannot have changed, per definition.
 		// they are part of a static library that won't be recompiled during a Live++ session.
 		for (auto it = patch_compilandDB->compilands.begin(); it != patch_compilandDB->compilands.end(); ++it)
@@ -3926,6 +4327,10 @@ LiveModule::ErrorType::Enum LiveModule::Update(FileAttributeCache* fileCache, Di
 				continue;
 			}
 
+			const std::wstring& wideObjPath = string::ToWideString(objPath);
+			const symbols::Compiland* compiland = it->second;
+			const uint32_t compilandUniqueId = symbols::GetCompilandId(compiland, wideObjPath.c_str(), modifiedOrNewObjFiles);
+
 			for (size_t i = 0u; i < coffDb->symbols.size(); ++i)
 			{
 				const coff::Symbol* symbol = coffDb->symbols[i];
@@ -3936,7 +4341,7 @@ LiveModule::ErrorType::Enum LiveModule::Update(FileAttributeCache* fileCache, Di
 
 				++functionsCount;
 
-				const ImmutableString& functionName = coff::GetSymbolName(coffDb, symbol);
+				const ImmutableString& functionName = symbols::TransformAnonymousNamespacePattern(coff::GetSymbolName(coffDb, symbol), compilandUniqueId);
 				if (symbols::IsExceptionRelatedSymbol(functionName))
 				{
 					LC_LOG_DEV("Ignoring exception-related function %s", functionName.c_str());
@@ -3963,16 +4368,14 @@ LiveModule::ErrorType::Enum LiveModule::Update(FileAttributeCache* fileCache, Di
 				if (originalContribution)
 				{
 					const ImmutableString& compilandName = symbols::GetContributionCompilandName(originalData.data->contributionDb, originalContribution);
-					const symbols::Compiland* compiland = symbols::FindCompiland(originalData.data->compilandDb, compilandName);
-					if (!compiland)
+					const symbols::Compiland* originalCompiland = symbols::FindCompiland(originalData.data->compilandDb, compilandName);
+					if (!originalCompiland)
 					{
 						LC_LOG_DEV("Ignoring function %s originally contributed from lib compiland %s", functionName.c_str(), compilandName.c_str());
 						continue;
 					}
 				}
 
-				patchedFunctions.insert(functionName);
-
 				const size_t moduleProcessCount = originalData.data->processes.size();
 				for (size_t p = 0u; p < moduleProcessCount; ++p)
 				{
@@ -3980,9 +4383,9 @@ LiveModule::ErrorType::Enum LiveModule::Update(FileAttributeCache* fileCache, Di
 
 					const ModuleCache::ProcessData& hookProcessData = originalData.data->processes[p];
 
-					const unsigned int pid = hookProcessData.processId;
+					const Process::Id pid = hookProcessData.processId;
 					void* moduleBase = hookProcessData.moduleBase;
-					process::Handle processHandle = hookProcessData.processHandle;
+					Process::Handle processHandle = hookProcessData.processHandle;
 
 					char* originalAddress = pointer::Offset<char*>(moduleBase, originalData.symbol->rva);
 					char* patchAddress = pointer::Offset<char*>(loadedPatches[p], patchSymbol->rva);
@@ -3999,103 +4402,13 @@ LiveModule::ErrorType::Enum LiveModule::Update(FileAttributeCache* fileCache, Di
 				}
 			}
 		}
-
-		{
-			// functions in lib compilands cannot have changed, per definition. but there can be code linked in from libraries
-			// that calls these functions, therefore they need to be patched to their original function, otherwise
-			// there would be functions working on new data.
-			LC_LOG_DEV("Patching public functions in lib compilands");
-			LC_LOG_INDENT_DEV;
-
-			for (auto it : patch_symbolDB->patchableFunctionSymbols)
-			{
-				const symbols::Symbol* symbol = it;
-				const ImmutableString& functionName = symbol->name;
-
-				++functionsCount;
-
-				// don't patch functions that were already patched from original to new code
-				if (patchedFunctions.find(functionName) != patchedFunctions.end())
-				{
-					LC_LOG_DEV("Ignoring function %s that was patched already", functionName.c_str());
-					continue;
-				}
-
-				// note that when patching new functions to original ones, the same rules as for patching relocations apply,
-				// i.e. not all functions should be patched.
-				if (symbols::IsExceptionRelatedSymbol(functionName))
-				{
-					LC_LOG_DEV("Ignoring exception-related function %s", functionName.c_str());
-					continue;
-				}
-				else if (symbols::IsRuntimeCheckRelatedSymbol(functionName))
-				{
-					LC_LOG_DEV("Ignoring runtime check function %s", functionName.c_str());
-					continue;
-				}
-				else if (symbols::IsSdlCheckRelatedSymbol(functionName))
-				{
-					LC_LOG_DEV("Ignoring SDL check function %s", functionName.c_str());
-					continue;
-				}
-
-				// check whether the function is at least 5 bytes long to consider it for patching
-				const symbols::Contribution* contribution = symbols::FindContributionByRVA(patch_contributionDB, symbol->rva);
-				if (!contribution)
-				{
-					LC_ERROR_DEV("Ignoring function %s because its contribution cannot be found", functionName.c_str());
-					continue;
-				}
-
-				if (contribution->size < 5u)
-				{
-					LC_LOG_DEV("Ignoring function %s that is only %d bytes long", functionName.c_str(), contribution->size);
-					continue;
-				}
-
-				const ModuleCache::FindSymbolData& originalData = m_moduleCache->FindSymbolByName(token, functionName);
-				if (!originalData.symbol)
-				{
-					LC_LOG_DEV("Ignoring new function %s", functionName.c_str());
-					continue;
-				}
-
-				const size_t moduleProcessCount = originalData.data->processes.size();
-				for (size_t p = 0u; p < moduleProcessCount; ++p)
-				{
-					++functionsPatchedCount;
-
-					const unsigned int pid = originalData.data->processes[p].processId;
-					void* moduleBase = originalData.data->processes[p].moduleBase;
-					process::Handle processHandle = originalData.data->processes[p].processHandle;
-
-					const symbols::Symbol* patchSymbol = symbol;
-
-					char* srcAddress = pointer::Offset<char*>(loadedPatches[p], patchSymbol->rva);
-					char* destAddress = pointer::Offset<char*>(moduleBase, originalData.symbol->rva);
-
-					LC_LOG_DEV("Patching function %s at 0x%p (0x%X) (PID: %d)", functionName.c_str(), moduleBase, patchSymbol->rva, pid);
-
-					const functions::LibraryRecord& record = functions::PatchLibraryFunction(srcAddress, destAddress, 
-						patchSymbol->rva, originalData.symbol->rva, 
-						contribution, processHandle, originalData.data->index);
-					if (functions::IsValidRecord(record))
-					{
-						compiledModulePatch->RegisterLibraryFunctionPatch(record);
-					}
-				}
-			}
-		}
-
-		LC_LOG_TELEMETRY("Patched %d of %d functions in %.3fms (avg: %.3fus)", functionsPatchedCount, functionsCount,
-			patchingFunctionsScope.ReadMilliSeconds(), (patchingFunctionsScope.ReadMicroSeconds() / functionsPatchedCount));
 	}
 
 	// resume the main processes again
 	for (size_t p = 0u; p < processCount; ++p)
 	{
 		const PerProcessData& data = processData[p];
-		process::Resume(data.liveProcess->GetProcessHandle());
+		Process::Resume(data.liveProcess->GetProcessHandle());
 	}
 
 	{
@@ -4109,11 +4422,11 @@ LiveModule::ErrorType::Enum LiveModule::Update(FileAttributeCache* fileCache, Di
 				const size_t count = hookData.data->processes.size();
 				for (size_t p = 0u; p < count; ++p)
 				{
-					const unsigned int pid = hookData.data->processes[p].processId;
+					const Process::Id pid = hookData.data->processes[p].processId;
 					void* moduleBase = hookData.data->processes[p].moduleBase;
 					const DuplexPipe* pipe = hookData.data->processes[p].pipe;
 
-					LC_LOG_USER("Calling post-patch hooks (PID: %d)", pid);
+					LC_LOG_USER("Calling post-patch hooks (PID: %d)", +pid);
 					pipe->SendCommandAndWaitForAck(MakeCallHooksCommand(hook::Type::POSTPATCH, moduleBase, hookData), nullptr, 0u);
 				}
 
@@ -4172,23 +4485,22 @@ bool LiveModule::InstallCompiledPatches(LiveProcess* liveProcess, void* original
 
 	telemetry::Scope wholeScope("Installing patches");
 
-	process::Handle processHandle = liveProcess->GetProcessHandle();
-	const unsigned int processId = liveProcess->GetProcessId();
+	Process::Handle processHandle = liveProcess->GetProcessHandle();
+	const Process::Id processId = liveProcess->GetProcessId();
 	const DuplexPipe* pipe = liveProcess->GetPipe();
 
 	for (auto modulePatch : m_compiledModulePatches)
 	{
 		const std::wstring& originalExePath = modulePatch->GetExePath();
 
-		LC_LOG_USER("Installing patch %S (PID: %d)", originalExePath.c_str(), processId);
+		LC_LOG_USER("Installing patch %S (PID: %d)", originalExePath.c_str(), +processId);
 
 		// this image needs to be copied because it is loaded already.
 		// create a new name based on the process ID, which must be unique.
 		std::wstring exePath = originalExePath;
 		{
-			exePath += L"_";
-			exePath += std::to_wstring(processId);
-			file::Copy(originalExePath.c_str(), exePath.c_str());
+			exePath = string::Replace(exePath, L".exe", std::wstring(L"_PID_") + std::to_wstring(+processId) + L".exe");
+			Filesystem::Copy(originalExePath.c_str(), exePath.c_str());
 		}
 
 		const size_t token = modulePatch->GetToken();
@@ -4197,7 +4509,7 @@ bool LiveModule::InstallCompiledPatches(LiveProcess* liveProcess, void* original
 
 		// note that the image on disk we are trying to load had its entry point patched already when it was
 		// loaded for the first time, so we don't have to do that at this point.
-		executable::Image* image = executable::OpenImage(exePath.c_str(), file::OpenMode::READ_ONLY);
+		executable::Image* image = executable::OpenImage(exePath.c_str(), Filesystem::OpenMode::READ);
 		if (!image)
 		{
 			LC_ERROR_USER("Cannot load patch executable %S", exePath.c_str());
@@ -4220,16 +4532,16 @@ bool LiveModule::InstallCompiledPatches(LiveProcess* liveProcess, void* original
 			// before doing anything further, we need to ensure that the patch can be loaded into the address space at a suitable location.
 			// for 64-bit applications, this means that the patch must lie in a +/-2GB range of the main executable.
 			// 32-bit executables can reach the whole address space due to modulo addressing.
-			LC_LOG_DEV("Scanning memory for suitable patch location (PID: %d)", processId);
+			LC_LOG_DEV("Scanning memory for suitable patch location (PID: %d)", +processId);
 
 			// disable the main process before scanning its memory to ensure that no operation allocates/frees virtual memory concurrently
-			process::Suspend(processHandle);
+			Process::Suspend(processHandle);
 
 			const executable::PreferredBase preferredImageBase = FindPreferredImageBase(patchImageSize, processId, processHandle, originalModuleBase);
 
 			// rebase the patch image to its preferred base address
-			executable::Image* rebasedImage = executable::OpenImage(exePath.c_str(), file::OpenMode::READ_AND_WRITE);
-			LC_LOG_DEV("Rebasing patch executable to image base 0x%" PRIX64 " (PID: %d)", preferredImageBase, processId);
+			executable::Image* rebasedImage = executable::OpenImage(exePath.c_str(), Filesystem::OpenMode::READ_WRITE);
+			LC_LOG_DEV("Rebasing patch executable to image base 0x%" PRIX64 " (PID: %d)", preferredImageBase, +processId);
 			executable::RebaseImage(rebasedImage, preferredImageBase);
 			executable::CloseImage(rebasedImage);
 
@@ -4237,7 +4549,7 @@ bool LiveModule::InstallCompiledPatches(LiveProcess* liveProcess, void* original
 			// will allocate virtual memory at the patch's preferred image base, possibly rendering the patch unusable because
 			// it cannot be loaded.
 			// the chances of that happening are *very* rare though, and we can always load the next patch then.
-			process::Resume(processHandle);
+			Process::Resume(processHandle);
 #endif
 
 			pipe->SendCommandAndWaitForAck(cmd, nullptr, 0u);
@@ -4286,10 +4598,10 @@ bool LiveModule::InstallCompiledPatches(LiveProcess* liveProcess, void* original
 		LC_LOG_DEV("Patching dynamic initializers");
 		for (auto rva : patchData.patchedInitializers)
 		{
-			LC_LOG_DEV("Patching dynamic initializer symbol at RVA 0x%X (PID: %d)", rva, processId);
+			LC_LOG_DEV("Patching dynamic initializer symbol at RVA 0x%X (PID: %d)", rva, +processId);
 
 			void* initializerAddress = pointer::Offset<void*>(moduleBase, rva);
-			process::WriteProcessMemory(processHandle, initializerAddress, nullptr);
+			Process::WriteProcessMemory(processHandle, initializerAddress, nullptr);
 		}
 
 
@@ -4336,7 +4648,7 @@ bool LiveModule::InstallCompiledPatches(LiveProcess* liveProcess, void* original
 		
 
 		// suspend the main processes before patching functions, because they might not use synchronization points.
-		process::Suspend(processHandle);
+		Process::Suspend(processHandle);
 
 
 		// patch all functions
@@ -4357,7 +4669,7 @@ bool LiveModule::InstallCompiledPatches(LiveProcess* liveProcess, void* original
 
 
 		// resume the main processes again
-		process::Resume(processHandle);
+		Process::Resume(processHandle);
 
 
 		LC_LOG_DEV("Calling post-patch hooks");
@@ -4420,7 +4732,7 @@ bool LiveModule::actions::LoadPatchInfo::Execute(const CommandType* command, con
 
 void LiveModule::UpdateDirectoryCache(const ImmutableString& path, symbols::Dependency* dependency, DirectoryCache* cache)
 {
-	const std::wstring& directoryOnly = file::GetDirectory(string::ToWideString(path));
+	const std::wstring directoryOnly = Filesystem::GetDirectory(string::ToWideString(path).c_str()).GetString();
 	dependency->parentDirectory = cache->AddDirectory(directoryOnly);
 }
 
@@ -4473,7 +4785,7 @@ void LiveModule::CallPrecompileHooks()
 		{
 			const ModuleCache::ProcessData& processData = hookData.data->processes[p];
 
-			const unsigned int pid = processData.processId;
+			const Process::Id pid = processData.processId;
 			void* moduleBase = processData.moduleBase;
 			const DuplexPipe* pipe = processData.pipe;
 
@@ -4493,7 +4805,7 @@ void LiveModule::CallPostcompileHooks()
 		{
 			const ModuleCache::ProcessData& processData = hookData.data->processes[p];
 
-			const unsigned int pid = processData.processId;
+			const Process::Id pid = processData.processId;
 			void* moduleBase = processData.moduleBase;
 			const DuplexPipe* pipe = processData.pipe;
 
