Integration of Live++ 1.6.3

diff --git a/Engine/Source/Developer/Windows/LiveCodingServer/Private/External/LC_LiveModule.h b/Engine/Source/Developer/Windows/LiveCodingServer/Private/External/LC_LiveModule.h
--- a/Engine/Source/Developer/Windows/LiveCodingServer/Private/External/LC_LiveModule.h
+++ b/Engine/Source/Developer/Windows/LiveCodingServer/Private/External/LC_LiveModule.h
@@ -1,12 +1,14 @@
-// Copyright 2011-2019 Molecular Matters GmbH, all rights reserved.
+// Copyright 2011-2020 Molecular Matters GmbH, all rights reserved.
 
 #pragma once
 
+// BEGIN EPIC MOD
 #include "CoreTypes.h"
+// END EPIC MOD
 #include "LC_CoffCache.h"
 #include "LC_ModuleCache.h"
 #include "LC_Symbols.h"
-#include "LC_Process.h"
+#include "LC_ProcessTypes.h"
 #include "LC_Commands.h"
 #include "LC_Executable.h"
 #include "LC_Semaphore.h"
@@ -18,6 +20,7 @@ class FileAttributeCache;
 class DirectoryCache;
 class ModulePatch;
 class LiveProcess;
+class VirtualMemoryRange;
 
 class LiveModule
 {
@@ -27,6 +30,7 @@ public:
 		LiveProcess* liveProcess;
 		void* originalModuleBase;
 		std::wstring modulePath;
+		VirtualMemoryRange* virtualMemoryRange;
 	};
 
 	struct ErrorType
@@ -51,12 +55,6 @@ public:
 		};
 	};
 
-	struct ModifiedObjFile
-	{
-		std::wstring objPath;
-		std::wstring amalgamatedObjPath;	// optional
-	};
-
 	struct CompileResult
 	{
 		unsigned int exitCode;
@@ -82,7 +80,7 @@ public:
 	// in DEFAULT mode, this checks for file modifications, compiles files automatically, builds a patch containing changes, and loads them into the host application.
 	// in EXTERNAL_BUILD_SYSTEM mode, this does not compile files but builds a patch containing modified .objs, loading the patch into the host application.
 	// optionally, an array of modified or new .objs can be given in this mode, which builds a patch containing these files, not checking for any other modifications.
-	ErrorType::Enum Update(FileAttributeCache* fileCache, DirectoryCache* directoryCache, UpdateType::Enum updateType, const types::vector<ModifiedObjFile>& modifiedOrNewObjFiles);
+	ErrorType::Enum Update(FileAttributeCache* fileCache, DirectoryCache* directoryCache, UpdateType::Enum updateType, const types::vector<symbols::ModifiedObjFile>& modifiedOrNewObjFiles);
 	bool InstallCompiledPatches(LiveProcess* liveProcess, void* originalModuleBase);
 
 	const std::wstring& GetModuleName(void) const;
@@ -142,10 +140,9 @@ private:
 	symbols::ImageSectionDB* m_imageSectionDB = nullptr;
 	types::StringMap<types::vector<const symbols::Symbol*>> m_externalSymbolsPerCompilandCache;
 	types::StringMap<ImmutableString> m_pchSymbolToCompilandName;
-	types::vector<ImmutableString> m_weakSymbolsInLibs;
 
 	// patch data
-	types::unordered_map<unsigned int, types::unordered_set<const void*>> m_patchedAddressesPerProcess;
+	types::unordered_map<Process::Id, types::unordered_set<const void*>> m_patchedAddressesPerProcess;
 	unsigned int m_patchCounter = 0u;
 
 	// data pertaining to the next patch
diff --git a/Engine/Source/Developer/Windows/LiveCodingServer/Private/External/LC_LiveProcess.cpp b/Engine/Source/Developer/Windows/LiveCodingServer/Private/External/LC_LiveProcess.cpp
--- a/Engine/Source/Developer/Windows/LiveCodingServer/Private/External/LC_LiveProcess.cpp
+++ b/Engine/Source/Developer/Windows/LiveCodingServer/Private/External/LC_LiveProcess.cpp
@@ -1,15 +1,21 @@
-// Copyright 2011-2019 Molecular Matters GmbH, all rights reserved.
+// Copyright 2011-2020 Molecular Matters GmbH, all rights reserved.
 
+// BEGIN EPIC MOD
+//#include PCH_INCLUDE
+// END EPIC MOD
 #include "LC_LiveProcess.h"
 #include "LC_HeartBeat.h"
 #include "LC_CodeCave.h"
 #include "LC_Event.h"
 #include "LC_PrimitiveNames.h"
 #include "LC_VisualStudioAutomation.h"
+#include "LC_Process.h"
+// BEGIN EPIC MOD
 #include "LC_Logging.h"
+// END EPIC MOD
 
 
-LiveProcess::LiveProcess(process::Handle processHandle, unsigned int processId, unsigned int commandThreadId, const void* jumpToSelf, const DuplexPipe* pipe,
+LiveProcess::LiveProcess(Process::Handle processHandle, Process::Id processId, Thread::Id commandThreadId, const void* jumpToSelf, const DuplexPipe* pipe,
 	const wchar_t* imagePath, const wchar_t* commandLine, const wchar_t* workingDirectory, const void* environment, size_t environmentSize)
 	: m_processHandle(processHandle)
 	, m_processId(processId)
@@ -22,10 +28,14 @@ LiveProcess::LiveProcess(process::Handle processHandle, unsigned int processId,
 	, m_environment(environment, environmentSize)
 	, m_imagesTriedToLoad()
 	, m_heartBeatDelta(0ull)
+// BEGIN EPIC MOD
 #if WITH_VISUALSTUDIO_DTE
+// END EPIC MOD
 	, m_vsDebugger(nullptr)
 	, m_vsDebuggerThreads()
+// BEGIN EPIC MOD
 #endif
+// END EPIC MOD
 	, m_codeCave(nullptr)
 	, m_restartState(RestartState::DEFAULT)
 {
@@ -56,7 +66,9 @@ bool LiveProcess::MadeProgress(void) const
 
 void LiveProcess::HandleDebuggingPreCompile(void)
 {
+// BEGIN EPIC MOD
 #if WITH_VISUALSTUDIO_DTE
+// END EPIC MOD
 	if (!MadeProgress())
 	{
 		// this process did not make progress.
@@ -70,22 +82,52 @@ void LiveProcess::HandleDebuggingPreCompile(void)
 			m_vsDebuggerThreads = visualStudio::EnumerateThreads(m_vsDebugger);
 			if (m_vsDebuggerThreads.size() > 0u)
 			{
-				visualStudio::FreezeThreads(m_vsDebugger, m_vsDebuggerThreads);
-				visualStudio::ThawThread(m_vsDebugger, m_vsDebuggerThreads, m_commandThreadId);
-				visualStudio::Resume(m_vsDebugger);
-
-				LC_SUCCESS_USER("Automating debugger attached to process (PID: %d)", m_processId);
-
-				return;
+				if (visualStudio::FreezeThreads(m_vsDebugger, m_vsDebuggerThreads))
+				{
+					if (visualStudio::ThawThread(m_vsDebugger, m_vsDebuggerThreads, m_commandThreadId))
+					{
+						if (visualStudio::Resume(m_vsDebugger))
+						{
+							LC_SUCCESS_USER("Automating debugger attached to process (PID: %d)", m_processId);
+
+							return;
+						}
+						else
+						{
+							LC_LOG_USER("Failed to resume debugger (PID: %d)", m_processId);
+						}
+					}
+					else
+					{
+						// BEGIN EPIC MOD
+						LC_LOG_USER("Failed to thaw Live coding command thread in debugger (PID: %d)", m_processId);
+						// END EPIC MOD
+					}
+				}
+				else
+				{
+					LC_LOG_USER("Failed to freeze threads in debugger (PID: %d)", m_processId);
+				}
+			}
+			else
+			{
+				LC_LOG_USER("Failed to enumerate threads in debugger (PID: %d)", m_processId);
 			}
 		}
+		else
+		{
+			LC_LOG_USER("Failed to find debugger attached to process (PID: %d)", m_processId);
+		}
 
 		// no debugger could be found or an error occurred.
 		// continue by installing a code cave.
+		m_vsDebugger = nullptr;
 		LC_LOG_USER("Failed to automate debugger attached to process (PID: %d), using fallback mechanism", m_processId);
 		LC_SUCCESS_USER("Waiting for client process (PID: %d), hit 'Continue' (F5 in Visual Studio) if being held in the debugger", m_processId);
 	}
+// BEGIN EPIC MOD
 #endif
+// END EPIC MOD
 
 	// this process either made progress and is not held in the debugger, or we failed automating the debugger.
 	// "halt" this process by installing a code cave.
@@ -100,21 +142,30 @@ void LiveProcess::HandleDebuggingPostCompile(void)
 		// we installed a code cave previously, remove it
 		UninstallCodeCave();
 	}
+// BEGIN EPIC MOD
 #if WITH_VISUALSTUDIO_DTE
+// END EPIC MOD
 	else if (m_vsDebugger)
 	{
 		// we automated the debugger previously. break into the debugger again and resume all threads.
 		// when debugging a C# project that calls into C++ code, the VS debugger sometimes creates new MTA threads in between our PreCompile and PostCompile calls.
 		// try getting a new list of threads and thaw them all as well.
-		visualStudio::Break(m_vsDebugger);
-		
-		types::vector<EnvDTE::ThreadPtr> newDebuggerThreads = visualStudio::EnumerateThreads(m_vsDebugger);
-		visualStudio::ThawThreads(m_vsDebugger, newDebuggerThreads);
-		visualStudio::ThawThreads(m_vsDebugger, m_vsDebuggerThreads);
+		if (visualStudio::Break(m_vsDebugger))
+		{
+			types::vector<EnvDTE::ThreadPtr> newDebuggerThreads = visualStudio::EnumerateThreads(m_vsDebugger);
+			if (newDebuggerThreads.size() > 0u)
+			{
+				visualStudio::ThawThreads(m_vsDebugger, newDebuggerThreads);
+			}
+
+			visualStudio::ThawThreads(m_vsDebugger, m_vsDebuggerThreads);
+		}
 	}
 
 	m_vsDebugger = nullptr;
+// BEGIN EPIC MOD
 #endif
+// END EPIC MOD
 }
 
 
@@ -186,14 +237,16 @@ void LiveProcess::WaitForExitBeforeRestart(void)
 		executeRestart.Signal();
 
 		// wait until the client terminates
-		process::Wait(m_processHandle);
+		Process::Wait(m_processHandle);
 
 		m_restartState = RestartState::SUCCESSFUL_EXIT;
 	}
 }
 
 
+// BEGIN EPIC MOD
 void LiveProcess::Restart(void* restartJob)
+// END EPIC MOD
 {
 	if (m_restartState == RestartState::SUCCESSFUL_EXIT)
 	{
@@ -208,10 +261,10 @@ void LiveProcess::Restart(void* restartJob)
 		}
 		// END EPIC MOD
 		// BEGIN EPIC MOD - Prevent orphaned console instances if processes fail to restart. Job object will be duplicated into child process.
-		process::Context* context = process::Spawn(m_imagePath.c_str(), m_workingDirectory.c_str(), commandLine.c_str(), m_environment.GetData(), process::SpawnFlags::SUSPENDED);
+		Process::Context* context = Process::Spawn(m_imagePath.c_str(), m_workingDirectory.c_str(), commandLine.c_str(), m_environment.GetData(), Process::SpawnFlags::SUSPENDED);
 		void* TargetHandle;
-		DuplicateHandle(GetCurrentProcess(), restartJob, context->pi.hProcess, &TargetHandle, 0, Windows::TRUE, DUPLICATE_SAME_ACCESS);
-		ResumeThread(context->pi.hThread);
+		DuplicateHandle(GetCurrentProcess(), restartJob, +Process::GetHandle(context), &TargetHandle, 0, Windows::TRUE, DUPLICATE_SAME_ACCESS);
+		Process::ResumeMainThread(context);
 		// END EPIC MOD
 
 		m_restartState = RestartState::SUCCESSFUL_RESTART;
diff --git a/Engine/Source/Developer/Windows/LiveCodingServer/Private/External/LC_LiveProcess.h b/Engine/Source/Developer/Windows/LiveCodingServer/Private/External/LC_LiveProcess.h
--- a/Engine/Source/Developer/Windows/LiveCodingServer/Private/External/LC_LiveProcess.h
+++ b/Engine/Source/Developer/Windows/LiveCodingServer/Private/External/LC_LiveProcess.h
@@ -1,13 +1,15 @@
-// Copyright 2011-2019 Molecular Matters GmbH, all rights reserved.
+// Copyright 2011-2020 Molecular Matters GmbH, all rights reserved.
 
 #pragma once
 
-#include "CoreTypes.h"
-#include "LC_Process.h"
+#include "LC_ProcessTypes.h"
+#include "LC_ThreadTypes.h"
 #include "LC_Executable.h"
 #include "LC_MemoryBlock.h"
+// BEGIN EPIC MOD
 #include "LC_Types.h"
 #include "VisualStudioDTE.h"
+// END EPIC MOD
 
 
 class DuplexPipe;
@@ -16,7 +18,7 @@ class CodeCave;
 class LiveProcess
 {
 public:
-	LiveProcess(process::Handle processHandle, unsigned int processId, unsigned int commandThreadId, const void* jumpToSelf, const DuplexPipe* pipe,
+	LiveProcess(Process::Handle processHandle, Process::Id processId, Thread::Id commandThreadId, const void* jumpToSelf, const DuplexPipe* pipe,
 		const wchar_t* imagePath, const wchar_t* commandLine, const wchar_t* workingDirectory, const void* environment, size_t environmentSize);
 
 	void ReadHeartBeatDelta(const wchar_t* const processGroupName);
@@ -40,21 +42,23 @@ public:
 
 	bool PrepareForRestart(void);
 	void WaitForExitBeforeRestart(void);
+	// BEGIN EPIC MOD
 	void Restart(void* restartJob);
+	// END EPIC MOD
 	bool WasSuccessfulRestart(void) const;
 
 
-	inline process::Handle GetProcessHandle(void) const
+	inline Process::Handle GetProcessHandle(void) const
 	{
 		return m_processHandle;
 	}
 
-	inline unsigned int GetProcessId(void) const
+	inline Process::Id GetProcessId(void) const
 	{
 		return m_processId;
 	}
 
-	inline unsigned int GetCommandThreadId(void) const
+	inline Thread::Id GetCommandThreadId(void) const
 	{
 		return m_commandThreadId;
 	}
@@ -89,9 +93,9 @@ public:
 	// END EPIC MOD
 
 private:
-	process::Handle m_processHandle;
-	unsigned int m_processId;
-	unsigned int m_commandThreadId;
+	Process::Handle m_processHandle;
+	Process::Id m_processId;
+	Thread::Id m_commandThreadId;
 	const void* m_jumpToSelf;
 	const DuplexPipe* m_pipe;
 
@@ -121,11 +125,15 @@ private:
 
 	uint64_t m_heartBeatDelta;
 
+// BEGIN EPIC MOD
 #if WITH_VISUALSTUDIO_DTE
+// END EPIC MOD
 	// for handling communication with the VS debugger
 	EnvDTE::DebuggerPtr m_vsDebugger;
 	types::vector<EnvDTE::ThreadPtr> m_vsDebuggerThreads;
+// BEGIN EPIC MOD
 #endif
+// END EPIC MOD
 
 	// fallback in case communication with the VS debugger is not possible
 	CodeCave* m_codeCave;
diff --git a/Engine/Source/Developer/Windows/LiveCodingServer/Private/External/LC_Mallocator.cpp b/Engine/Source/Developer/Windows/LiveCodingServer/Private/External/LC_Mallocator.cpp
--- a/Engine/Source/Developer/Windows/LiveCodingServer/Private/External/LC_Mallocator.cpp
+++ b/Engine/Source/Developer/Windows/LiveCodingServer/Private/External/LC_Mallocator.cpp
@@ -1,9 +1,16 @@
-// Copyright 2011-2019 Molecular Matters GmbH, all rights reserved.
+// Copyright 2011-2020 Molecular Matters GmbH, all rights reserved.
 
+// BEGIN EPIC MOD
+//#include PCH_INCLUDE
+// END EPIC MOD
 #include "LC_Mallocator.h"
+// BEGIN EPIC MOD
+#include "LC_Assert.h"
 #include "LC_Platform.h"
+#include LC_PLATFORM_INCLUDE(LC_Foundation)
 #include "LC_Logging.h"
 #include "Windows/WindowsHWrapper.h"
+// END EPIC MOD
 
 namespace
 {
diff --git a/Engine/Source/Developer/Windows/LiveCodingServer/Private/External/LC_MicroAllocator.cpp b/Engine/Source/Developer/Windows/LiveCodingServer/Private/External/LC_MicroAllocator.cpp
--- a/Engine/Source/Developer/Windows/LiveCodingServer/Private/External/LC_MicroAllocator.cpp
+++ b/Engine/Source/Developer/Windows/LiveCodingServer/Private/External/LC_MicroAllocator.cpp
@@ -1,11 +1,17 @@
-// Copyright 2011-2019 Molecular Matters GmbH, all rights reserved.
+// Copyright 2011-2020 Molecular Matters GmbH, all rights reserved.
 
+// BEGIN EPIC MOD
+//#include PCH_INCLUDE
+// END EPIC MOD
 #include "LC_MicroAllocator.h"
 #include "LC_VirtualMemory.h"
+#include "LC_BitUtil.h"
+// BEGIN EPIC MOD
+#include "LC_Assert.h"
 #include "LC_PoolAllocator.h"
 #include "LC_Platform.h"
 #include "LC_Logging.h"
-#include "LC_BitUtil.h"
+// END EPIC MOD
 
 
 MicroAllocator::MicroAllocator(const char* name, size_t alignment)
@@ -19,7 +25,7 @@ MicroAllocator::MicroAllocator(const char* name, size_t alignment)
 		const size_t elementSize = i;
 
 		// all pools grow such that new blocks need to be allocated for every ELEMENT_COUNT_PER_POOL_GROWTH elements
-		const uint32_t pageSize = virtualMemory::GetPageSize();
+		const uint32_t pageSize = VirtualMemory::GetPageSize();
 		const size_t growSize = bitUtil::RoundUpToMultiple<size_t>(elementSize * ELEMENT_COUNT_PER_POOL_GROWTH, pageSize);
 		m_poolAllocators[i] = new PoolAllocator<PoolAllocatorMultiThreadPolicy>(name, elementSize, alignment, growSize);
 	}
@@ -55,6 +61,7 @@ void* MicroAllocator::Allocate(size_t size, size_t alignment)
 	return m_blockAllocator.Allocate(size, alignment);
 }
 
+
 void MicroAllocator::Free(void* ptr, size_t size)
 {
 	if (size <= POOL_ALLOCATOR_COUNT)
@@ -68,6 +75,7 @@ void MicroAllocator::Free(void* ptr, size_t size)
 	}
 }
 
+
 void MicroAllocator::PrintStats(void) const
 {
 	AllocatorStats stats;
diff --git a/Engine/Source/Developer/Windows/LiveCodingServer/Private/External/LC_ModuleCache.cpp b/Engine/Source/Developer/Windows/LiveCodingServer/Private/External/LC_ModuleCache.cpp
--- a/Engine/Source/Developer/Windows/LiveCodingServer/Private/External/LC_ModuleCache.cpp
+++ b/Engine/Source/Developer/Windows/LiveCodingServer/Private/External/LC_ModuleCache.cpp
@@ -1,5 +1,8 @@
-// Copyright 2011-2019 Molecular Matters GmbH, all rights reserved.
+// Copyright 2011-2020 Molecular Matters GmbH, all rights reserved.
 
+// BEGIN EPIC MOD
+//#include PCH_INCLUDE
+// END EPIC MOD
 #include "LC_ModuleCache.h"
 #include "LC_LiveProcess.h"
 
@@ -17,7 +20,7 @@ size_t ModuleCache::Insert(const symbols::SymbolDB* symbolDb, const symbols::Con
 	CriticalSection::ScopedLock lock(&m_cs);
 
 	const size_t token = m_cache.size();
-	m_cache.push_back(Data { static_cast<uint16_t>(token), symbolDb, contributionDb, compilandDb, thunkDb, imageSectionDb });
+	m_cache.push_back(Data { static_cast<uint16_t>(token), symbolDb, contributionDb, compilandDb, thunkDb, imageSectionDb, {} });
 
 	return token;
 }
@@ -109,7 +112,7 @@ ModuleCache::FindHookData ModuleCache::FindHooksInSectionBackwards(size_t ignore
 }
 
 
-types::vector<void*> ModuleCache::GatherModuleBases(unsigned int processId) const
+types::vector<void*> ModuleCache::GatherModuleBases(Process::Id processId) const
 {
 	types::vector<void*> result;
 
diff --git a/Engine/Source/Developer/Windows/LiveCodingServer/Private/External/LC_ModuleCache.h b/Engine/Source/Developer/Windows/LiveCodingServer/Private/External/LC_ModuleCache.h
--- a/Engine/Source/Developer/Windows/LiveCodingServer/Private/External/LC_ModuleCache.h
+++ b/Engine/Source/Developer/Windows/LiveCodingServer/Private/External/LC_ModuleCache.h
@@ -1,11 +1,14 @@
-// Copyright 2011-2019 Molecular Matters GmbH, all rights reserved.
+// Copyright 2011-2020 Molecular Matters GmbH, all rights reserved.
 
 #pragma once
 
+// BEGIN EPIC MOD
 #include "CoreTypes.h"
+// END EPIC MOD
 #include "LC_Symbols.h"
 #include "LC_Hook.h"
 #include "LC_CriticalSection.h"
+#include "LC_ProcessTypes.h"
 
 
 class LiveProcess;
@@ -21,8 +24,8 @@ public:
 	{
 		// all data except moduleBase is redundant and stored per cache entry, but this doesn't
 		// increase memory requirements by much. we'd rather have all information accessible fast.
-		unsigned int processId;
-		process::Handle processHandle;
+		Process::Id processId;
+		Process::Handle processHandle;
 		const DuplexPipe* pipe;
 
 		void* moduleBase;
@@ -74,7 +77,7 @@ public:
 	FindHookData FindHooksInSectionBackwards(size_t ignoreToken, const ImmutableString& sectionName) const;
 
 
-	types::vector<void*> GatherModuleBases(unsigned int processId) const;
+	types::vector<void*> GatherModuleBases(Process::Id processId) const;
 
 
 	inline size_t GetSize(void) const
diff --git a/Engine/Source/Developer/Windows/LiveCodingServer/Private/External/LC_ModulePatch.h b/Engine/Source/Developer/Windows/LiveCodingServer/Private/External/LC_ModulePatch.h
--- a/Engine/Source/Developer/Windows/LiveCodingServer/Private/External/LC_ModulePatch.h
+++ b/Engine/Source/Developer/Windows/LiveCodingServer/Private/External/LC_ModulePatch.h
@@ -1,8 +1,10 @@
-// Copyright 2011-2019 Molecular Matters GmbH, all rights reserved.
+// Copyright 2011-2020 Molecular Matters GmbH, all rights reserved.
 
 #pragma once
 
+// BEGIn EPIC MOD
 #include "CoreTypes.h"
+// END EPIC MOD
 #include "LC_RelocationPatcher.h"
 #include "LC_FunctionPatcher.h"
 #include "LC_ExecutablePatcher.h"
diff --git a/Engine/Source/Developer/Windows/LiveCodingServer/Private/External/LC_NameMangling.cpp b/Engine/Source/Developer/Windows/LiveCodingServer/Private/External/LC_NameMangling.cpp
--- a/Engine/Source/Developer/Windows/LiveCodingServer/Private/External/LC_NameMangling.cpp
+++ b/Engine/Source/Developer/Windows/LiveCodingServer/Private/External/LC_NameMangling.cpp
@@ -1,5 +1,8 @@
-// Copyright 2011-2019 Molecular Matters GmbH, all rights reserved.
+// Copyright 2011-2020 Molecular Matters GmbH, all rights reserved.
 
+// BEGIN EPIC MOD
+//#include PCH_INCLUDE
+// END EPC MOD
 #include "LC_NameMangling.h"
 #include "LC_StringUtil.h"
 
diff --git a/Engine/Source/Developer/Windows/LiveCodingServer/Private/External/LC_OutputFile.cpp b/Engine/Source/Developer/Windows/LiveCodingServer/Private/External/LC_OutputFile.cpp
--- a/Engine/Source/Developer/Windows/LiveCodingServer/Private/External/LC_OutputFile.cpp
+++ b/Engine/Source/Developer/Windows/LiveCodingServer/Private/External/LC_OutputFile.cpp
@@ -1,8 +1,13 @@
-// Copyright 2011-2019 Molecular Matters GmbH, all rights reserved.
+// Copyright 2011-2020 Molecular Matters GmbH, all rights reserved.
 
+// BEGIN EPIC MOD
+//#include PCH_INCLUDE
+// END EPIC MOD
 #include "LC_OutputFile.h"
-#include "LC_FileUtil.h"
+#include "LC_Filesystem.h"
+// BEGIN EPIC MOD
 #include "LC_Logging.h"
+// END EPIC MOD
 
 
 OutputFile::OutputFile(const wchar_t* logFilePath)
@@ -15,7 +20,7 @@ OutputFile::OutputFile(const wchar_t* logFilePath)
 	}
 	else
 	{
-		LC_LOG_USER("Creating log file at %S", file::NormalizePath(logFilePath).c_str());
+		LC_LOG_USER("Creating log file at %S", Filesystem::NormalizePath(logFilePath).GetString());
 	}
 }
 
@@ -41,4 +46,4 @@ void OutputFile::WriteToFile(const char* text)
 		const size_t length = strlen(text);
 		::WriteFile(m_logFile, text, static_cast<DWORD>(length), &bytesWritten, nullptr);
 	}
-}
+}
\ No newline at end of file
diff --git a/Engine/Source/Developer/Windows/LiveCodingServer/Private/External/LC_Patch.cpp b/Engine/Source/Developer/Windows/LiveCodingServer/Private/External/LC_Patch.cpp
--- a/Engine/Source/Developer/Windows/LiveCodingServer/Private/External/LC_Patch.cpp
+++ b/Engine/Source/Developer/Windows/LiveCodingServer/Private/External/LC_Patch.cpp
@@ -1,41 +1,54 @@
-// Copyright 2011-2019 Molecular Matters GmbH, all rights reserved.
+// Copyright 2011-2020 Molecular Matters GmbH, all rights reserved.
 
+// BEGIN EPIC MOD
+//#include PCH_INCLUDE
+// END EPIC MOD
 #include "LC_Patch.h"
-#include "LC_Assembly.h"
+#include "LC_Assembler.h"
 #include "LC_PointerUtil.h"
-#include "LC_Platform.h"
+#include "LC_Process.h"
+// BEGIN EPIC MOD
+#include "LC_Assert.h"
 #include "LC_Logging.h"
+// END EPIC MOD
 
+// BEGIN EPIC MODS
+#pragma warning(push)
+#pragma warning(disable:4686) // possible change in behavior, change in UDT return calling convention
+// END EPIC MODS
 
 namespace
 {
-	void Install(process::Handle processHandle, void* address, const assembly::Instruction& instruction)
+	void Install(Process::Handle processHandle, void* address, const Assembler::Instruction& instruction)
 	{
-		process::WriteProcessMemory(processHandle, address, instruction.code, instruction.size);
-		process::FlushInstructionCache(processHandle, address, instruction.size);
+		Process::WriteProcessMemory(processHandle, address, instruction.machineCode, instruction.size);
+		Process::FlushInstructionCache(processHandle, address, instruction.size);
 	}
-}
 
 
-void patch::InstallNOPs(process::Handle processHandle, void* address, size_t size)
-{
-	const assembly::Instruction nop = assembly::MakeNOP();
-	for (size_t i = 0u; i < size; ++i)
+	template <uint8_t N>
+	void Install(Process::Handle processHandle, void* address, const Assembler::SizedInstruction<N>& instruction)
 	{
-		process::WriteProcessMemory(processHandle, pointer::Offset<void*>(address, i), nop.code, nop.size);
+		Process::WriteProcessMemory(processHandle, address, instruction.machineCode, N);
+		Process::FlushInstructionCache(processHandle, address, N);
 	}
+}
 
-	process::FlushInstructionCache(processHandle, address, size);
+
+void patch::InstallNOPs(Process::Handle processHandle, void* address, uint8_t size)
+{
+	const Assembler::Instruction& nop = Assembler::MakeNOP(size);
+	Install(processHandle, address, nop);
 }
 
 
-void patch::InstallJumpToSelf(process::Handle processHandle, void* address)
+void patch::InstallJumpToSelf(Process::Handle processHandle, void* address)
 {
 	InstallRelativeShortJump(processHandle, address, address);
 }
 
 
-void patch::InstallRelativeShortJump(process::Handle processHandle, void* address, void* destination)
+void patch::InstallRelativeShortJump(Process::Handle processHandle, void* address, void* destination)
 {
 	char* oldFuncAddr = static_cast<char*>(address);
 	char* newFuncAddr = static_cast<char*>(destination);
@@ -44,12 +57,11 @@ void patch::InstallRelativeShortJump(process::Handle processHandle, void* addres
 	LC_ASSERT(displacement >= std::numeric_limits<int8_t>::min(), "Displacement is out-of-range.");
 	LC_ASSERT(displacement <= std::numeric_limits<int8_t>::max(), "Displacement is out-of-range.");
 
-	const assembly::Instruction jump = assembly::MakeRelativeShortJump(static_cast<int8_t>(displacement));
-	Install(processHandle, address, jump);
+	Install(processHandle, address, Assembler::MakeJMP8Relative(static_cast<int8_t>(displacement)));
 }
 
 
-void patch::InstallRelativeNearJump(process::Handle processHandle, void* address, void* destination)
+void patch::InstallRelativeNearJump(Process::Handle processHandle, void* address, void* destination)
 {
 	char* oldFuncAddr = static_cast<char*>(address);
 	char* newFuncAddr = static_cast<char*>(destination);
@@ -58,6 +70,9 @@ void patch::InstallRelativeNearJump(process::Handle processHandle, void* address
 	LC_ASSERT(displacement >= std::numeric_limits<int32_t>::min(), "Displacement is out-of-range.");
 	LC_ASSERT(displacement <= std::numeric_limits<int32_t>::max(), "Displacement is out-of-range.");
 
-	const assembly::Instruction jump = assembly::MakeRelativeNearJump(static_cast<int32_t>(displacement));
-	Install(processHandle, address, jump);
+	Install(processHandle, address, Assembler::MakeJMP32Relative(static_cast<int32_t>(displacement)));
 }
+
+// BEGIN EPIC MODS
+#pragma warning(pop)
+// END EPIC MODS
diff --git a/Engine/Source/Developer/Windows/LiveCodingServer/Private/External/LC_Patch.h b/Engine/Source/Developer/Windows/LiveCodingServer/Private/External/LC_Patch.h
--- a/Engine/Source/Developer/Windows/LiveCodingServer/Private/External/LC_Patch.h
+++ b/Engine/Source/Developer/Windows/LiveCodingServer/Private/External/LC_Patch.h
@@ -1,16 +1,18 @@
-// Copyright 2011-2019 Molecular Matters GmbH, all rights reserved.
+// Copyright 2011-2020 Molecular Matters GmbH, all rights reserved.
 
 #pragma once
 
+// BEGIN EPIC MOD
 #include "CoreTypes.h"
-#include "LC_Process.h"
+// END EPIC MOD
+#include "LC_ProcessTypes.h"
 
 
 namespace patch
 {
-	void InstallNOPs(process::Handle processHandle, void* address, size_t size);
-	void InstallJumpToSelf(process::Handle processHandle, void* address);
+	void InstallNOPs(Process::Handle processHandle, void* address, uint8_t size);
+	void InstallJumpToSelf(Process::Handle processHandle, void* address);
 
-	void InstallRelativeShortJump(process::Handle processHandle, void* address, void* destination);
-	void InstallRelativeNearJump(process::Handle processHandle, void* address, void* destination);
+	void InstallRelativeShortJump(Process::Handle processHandle, void* address, void* destination);
+	void InstallRelativeNearJump(Process::Handle processHandle, void* address, void* destination);
 }
diff --git a/Engine/Source/Developer/Windows/LiveCodingServer/Private/External/LC_PoolAllocator.cpp b/Engine/Source/Developer/Windows/LiveCodingServer/Private/External/LC_PoolAllocator.cpp
--- a/Engine/Source/Developer/Windows/LiveCodingServer/Private/External/LC_PoolAllocator.cpp
+++ b/Engine/Source/Developer/Windows/LiveCodingServer/Private/External/LC_PoolAllocator.cpp
@@ -1,11 +1,17 @@
-// Copyright 2011-2019 Molecular Matters GmbH, all rights reserved.
+// Copyright 2011-2020 Molecular Matters GmbH, all rights reserved.
 
+// BEGIN EPIC MOD
+//#include PCH_INCLUDE
+// END EPIC MOD
 #include "LC_PoolAllocator.h"
 #include "LC_PointerUtil.h"
 #include "LC_VirtualMemory.h"
+#include "LC_BitUtil.h"
+// BEGIN EPIC MOD
+#include "LC_Assert.h"
 #include "LC_Platform.h"
 #include "LC_Logging.h"
-#include "LC_BitUtil.h"
+// END EPIC MOD
 
 
 namespace freeList
@@ -91,7 +97,7 @@ void* PoolAllocator<T>::Allocate(size_t size, size_t alignment)
 	if (!m_freeList)
 	{
 		// no memory left, allocate a new block
-		void* block = virtualMemory::Allocate(m_growSize);
+		void* block = VirtualMemory::Allocate(m_growSize);
 
 		// initialize free list in this block of memory
 		m_freeList = freeList::Initialize(block, m_growSize, m_maxSize, m_maxAlignment, sizeof(BlockHeader));
@@ -144,7 +150,7 @@ void PoolAllocator<T>::Purge(void)
 	{
 		BlockHeader* temp = header;
 		header = header->next;
-		virtualMemory::Free(temp);
+		VirtualMemory::Free(temp);
 
 		m_stats.UnregisterAllocation(m_growSize);
 	}
diff --git a/Engine/Source/Developer/Windows/LiveCodingServer/Private/External/LC_RelocationPatcher.cpp b/Engine/Source/Developer/Windows/LiveCodingServer/Private/External/LC_RelocationPatcher.cpp
--- a/Engine/Source/Developer/Windows/LiveCodingServer/Private/External/LC_RelocationPatcher.cpp
+++ b/Engine/Source/Developer/Windows/LiveCodingServer/Private/External/LC_RelocationPatcher.cpp
@@ -1,10 +1,16 @@
-// Copyright 2011-2019 Molecular Matters GmbH, all rights reserved.
+// Copyright 2011-2020 Molecular Matters GmbH, all rights reserved.
 
+// BEGIN EPIC MOD
+//#include PCH_INCLUDE
+// END EPIC MOD
 #include "LC_RelocationPatcher.h"
 #include "LC_StringUtil.h"
 #include "LC_NameMangling.h"
 #include "LC_PointerUtil.h"
 #include "LC_CoffDetail.h"
+// BEGIN EPIC MOD
+#include "LC_Process.h"
+// END EPIC MOD
 
 
 bool relocations::WouldPatchRelocation(const ImmutableString& dstSymbolName)
@@ -108,15 +114,15 @@ relocations::Record relocations::PatchRelocation
 	const types::StringSet& forceRelocationSymbols,
 	const ModuleCache* moduleCache,
 	const ImmutableString& srcSymbolName,
+	const ImmutableString& dstSymbolName,
 	const symbols::Symbol* srcSymbol,
 	size_t newModuleIndex,
 	void* newModuleBases[]
 )
 {
-	Record record = { coff::Relocation::Type::UNKNOWN, 0u, 0u };
+	Record record = { coff::Relocation::Type::UNKNOWN, 0u, 0u, {} };
 
 	const coff::Relocation::Type::Enum type = relocation->type;
-	const ImmutableString& dstSymbolName = coff::GetRelocationDstSymbolName(coffDb, relocation);
 	const uint32_t characteristics = coff::GetRelocationDestinationSectionCharacteristics(coffDb, relocation);
 
 	const bool forceRelocation = (forceRelocationSymbols.find(dstSymbolName) != forceRelocationSymbols.end());
@@ -237,7 +243,7 @@ relocations::Record relocations::PatchRelocation
 			for (size_t p = 0u; p < count; ++p)
 			{
 				void* moduleBase = originalData.data->processes[p].moduleBase;
-				process::Handle processHandle = originalData.data->processes[p].processHandle;
+				Process::Handle processHandle = originalData.data->processes[p].processHandle;
 				void* newModuleBase = newModuleBases[p];
 
 				// find the address of the relocation.
@@ -270,7 +276,7 @@ relocations::Record relocations::PatchRelocation
 				const uint32_t displacement = pointer::Displacement<uint32_t>(byteFollowingRelocation, originalAddress);
 #endif
 
-				process::WriteProcessMemory(processHandle, relocationAddress, displacement);
+				Process::WriteProcessMemory(processHandle, relocationAddress, displacement);
 
 				LC_LOG_DEV("Patched relocation from symbol %s to %s at 0x%p (0x%x + 0x%x)", srcSymbolName.c_str(), dstSymbolName.c_str(), newModuleBase, srcSymbol->rva, relocation->srcRva);
 			}
@@ -298,14 +304,14 @@ relocations::Record relocations::PatchRelocation
 				const size_t count = originalData.data->processes.size();
 				for (size_t p = 0u; p < count; ++p)
 				{
-					process::Handle processHandle = originalData.data->processes[p].processHandle;
+					Process::Handle processHandle = originalData.data->processes[p].processHandle;
 					void* newModuleBase = newModuleBases[p];
 
 					// find the address of the relocation.
 					// the relocation's RVA is relative to the start of the function.
 					void* relocationAddress = pointer::Offset<void*>(newModuleBase, srcSymbol->rva + relocation->srcRva);
 
-					process::WriteProcessMemory(processHandle, relocationAddress, sectionRelativeRva);
+					Process::WriteProcessMemory(processHandle, relocationAddress, sectionRelativeRva);
 
 					LC_LOG_DEV("Patched relocation from symbol %s to %s at 0x%p (0x%x + 0x%x)", srcSymbolName.c_str(), dstSymbolName.c_str(), newModuleBase, srcSymbol->rva, relocation->srcRva);
 				}
@@ -335,7 +341,7 @@ relocations::Record relocations::PatchRelocation
 			for (size_t p = 0u; p < count; ++p)
 			{
 				void* moduleBase = originalData.data->processes[p].moduleBase;
-				process::Handle processHandle = originalData.data->processes[p].processHandle;
+				Process::Handle processHandle = originalData.data->processes[p].processHandle;
 				void* newModuleBase = newModuleBases[p];
 
 				// find the address of the relocation.
@@ -347,7 +353,7 @@ relocations::Record relocations::PatchRelocation
 				// The target's 32-bit VA.
 				const uint32_t va = static_cast<uint32_t>(reinterpret_cast<uintptr_t>(originalAddress));
 
-				process::WriteProcessMemory(processHandle, relocationAddress, va);
+				Process::WriteProcessMemory(processHandle, relocationAddress, va);
 
 				LC_LOG_DEV("Patched relocation from symbol %s to %s at 0x%p (0x%x + 0x%x)", srcSymbolName.c_str(), dstSymbolName.c_str(), newModuleBase, srcSymbol->rva, relocation->srcRva);
 			}
@@ -366,7 +372,7 @@ relocations::Record relocations::PatchRelocation
 			for (size_t p = 0u; p < count; ++p)
 			{
 				void* moduleBase = originalData.data->processes[p].moduleBase;
-				process::Handle processHandle = originalData.data->processes[p].processHandle;
+				Process::Handle processHandle = originalData.data->processes[p].processHandle;
 				void* newModuleBase = newModuleBases[p];
 
 				// find the address of the relocation.
@@ -381,7 +387,7 @@ relocations::Record relocations::PatchRelocation
 				const int64_t displacementToNewBase = pointer::Displacement<int64_t>(newModuleBase, originalAddress);
 				const int32_t displacement = static_cast<int32_t>(displacementToNewBase);
 
-				process::WriteProcessMemory(processHandle, relocationAddress, displacement);
+				Process::WriteProcessMemory(processHandle, relocationAddress, displacement);
 
 				LC_LOG_DEV("Patched relocation from symbol %s to %s at 0x%p (0x%x + 0x%x)", srcSymbolName.c_str(), dstSymbolName.c_str(), newModuleBase, srcSymbol->rva, relocation->srcRva);
 			}
@@ -400,7 +406,7 @@ relocations::Record relocations::PatchRelocation
 			for (size_t p = 0u; p < count; ++p)
 			{
 				void* moduleBase = originalData.data->processes[p].moduleBase;
-				process::Handle processHandle = originalData.data->processes[p].processHandle;
+				Process::Handle processHandle = originalData.data->processes[p].processHandle;
 				void* newModuleBase = newModuleBases[p];
 
 				// find the address of the relocation.
@@ -412,7 +418,7 @@ relocations::Record relocations::PatchRelocation
 				// The target's 64-bit VA.
 				const uint64_t va = static_cast<uint64_t>(reinterpret_cast<uintptr_t>(originalAddress));
 
-				process::WriteProcessMemory(processHandle, relocationAddress, va);
+				Process::WriteProcessMemory(processHandle, relocationAddress, va);
 
 				LC_LOG_DEV("Patched relocation from symbol %s to %s at 0x%p (0x%x + 0x%x)", srcSymbolName.c_str(), dstSymbolName.c_str(), newModuleBase, srcSymbol->rva, relocation->srcRva);
 			}
@@ -433,7 +439,7 @@ relocations::Record relocations::PatchRelocation
 void relocations::PatchRelocation
 (
 	const Record& record,
-	process::Handle processHandle,
+	Process::Handle processHandle,
 	void* processModuleBases[],
 	void* newModuleBase
 )
@@ -470,13 +476,13 @@ void relocations::PatchRelocation
 			const uint32_t displacement = pointer::Displacement<uint32_t>(byteFollowingRelocation, originalAddress);
 #endif
 
-			process::WriteProcessMemory(processHandle, relocationAddress, displacement);
+			Process::WriteProcessMemory(processHandle, relocationAddress, displacement);
 		}
 		break;
 
 		case coff::Relocation::Type::SECTION_RELATIVE:
 		{
-			process::WriteProcessMemory(processHandle, relocationAddress, record.data.sectionRelativeRelocation.sectionRelativeRva);
+			Process::WriteProcessMemory(processHandle, relocationAddress, record.data.sectionRelativeRelocation.sectionRelativeRva);
 		}
 		break;
 
@@ -486,7 +492,7 @@ void relocations::PatchRelocation
 			const void* originalAddress = pointer::Offset<const void*>(moduleBase, record.data.va32Relocation.originalModuleRva);
 			const uint32_t va = static_cast<uint32_t>(reinterpret_cast<uintptr_t>(originalAddress));
 
-			process::WriteProcessMemory(processHandle, relocationAddress, va);
+			Process::WriteProcessMemory(processHandle, relocationAddress, va);
 #endif
 		}
 		break;
@@ -497,7 +503,7 @@ void relocations::PatchRelocation
 			const int64_t displacementToNewBase = pointer::Displacement<int64_t>(newModuleBase, originalAddress);
 			const int32_t displacement = static_cast<int32_t>(displacementToNewBase);
 
-			process::WriteProcessMemory(processHandle, relocationAddress, displacement);
+			Process::WriteProcessMemory(processHandle, relocationAddress, displacement);
 		}
 		break;
 
@@ -507,7 +513,7 @@ void relocations::PatchRelocation
 			const void* originalAddress = pointer::Offset<const void*>(moduleBase, record.data.va64Relocation.originalModuleRva);
 			const uint64_t va = static_cast<uint64_t>(reinterpret_cast<uintptr_t>(originalAddress));
 
-			process::WriteProcessMemory(processHandle, relocationAddress, va);
+			Process::WriteProcessMemory(processHandle, relocationAddress, va);
 		}
 		break;
 #endif
diff --git a/Engine/Source/Developer/Windows/LiveCodingServer/Private/External/LC_RelocationPatcher.h b/Engine/Source/Developer/Windows/LiveCodingServer/Private/External/LC_RelocationPatcher.h
--- a/Engine/Source/Developer/Windows/LiveCodingServer/Private/External/LC_RelocationPatcher.h
+++ b/Engine/Source/Developer/Windows/LiveCodingServer/Private/External/LC_RelocationPatcher.h
@@ -1,12 +1,14 @@
-// Copyright 2011-2019 Molecular Matters GmbH, all rights reserved.
+// Copyright 2011-2020 Molecular Matters GmbH, all rights reserved.
 
 #pragma once
 
+// BEGIN EPIC MOD
 #include "CoreTypes.h"
+// END EPIC MOD
 #include "LC_Coff.h"
 #include "LC_Symbols.h"
 #include "LC_ModuleCache.h"
-#include "LC_Process.h"
+#include "LC_ProcessTypes.h"
 
 
 namespace relocations
@@ -63,6 +65,7 @@ namespace relocations
 		const types::StringSet& forceRelocationSymbols,
 		const ModuleCache* moduleCache,
 		const ImmutableString& srcSymbolName,
+		const ImmutableString& dstSymbolName,
 		const symbols::Symbol* srcSymbol,
 		size_t newModuleIndex,
 		void* newModuleBases[]
@@ -71,7 +74,7 @@ namespace relocations
 	void PatchRelocation
 	(
 		const Record& record,
-		process::Handle processHandle,
+		Process::Handle processHandle,
 		void* processModuleBases[],
 		void* newModuleBase
 	);
diff --git a/Engine/Source/Developer/Windows/LiveCodingServer/Private/External/LC_Scheduler.cpp b/Engine/Source/Developer/Windows/LiveCodingServer/Private/External/LC_Scheduler.cpp
--- a/Engine/Source/Developer/Windows/LiveCodingServer/Private/External/LC_Scheduler.cpp
+++ b/Engine/Source/Developer/Windows/LiveCodingServer/Private/External/LC_Scheduler.cpp
@@ -1,4 +1,8 @@
 // Copyright 2011-2019 Molecular Matters GmbH, all rights reserved.
+
+// BEGIN EPIC MOD
+//#include PCH_INCLUDE
+// END EPIC MOD
 #include "LC_Scheduler.h"
 #include "LC_SchedulerWorkerThread.h"
 #include "LC_SchedulerQueue.h"
@@ -34,7 +38,9 @@ namespace
 		SYSTEM_LOGICAL_PROCESSOR_INFORMATION* buffer = nullptr;
 		DWORD bytesNeeded = 0u;
 		BOOL result = ::GetLogicalProcessorInformation(buffer, &bytesNeeded);
-		if (result == Windows::FALSE) 
+		// BEGIN EPIC MOD
+		if (result == Windows::FALSE)
+		// END EPIC MOD 
 		{
 			if (::GetLastError() == ERROR_INSUFFICIENT_BUFFER)
 			{
@@ -43,7 +49,9 @@ namespace
 				result = ::GetLogicalProcessorInformation(buffer, &bytesNeeded);
 
 				// return the number of logical processors in case anything went wrong
+				// BEGIN EPIC MOD
 				if ((result == Windows::FALSE) || (buffer == nullptr))
+				// END EPIC MOD
 				{
 					::free(buffer);
 					return GetLogicalProcessorCount();
@@ -148,7 +156,7 @@ void scheduler::WaitForTask(TaskBase* task)
 		else
 		{
 			// no task available
-			thread::Sleep(10u);
+			Thread::Current::SleepMilliSeconds(10u);
 		}
 	}
 }
diff --git a/Engine/Source/Developer/Windows/LiveCodingServer/Private/External/LC_SchedulerQueue.h b/Engine/Source/Developer/Windows/LiveCodingServer/Private/External/LC_SchedulerQueue.h
--- a/Engine/Source/Developer/Windows/LiveCodingServer/Private/External/LC_SchedulerQueue.h
+++ b/Engine/Source/Developer/Windows/LiveCodingServer/Private/External/LC_SchedulerQueue.h
@@ -1,12 +1,15 @@
-// Copyright 2011-2019 Molecular Matters GmbH, all rights reserved.
+// Copyright 2011-2020 Molecular Matters GmbH, all rights reserved.
 
 #pragma once
 
+// BEGIN EPIC MOD
 #include "CoreTypes.h"
+// END EPIC MOD
 #include "LC_Semaphore.h"
 #include "LC_CriticalSection.h"
+// BEGIN EPIC MOD
 #include <stdint.h>
-
+// END EPIC MOD
 
 namespace scheduler
 {
diff --git a/Engine/Source/Developer/Windows/LiveCodingServer/Private/External/LC_SchedulerWorkerThread.cpp b/Engine/Source/Developer/Windows/LiveCodingServer/Private/External/LC_SchedulerWorkerThread.cpp
--- a/Engine/Source/Developer/Windows/LiveCodingServer/Private/External/LC_SchedulerWorkerThread.cpp
+++ b/Engine/Source/Developer/Windows/LiveCodingServer/Private/External/LC_SchedulerWorkerThread.cpp
@@ -1,24 +1,30 @@
-// Copyright 2011-2019 Molecular Matters GmbH, all rights reserved.
+// Copyright 2011-2020 Molecular Matters GmbH, all rights reserved.
 
+// BEGIN EPIC MOD
+//#include PCH_INCLUDE
+// END EPIC MOD
 #include "LC_SchedulerWorkerThread.h"
 #include "LC_SchedulerQueue.h"
 #include "LC_SchedulerTask.h"
+#include "LC_Thread.h"
 
 
 scheduler::WorkerThread::WorkerThread(TaskQueue* queue)
 	: m_thread()
 {
-	m_thread = thread::Create("Live coding worker", 128u * 1024u, &scheduler::WorkerThread::ThreadFunction, this, queue);
+	// BEGIN EPIC MOD
+	m_thread = Thread::CreateFromMemberFunction("Live coding worker", 128u * 1024u, this, &scheduler::WorkerThread::ThreadFunction, queue);
+	// END EPIC MOD
 }
 
 
 scheduler::WorkerThread::~WorkerThread(void)
 {
-	thread::Join(m_thread);
+	Thread::Join(m_thread);
 }
 
 
-unsigned int scheduler::WorkerThread::ThreadFunction(TaskQueue* queue)
+Thread::ReturnValue scheduler::WorkerThread::ThreadFunction(TaskQueue* queue)
 {
 	for (;;)
 	{
@@ -32,5 +38,5 @@ unsigned int scheduler::WorkerThread::ThreadFunction(TaskQueue* queue)
 		task->Execute();
 	}
 
-	return 0u;
+	return Thread::ReturnValue(0u);
 }
diff --git a/Engine/Source/Developer/Windows/LiveCodingServer/Private/External/LC_SchedulerWorkerThread.h b/Engine/Source/Developer/Windows/LiveCodingServer/Private/External/LC_SchedulerWorkerThread.h
--- a/Engine/Source/Developer/Windows/LiveCodingServer/Private/External/LC_SchedulerWorkerThread.h
+++ b/Engine/Source/Developer/Windows/LiveCodingServer/Private/External/LC_SchedulerWorkerThread.h
@@ -1,9 +1,11 @@
-// Copyright 2011-2019 Molecular Matters GmbH, all rights reserved.
+// Copyright 2011-2020 Molecular Matters GmbH, all rights reserved.
 
 #pragma once
 
+// BEGIN EPIC MOD
 #include "CoreTypes.h"
-#include "LC_Thread.h"
+// END EPIC MOD
+#include "LC_ThreadTypes.h"
 
 
 namespace scheduler
@@ -18,8 +20,8 @@ namespace scheduler
 		~WorkerThread(void);
 
 	private:
-		unsigned int ThreadFunction(TaskQueue* queue);
+		Thread::ReturnValue ThreadFunction(TaskQueue* queue);
 
-		thread::Handle m_thread;
+		Thread::Handle m_thread;
 	};
 }
diff --git a/Engine/Source/Developer/Windows/LiveCodingServer/Private/External/LC_ServerCommandThread.cpp b/Engine/Source/Developer/Windows/LiveCodingServer/Private/External/LC_ServerCommandThread.cpp
--- a/Engine/Source/Developer/Windows/LiveCodingServer/Private/External/LC_ServerCommandThread.cpp
+++ b/Engine/Source/Developer/Windows/LiveCodingServer/Private/External/LC_ServerCommandThread.cpp
@@ -1,17 +1,24 @@
-// Copyright 2011-2019 Molecular Matters GmbH, all rights reserved.
+// Copyright 2011-2020 Molecular Matters GmbH, all rights reserved.
 
+// BEGIN EPIC MOD
+//#include PCH_INCLUDE
+// END EPIC MOD
 #include "LC_ServerCommandThread.h"
 #include "LC_Commands.h"
+// BEGIN EPIC MOD
+//#include "LC_MainFrame.h"
+// END EPIC MOD
 #include "LC_Telemetry.h"
 #include "LC_Symbols.h"
-#include "LC_FileUtil.h"
+#include "LC_Filesystem.h"
 #include "LC_Process.h"
 #include "LC_Compiler.h"
 #include "LC_StringUtil.h"
 #include "LC_CommandMap.h"
 #include "LC_FileAttributeCache.h"
-#include "LiveCodingServer.h"
-#include "Containers/UnrealString.h"
+// BEGIN EPIC MOD
+//#include "LC_App.h"
+// END EPIC MOD
 #include "LC_Shortcut.h"
 #include "LC_Key.h"
 #include "LC_ChangeNotification.h"
@@ -20,15 +27,25 @@
 #include "LC_LiveModule.h"
 #include "LC_LiveProcess.h"
 #include "LC_CodeCave.h"
+// BEGIN EPIC MOD
+//#include "LC_ExceptionHandlerDialog.h"
+//#include "LC_Resource.h"
+// END EPIC MOD
 #include "LC_PrimitiveNames.h"
-#include "LC_AppSettings.h"
-#include "LC_Allocators.h"
-#include "LC_DuplexPipeClient.h"
 #include "LC_MemoryStream.h"
 #include "LC_NamedSharedMemory.h"
 #include "LC_VisualStudioAutomation.h"
+#include "LC_Thread.h"
 #include <mmsystem.h>
 
+// BEGIN EPIC MOD
+#include "LC_AppSettings.h"
+#include "LC_Allocators.h"
+#include "LC_DuplexPipeClient.h"
+#include "LiveCodingServer.h"
+#include "Containers/UnrealString.h"
+// END EPIC MOD
+
 // unreachable code
 #pragma warning (disable : 4702)
 
@@ -65,7 +82,7 @@ namespace
 
 	static executable::Header GetImageHeader(const wchar_t* path)
 	{
-		executable::Image* image = executable::OpenImage(path, file::OpenMode::READ_ONLY);
+		executable::Image* image = executable::OpenImage(path, Filesystem::OpenMode::READ);
 		if (!image)
 		{
 			return executable::Header {};
@@ -100,16 +117,26 @@ ServerCommandThread::ServerCommandThread(MainFrame* mainFrame, const wchar_t* co
 	, m_restartCS()
 	, m_restartJob(nullptr)
 	, m_restartedProcessCount(0u)
+// BEGIN EPIC MOD
 #if WITH_VISUALSTUDIO_DTE
+// END EPIC MOD
 	, m_restartedProcessIdToDebugger()
+// BEGIN EPIC MOD
 #endif
+// END EPIC MOD
 {
+// BEGIN EPIC MOD
 #if WITH_VISUALSTUDIO_DTE
+// END EPIC MOD
 	visualStudio::Startup();
+// BEGIN EPIC MOD
 #endif
+// END EPIC MOD
 
-	m_serverThread = thread::Create("Live coding server", 64u * 1024u, &ServerCommandThread::ServerThread, this);
-	m_compileThread = thread::Create("Live coding compilation", 64u * 1024u, &ServerCommandThread::CompileThread, this);
+	// BEGIN EPIC MOD
+	m_serverThread = Thread::CreateFromMemberFunction("Live coding server", 64u * 1024u, this, &ServerCommandThread::ServerThread);
+	m_compileThread = Thread::CreateFromMemberFunction("Live coding compilation", 64u * 1024u, this, &ServerCommandThread::CompileThread);
+	// END EPIC MOD
 
 	m_liveModules.reserve(256u);
 	m_liveProcesses.reserve(8u);
@@ -127,9 +154,13 @@ ServerCommandThread::~ServerCommandThread(void)
 	// some intensive work.
 	delete m_directoryCache;
 
+// BEGIN EPIC MOD
 #if WITH_VISUALSTUDIO_DTE
+// END EPIC MOD
 	visualStudio::Shutdown();
+// BEGIN EPIC MOD
 #endif
+// END EPIC MOD
 }
 
 
@@ -159,8 +190,8 @@ void ServerCommandThread::RestartTargets(void)
 		LiveProcess* liveProcess = *processIt;
 		if (liveProcess->WasSuccessfulRestart())
 		{
-			process::Handle processHandle = liveProcess->GetProcessHandle();
-			process::Close(processHandle);
+			Process::Handle processHandle = liveProcess->GetProcessHandle();
+			Process::Close(processHandle);
 
 			// tell live modules to remove this process
 			const size_t moduleCount = m_liveModules.size();
@@ -190,15 +221,19 @@ void ServerCommandThread::RestartTargets(void)
 		{
 			++m_restartedProcessCount;
 
+// BEGIN EPIC MOD
 #if WITH_VISUALSTUDIO_DTE
+// END EPIC MOD
 			// check if a VS debugger is currently attached to the process about to restart
-			const unsigned int processId = liveProcess->GetProcessId();
+			const Process::Id processId = liveProcess->GetProcessId();
 			EnvDTE::DebuggerPtr debugger = visualStudio::FindDebuggerAttachedToProcess(processId);
 			if (debugger)
 			{
-				m_restartedProcessIdToDebugger.emplace(processId, debugger);
+				m_restartedProcessIdToDebugger.emplace(+processId, debugger);
 			}
+// BEGIN EPIC MOD
 #endif
+// END EPIC MOD
 		}
 	}
 
@@ -213,7 +248,9 @@ void ServerCommandThread::RestartTargets(void)
 	for (size_t i = 0u; i < count; ++i)
 	{
 		LiveProcess* liveProcess = m_liveProcesses[i];
+		// BEGIN EPIC MOD
 		liveProcess->Restart(m_restartJob);
+		// END EPIC MOD
 	}
 
 	// BEGIN EPIC MOD - Prevent orphaned console instances if processes fail to restart. Job object will be duplicated into child process.
@@ -226,6 +263,14 @@ void ServerCommandThread::RestartTargets(void)
 }
 
 
+void ServerCommandThread::CompileChanges(void)
+{
+	// protect against accepting this command while compilation is already in progress
+	CriticalSection::ScopedLock lock(&m_actionCS);
+	m_manualRecompileTriggered = true;
+}
+
+
 std::wstring ServerCommandThread::GetProcessImagePath(void) const
 {
 	// there must be at least one registered process.
@@ -236,16 +281,16 @@ std::wstring ServerCommandThread::GetProcessImagePath(void) const
 		return L"Unknown";
 	}
 
-	return process::GetImagePath(m_liveProcesses[0]->GetProcessHandle());
+	return Process::GetImagePath(m_liveProcesses[0]->GetProcessHandle()).GetString();
 }
 
 
-scheduler::Task<LiveModule*>* ServerCommandThread::LoadModule(unsigned int processId, void* moduleBase, const wchar_t* givenModulePath, scheduler::TaskBase* taskRoot)
+scheduler::Task<LiveModule*>* ServerCommandThread::LoadModule(Process::Id processId, void* moduleBase, const wchar_t* givenModulePath, scheduler::TaskBase* taskRoot)
 {
 	// note that the path we get from the client might not be normalized, depending on how the executable was launched.
 	// it is crucial to normalize the path again, otherwise we could load already loaded modules into the same
 	// Live++ instance, which would wreak havoc
-	const std::wstring& modulePath = file::NormalizePath(givenModulePath);
+	const std::wstring& modulePath = Filesystem::NormalizePath(givenModulePath).GetString();
 	const executable::Header imageHeader = GetImageHeader(modulePath.c_str());
 	if (!executable::IsValidHeader(imageHeader))
 	{
@@ -270,7 +315,7 @@ scheduler::Task<LiveModule*>* ServerCommandThread::LoadModule(unsigned int proce
 			if (otherLiveProcess->TriedToLoadImage(imageHeader))
 			{
 				// some *other* process loaded this module already
-				LC_LOG_USER("Registering module %S (PID: %d)", modulePath.c_str(), processId);
+				LC_LOG_USER("Registering module %S (PID: %d)", modulePath.c_str(), +processId);
 
 				LiveModule* liveModule = m_imageHeaderToLiveModule[imageHeader];
 				if (liveModule)
@@ -281,7 +326,7 @@ scheduler::Task<LiveModule*>* ServerCommandThread::LoadModule(unsigned int proce
 					const bool installedPatchesSuccessfully = liveModule->InstallCompiledPatches(liveProcess, moduleBase);
 					if (!installedPatchesSuccessfully)
 					{
-						LC_ERROR_USER("Compiled patches could not be installed (PID: %d)", processId);
+						LC_ERROR_USER("Compiled patches could not be installed (PID: %d)", +processId);
 						liveModule->UnregisterProcess(liveProcess);
 					}
 
@@ -303,8 +348,8 @@ scheduler::Task<LiveModule*>* ServerCommandThread::LoadModule(unsigned int proce
 
 	// accumulate module info
 	{
-		const file::Attributes attributes = file::GetAttributes(modulePath.c_str());
-		const uint64_t size = file::GetSize(attributes);
+		const Filesystem::PathAttributes attributes = Filesystem::GetAttributes(modulePath.c_str());
+		const uint64_t size = Filesystem::GetSize(attributes);
 
 		g_loadedModuleSize.Accumulate(size);
 		g_loadedModuleSize.Print();
@@ -339,12 +384,12 @@ scheduler::Task<LiveModule*>* ServerCommandThread::LoadModule(unsigned int proce
 }
 
 
-bool ServerCommandThread::UnloadModule(unsigned int processId, const wchar_t* givenModulePath)
+bool ServerCommandThread::UnloadModule(Process::Id processId, const wchar_t* givenModulePath)
 {
 	// note that the path we get from the client might not be normalized, depending on how the executable was launched.
 	// it is crucial to normalize the path again, otherwise we could load already loaded modules into the same
 	// Live++ instance, which would wreak havoc
-	const std::wstring& modulePath = file::NormalizePath(givenModulePath);
+	const std::wstring& modulePath = Filesystem::NormalizePath(givenModulePath).GetString();
 	const executable::Header imageHeader = GetImageHeader(modulePath.c_str());
 	if (!executable::IsValidHeader(imageHeader))
 	{
@@ -524,11 +569,11 @@ void ServerCommandThread::PrewarmCompilerEnvironmentCache(void)
 }
 
 
-unsigned int ServerCommandThread::ServerThread(void)
+Thread::ReturnValue ServerCommandThread::ServerThread(void)
 {
 	// keep named shared memory alive so that restarted processes don't try spawning new Live++ instances
-	NamedSharedMemory sharedMemory(primitiveNames::StartupNamedSharedMemory(m_processGroupName).c_str());
-	sharedMemory.Write(::GetCurrentProcessId());
+	Process::NamedSharedMemory* sharedMemory = Process::CreateNamedSharedMemory(primitiveNames::StartupNamedSharedMemory(m_processGroupName).c_str(), 4096u);
+	Process::WriteNamedSharedMemory(sharedMemory, ::GetCurrentProcessId());
 
 	// inter process event for telling client that server is ready
 	Event serverReadyEvent(primitiveNames::ServerReadyEvent(m_processGroupName).c_str(), Event::Type::AUTO_RESET);
@@ -550,8 +595,10 @@ unsigned int ServerCommandThread::ServerThread(void)
 		context->exceptionPipe.WaitForClient();
 
 		// a new client has connected, open a new thread for communication
-		context->commandThread = thread::Create("Live coding client command communication", 64u * 1024u, &ServerCommandThread::CommandThread, this, &context->pipe, context->readyEvent);
-		context->exceptionCommandThread = thread::Create("Live coding client exception command communication", 64u * 1024u, &ServerCommandThread::ExceptionCommandThread, this, &context->exceptionPipe);
+		// BEGIN EPIC MOD
+		context->commandThread = Thread::CreateFromMemberFunction("Live coding client command communication", 64u * 1024u, this, &ServerCommandThread::CommandThread, &context->pipe, context->readyEvent);
+		context->exceptionCommandThread = Thread::CreateFromMemberFunction("Live coding client exception command communication", 64u * 1024u, this, &ServerCommandThread::ExceptionCommandThread, &context->exceptionPipe);
+		// END EPIC MOD
 
 		// register this connection
 		{
@@ -560,7 +607,9 @@ unsigned int ServerCommandThread::ServerThread(void)
 		}
 	}
 
-	return 0u;
+	Process::DestroyNamedSharedMemory(sharedMemory);
+
+	return Thread::ReturnValue(0u);
 }
 
 
@@ -573,7 +622,7 @@ BOOL CALLBACK FocusApplicationWindows(HWND WindowHandle, LPARAM Lparam)
 	const types::vector<LiveProcess*>& Processes = *(const types::vector<LiveProcess*>*)Lparam;
 	for (LiveProcess* Process : Processes)
 	{
-		if (Process->GetProcessId() == WindowProcessId && IsWindowVisible(WindowHandle))
+		if (+Process->GetProcessId() == WindowProcessId && IsWindowVisible(WindowHandle))
 		{
 			SetForegroundWindow(WindowHandle);
 		}
@@ -605,21 +654,21 @@ struct ClientProxyThread
 	LiveProcess* m_process;
 	DuplexPipeClient* m_pipe;
 	std::vector<std::wstring> m_enableModules;
-	thread::Handle m_threadHandle;
+	Thread::Handle m_threadHandle;
 
 	ClientProxyThread(LiveProcess* process, DuplexPipeClient* pipe, const std::vector<std::wstring> enableModules)
 		: m_process(process)
 		, m_pipe(pipe)
 		, m_enableModules(enableModules)
 	{
-		m_threadHandle = thread::Create(64u * 1024u, &StaticEntryPoint, this);
-		thread::SetName("Live coding client proxy");
+		m_threadHandle = Thread::Create(64u * 1024u, &StaticEntryPoint, this);
+		Thread::Current::SetName("Live coding client proxy");
 	}
 
 	~ClientProxyThread()
 	{
-		thread::Join(m_threadHandle);
-		thread::Close(m_threadHandle);
+		Thread::Join(m_threadHandle);
+		Thread::Close(m_threadHandle);
 	}
 
 	static unsigned int __stdcall StaticEntryPoint(void* context)
@@ -662,7 +711,7 @@ bool ServerCommandThread::EnableRequiredModules(const TArray<FString>& RequiredM
 		types::vector<std::wstring> LoadModuleFileNames;
 		for (const FString& RequiredModule : RequiredModules)
 		{
-			std::wstring ModuleFileName = file::NormalizePath(*RequiredModule);
+			std::wstring ModuleFileName = Filesystem::NormalizePath(*RequiredModule).GetString();
 			if (liveProcess->IsPendingLazyLoadedModule(ModuleFileName))
 			{
 				LoadModuleFileNames.push_back(ModuleFileName);
@@ -776,16 +825,16 @@ void ServerCommandThread::CompileChanges(bool didAllProcessesMakeProgress)
 
 		for(const TPair<FString, TArray<FString>>& Pair : ModuleToObjectFiles)
 		{
-			std::wstring ModuleFileName = file::NormalizePath(*Pair.Key);
+			std::wstring ModuleFileName = Filesystem::NormalizePath(*Pair.Key).GetString();
 			if(ValidModuleFileNames.find(ModuleFileName) == ValidModuleFileNames.end())
 			{
 				// We couldn't find this exact module filename, but this could be a staged executable. See if we can just match the name.
-				std::wstring ModuleFileNameOnly = file::GetFilename(ModuleFileName);
+				std::wstring ModuleFileNameOnly = Filesystem::GetFilename(ModuleFileName.c_str()).GetString();
 
 				bool bFoundNameMatch = false;
 				for (const LiveModule* liveModule : m_liveModules)
 				{
-					if (ModuleFileNameOnly == file::GetFilename(liveModule->GetModuleName()))
+					if (ModuleFileNameOnly == Filesystem::GetFilename(liveModule->GetModuleName().c_str()).GetString())
 					{
 						ModuleFileName = liveModule->GetModuleName();
 						bFoundNameMatch = true;
@@ -802,10 +851,10 @@ void ServerCommandThread::CompileChanges(bool didAllProcessesMakeProgress)
 				}
 			}
 
-			types::vector<LiveModule::ModifiedObjFile> ObjectFiles;
+			types::vector<symbols::ModifiedObjFile> ObjectFiles;
 			for(const FString& ObjectFile : Pair.Value)
 			{
-				std::wstring NormalizedObjectFile = file::NormalizePath(*ObjectFile);
+				std::wstring NormalizedObjectFile = Filesystem::NormalizePath(*ObjectFile).GetString();
 
 				// If this file has a .lc.obj suffix, temporarily replace the original .obj file while generating the patch.
 				// It'd be nice to track this explicitly inside Live++ and just load the new file, but it requires a lot of changes and would make upgrades difficult.
@@ -818,23 +867,23 @@ void ServerCommandThread::CompileChanges(bool didAllProcessesMakeProgress)
 					OriginalObjectFile += L".obj";
 
 					// Back up the original file, if it exists
-					file::Attributes OriginalFileAttributes = file::GetAttributes(OriginalObjectFile.c_str());
-					if (file::DoesExist(OriginalFileAttributes))
+					Filesystem::PathAttributes OriginalFileAttributes = Filesystem::GetAttributes(OriginalObjectFile.c_str());
+					if (Filesystem::DoesExist(OriginalFileAttributes))
 					{
 						std::wstring OriginalObjectFileBackup = OriginalObjectFile + L".lctmp";
 						m_restoreFiles.push_back(std::make_pair(OriginalObjectFileBackup, OriginalObjectFile));
-						file::DeleteIfExists(OriginalObjectFileBackup.c_str());
-						file::Move(OriginalObjectFile.c_str(), OriginalObjectFileBackup.c_str());
+						Filesystem::DeleteIfExists(OriginalObjectFileBackup.c_str());
+						Filesystem::Move(OriginalObjectFile.c_str(), OriginalObjectFileBackup.c_str());
 					}
 
 					// Move the new file into place
 					m_restoreFiles.push_back(std::make_pair(OriginalObjectFile, NormalizedObjectFile));
-					file::Move(NormalizedObjectFile.c_str(), OriginalObjectFile.c_str());
+					Filesystem::Move(NormalizedObjectFile.c_str(), OriginalObjectFile.c_str());
 					NormalizedObjectFile = OriginalObjectFile;
 				}
 
 				// Add the file to the list of modifications
-				LiveModule::ModifiedObjFile ModifiedObjFile;
+				symbols::ModifiedObjFile ModifiedObjFile;
 				ModifiedObjFile.objPath = NormalizedObjectFile;
 				ObjectFiles.push_back(std::move(ModifiedObjFile));
 			}
@@ -889,14 +938,14 @@ void ServerCommandThread::CompileChanges(bool didAllProcessesMakeProgress)
 			else
 			{
 				// build a patch with the given list of .objs for this module
-				const types::vector<LiveModule::ModifiedObjFile>& objFiles = objFilesIt->second;
+				const types::vector<symbols::ModifiedObjFile>& objFiles = objFilesIt->second;
 				moduleUpdateError = liveModule->Update(&fileCache, m_directoryCache, updateType, objFiles);
 			}
 		}
 		else
 		{
 			// no optional .objs were given, update all live modules regularly
-			types::vector<LiveModule::ModifiedObjFile> emptyObjs;
+			types::vector<symbols::ModifiedObjFile> emptyObjs;
 			moduleUpdateError = liveModule->Update(&fileCache, m_directoryCache, updateType, emptyObjs);
 		}
 
@@ -918,7 +967,9 @@ void ServerCommandThread::CompileChanges(bool didAllProcessesMakeProgress)
 		// bring Live++ to front on success
 		if (appSettings::g_receiveFocusOnRecompile->GetValue() == appSettings::FocusOnRecompile::ON_SUCCESS)
 		{
+			// BEGIN EPIC MOD
 			GLiveCodingServer->GetBringToFrontDelegate().ExecuteIfBound();
+			// END EPIC MOD
 		}
 
 		// play sound on success
@@ -939,7 +990,9 @@ void ServerCommandThread::CompileChanges(bool didAllProcessesMakeProgress)
 		// bring Live++ to front on failure
 		if (appSettings::g_receiveFocusOnRecompile->GetValue() == appSettings::FocusOnRecompile::ON_ERROR)
 		{
+			// BEGIN EPIC MOD
 			GLiveCodingServer->GetBringToFrontDelegate().ExecuteIfBound();
+			// END EPIC MOD
 		}
 
 		// play sound on error
@@ -985,7 +1038,7 @@ void ServerCommandThread::CompileChanges(bool didAllProcessesMakeProgress)
 			break;
 	}
 	// END EPIC MOD
-
+	
 	LC_LOG_USER("---------- Finished (%.3fs) ----------", scope.ReadSeconds());
 
 	// EPIC REMOVED: g_theApp.GetMainFrame()->ResetStatusBarText();
@@ -1044,7 +1097,7 @@ void ServerCommandThread::CallPostcompileHooks(bool didAllProcessesMakeProgress)
 }
 // END EPIC MOD
 
-unsigned int ServerCommandThread::CompileThread(void)
+Thread::ReturnValue ServerCommandThread::CompileThread(void)
 {
 	input::Key keyControl(VK_CONTROL);
 	input::Key keyAlt(VK_MENU);
@@ -1094,16 +1147,31 @@ unsigned int ServerCommandThread::CompileThread(void)
 		const unsigned int changeNotificationTimeout = static_cast<unsigned int>(appSettings::g_continuousCompilationTimeout->GetValue());
 
 		const bool foundAnyModification = changeNotification.CheckOnce();
-		if (foundAnyModification)
+		if (isShortcutPressed || foundAnyModification || m_manualRecompileTriggered)
 		{
 			// clear the log if desired by the user
 			if (appSettings::g_clearLogOnRecompile->GetValue())
 			{
+				// BEGIN EPIC MOD
 				GLiveCodingServer->GetClearOutputDelegate().ExecuteIfBound();
+				// END EPIC MOD
 			}
 
-			LC_SUCCESS_USER("Detected file modification, re-checking until timeout (%d ms)", changeNotificationTimeout);
-			changeNotification.CheckNext(changeNotificationTimeout);
+			if (foundAnyModification)
+			{
+				LC_SUCCESS_USER("Detected file modification, re-checking until timeout (%d ms)", changeNotificationTimeout);
+				changeNotification.CheckNext(changeNotificationTimeout);
+			}
+			else if (isShortcutPressed)
+			{
+				// BEGIN EPIC MOD
+				LC_SUCCESS_USER("Accepted Live coding shortcut");
+				// END EPIC MDO
+			}
+			else if (m_manualRecompileTriggered)
+			{
+				LC_SUCCESS_USER("Manual recompile triggered");
+			}
 		}
 
 		if (isShortcutPressed || foundAnyModification || m_manualRecompileTriggered)
@@ -1121,12 +1189,12 @@ unsigned int ServerCommandThread::CompileThread(void)
 				for (auto processIt = m_liveProcesses.begin(); processIt != m_liveProcesses.end(); /* nothing */)
 				{
 					LiveProcess* liveProcess = *processIt;
-					process::Handle processHandle = liveProcess->GetProcessHandle();
-					if (!process::IsActive(processHandle))
+					Process::Handle processHandle = liveProcess->GetProcessHandle();
+					if (!Process::IsActive(processHandle))
 					{
 						LC_WARNING_USER("Process %d is no longer valid, disconnecting", liveProcess->GetProcessId());
 
-						process::Close(processHandle);
+						Process::Close(processHandle);
 
 						// tell live modules to remove this process
 						const size_t moduleCount = m_liveModules.size();
@@ -1200,28 +1268,12 @@ unsigned int ServerCommandThread::CompileThread(void)
 			// setup the same virtual drive we had when loading the project
 			AddVirtualDrive();
 
-			if (isShortcutPressed || m_manualRecompileTriggered)
-			{
-				// clear the log if desired by the user
-				if (appSettings::g_clearLogOnRecompile->GetValue())
-				{
-					GLiveCodingServer->GetClearOutputDelegate().ExecuteIfBound();
-				}
-
-				if (isShortcutPressed)
-				{
-					LC_SUCCESS_USER("Accepted live coding shortcut");
-				}
-				else if (m_manualRecompileTriggered)
-				{
-					LC_SUCCESS_USER("Manual recompile triggered");
-				}
-			}
-
 			// bring Live++ to front on shortcut trigger
 			if (appSettings::g_receiveFocusOnRecompile->GetValue() == appSettings::FocusOnRecompile::ON_SHORTCUT)
 			{
+				// BEGIN EPIC MOD
 				GLiveCodingServer->GetBringToFrontDelegate().ExecuteIfBound();
+				// END EPIC MOD
 			}
 
 			// BEGIN EPIC MOD - Add the ability for pre and post compile notifications
@@ -1237,8 +1289,8 @@ unsigned int ServerCommandThread::CompileThread(void)
 			// BEGIN EPIC MOD - Non-destructive compile
 			for (std::vector<std::pair<std::wstring, std::wstring>>::reverse_iterator it = m_restoreFiles.rbegin(); it != m_restoreFiles.rend(); it++)
 			{
-				file::DeleteIfExists(it->second.c_str());
-				file::Move(it->first.c_str(), it->second.c_str());
+				Filesystem::DeleteIfExists(it->second.c_str());
+				Filesystem::Move(it->first.c_str(), it->second.c_str());
 			}
 			m_restoreFiles.clear();
 			// END EPIC MOD
@@ -1280,14 +1332,14 @@ unsigned int ServerCommandThread::CompileThread(void)
 
 		m_restartCS.Leave();
 
-		thread::Sleep(10u);
+		Thread::Current::SleepMilliSeconds(10u);
 	}
 
-	return 0u;
+	return Thread::ReturnValue(0u);
 }
 
 
-unsigned int ServerCommandThread::CommandThread(DuplexPipeServer* pipe, Event* readyEvent)
+Thread::ReturnValue ServerCommandThread::CommandThread(DuplexPipeServer* pipe, Event* readyEvent)
 {
 	// handle incoming commands
 	CommandMap commandMap;
@@ -1333,7 +1385,7 @@ unsigned int ServerCommandThread::CommandThread(DuplexPipeServer* pipe, Event* r
 			// pipe was closed or is broken, bail out.
 			// remove ourselves from the array of threads first.
 			RemoveCommandThread(pipe);
-			return 1u;
+			return Thread::ReturnValue(1u);
 		}
 
 		// wait until we're allowed to handle commands again
@@ -1344,11 +1396,11 @@ unsigned int ServerCommandThread::CommandThread(DuplexPipeServer* pipe, Event* r
 	}
 
 	RemoveCommandThread(pipe);
-	return 0u;
+	return Thread::ReturnValue(0u);
 }
 
 
-unsigned int ServerCommandThread::ExceptionCommandThread(DuplexPipeServer* exceptionPipe)
+Thread::ReturnValue ServerCommandThread::ExceptionCommandThread(DuplexPipeServer* exceptionPipe)
 {
 	// handle incoming exception commands
 	CommandMap commandMap;
@@ -1360,11 +1412,11 @@ unsigned int ServerCommandThread::ExceptionCommandThread(DuplexPipeServer* excep
 		if ((!success) || (!exceptionPipe->IsValid()))
 		{
 			// pipe was closed or is broken, bail out
-			return 1u;
+			return Thread::ReturnValue(1u);
 		}
 	}
 
-	return 0u;
+	return Thread::ReturnValue(0u);
 }
 
 
@@ -1389,7 +1441,7 @@ void ServerCommandThread::RemoveCommandThread(const DuplexPipe* pipe)
 }
 
 
-LiveProcess* ServerCommandThread::FindProcessById(unsigned int processId)
+LiveProcess* ServerCommandThread::FindProcessById(Process::Id processId)
 {
 	const size_t count = m_liveProcesses.size();
 	for (size_t i = 0u; i < count; ++i)
@@ -1458,7 +1510,7 @@ bool ServerCommandThread::actions::BuildPatch::Execute(const CommandType* comman
 	for (unsigned int i = 0u; i < command->fileCount; ++i)
 	{
 		const commands::BuildPatch::PatchData patchData = payloadStream.Read<commands::BuildPatch::PatchData>();
-		const LiveModule::ModifiedObjFile modifiedObjFile = { patchData.objPath, patchData.amalgamatedObjPath };
+		const symbols::ModifiedObjFile modifiedObjFile = { patchData.objPath, patchData.amalgamatedObjPath };
 
 		commandThread->m_liveModuleToModifiedOrNewObjFiles[patchData.moduleName].push_back(modifiedObjFile);
 	}
@@ -1636,7 +1688,7 @@ bool ServerCommandThread::actions::EnableLazyLoadedModule::Execute(const Command
 	CriticalSection::ScopedLock lock(&commandThread->m_actionCS);
 
 	// Check if this module is already enabled - it may have been lazy-loaded, then fully loaded, by a restarted process. If so, translate this into a call to EnableModules.
-	const std::wstring modulePath = file::NormalizePath(command->fileName);
+	const std::wstring modulePath = Filesystem::NormalizePath(command->fileName).GetString();
 	for (LiveModule* module : commandThread->m_liveModules)
 	{
 		if(module->GetModuleName() == modulePath)
@@ -1682,11 +1734,11 @@ bool ServerCommandThread::actions::RegisterProcess::Execute(const CommandType* c
 	// protect against several client DLLs calling into this action at the same time
 	CriticalSection::ScopedLock lock(&commandThread->m_actionCS);
 
-	process::Handle processHandle = process::Open(command->processId);
+	Process::Handle processHandle = Process::Open(command->processId);
 
 	// check if any live module in this process group has patches installed already
 	{
-		const std::wstring& processPath = process::GetImagePath(processHandle);
+		const std::wstring& processPath = Process::GetImagePath(processHandle).GetString();
 
 		bool registeredSuccessfully = true;
 		if (!appSettings::g_installCompiledPatchesMultiProcess->GetValue())
@@ -1699,8 +1751,10 @@ bool ServerCommandThread::actions::RegisterProcess::Execute(const CommandType* c
 				LiveModule* liveModule = commandThread->m_liveModules[i];
 				if (liveModule->HasInstalledPatches())
 				{
+					// BEGIN EPIC MOD
 					std::wstring caption(L"Live coding - Registering process ");
-					caption += file::GetFilename(processPath);
+					// END EPIC MOD
+					caption += Filesystem::GetFilename(processPath.c_str()).GetString();
 
 					processGroupHasPatches = true;
 					// BEGIN EPIC MOD - Using non-modal error dialog
@@ -1726,35 +1780,39 @@ bool ServerCommandThread::actions::RegisterProcess::Execute(const CommandType* c
 			// commandThread->m_mainFrame->UpdateWindowTitle();
 			// END EPIC MOD
 
-			if (command->restartedProcessId == 0u)
+			if (command->restartedProcessId == Process::Id(0u))
 			{
 				// this is a new process
-				LC_SUCCESS_USER("Registered process %S (PID: %d)", processPath.c_str(), command->processId);
+				LC_SUCCESS_USER("Registered process %S (PID: %d)", processPath.c_str(), +command->processId);
 			}
 			else
 			{
 				// this process was restarted
-				LC_SUCCESS_USER("Registered restarted process %S (PID: %d, previous PID: %d)", processPath.c_str(), command->processId, command->restartedProcessId);
+				LC_SUCCESS_USER("Registered restarted process %S (PID: %d, previous PID: %d)", processPath.c_str(), +command->processId, command->restartedProcessId);
 
+// BEGIN EPIC MOD
 #if WITH_VISUALSTUDIO_DTE
+// END EPIC MOD
 				// reattach the debugger in case the previous process had a debugger attached
 				{
-					auto it = commandThread->m_restartedProcessIdToDebugger.find(command->restartedProcessId);
+					auto it = commandThread->m_restartedProcessIdToDebugger.find(+command->restartedProcessId);
 					if (it != commandThread->m_restartedProcessIdToDebugger.end())
 					{
-						LC_LOG_USER("Reattaching debugger to PID %d", command->processId);
+						LC_LOG_USER("Reattaching debugger to PID %d", +command->processId);
 
 						const EnvDTE::DebuggerPtr& debugger = it->second;
 						const bool success = visualStudio::AttachToProcess(debugger, command->processId);
 						if (!success)
 						{
-							LC_ERROR_USER("Failed to reattach debugger to PID %d", command->processId);
+							LC_ERROR_USER("Failed to reattach debugger to PID %d", +command->processId);
 						}
 
 						commandThread->m_restartedProcessIdToDebugger.erase(it);
 					}
 				}
+// BEGIN EPIC MOD
 #endif
+// END EPIC MOD
 
 				--commandThread->m_restartedProcessCount;
 				if (commandThread->m_restartedProcessCount == 0u)
@@ -1865,7 +1923,7 @@ bool ServerCommandThread::actions::EnableModules::Execute(const CommandType* com
 	}
 
 	// BEGIN EPIC MOD - Suppress output when lazy loading modules
-	if (commandThread->m_compileThread == nullptr || thread::GetId() != thread::GetId(commandThread->m_compileThread))
+	if (+commandThread->m_compileThread == nullptr || Thread::Current::GetId() != Thread::GetId(commandThread->m_compileThread))
 	{
 		if (loadModuleTaskCount > 0u)
 		{
diff --git a/Engine/Source/Developer/Windows/LiveCodingServer/Private/External/LC_ServerCommandThread.h b/Engine/Source/Developer/Windows/LiveCodingServer/Private/External/LC_ServerCommandThread.h
--- a/Engine/Source/Developer/Windows/LiveCodingServer/Private/External/LC_ServerCommandThread.h
+++ b/Engine/Source/Developer/Windows/LiveCodingServer/Private/External/LC_ServerCommandThread.h
@@ -1,9 +1,9 @@
-// Copyright 2011-2019 Molecular Matters GmbH, all rights reserved.
+// Copyright 2011-2020 Molecular Matters GmbH, all rights reserved.
 
 #pragma once
 
-#include "CoreTypes.h"
-#include "LC_Thread.h"
+#include "LC_ThreadTypes.h"
+#include "LC_ProcessTypes.h"
 #include "LC_Commands.h"
 #include "LC_Telemetry.h"
 #include "LC_DuplexPipeServer.h"
@@ -12,14 +12,16 @@
 #include "LC_Scheduler.h"
 #include "LC_Executable.h"
 #include "LC_RunMode.h"
-#include "LC_Types.h"
 #include "LC_LiveModule.h"
+// BEGIN EPIC MOD
+#include "LC_Types.h"
 #include "VisualStudioDTE.h"
-
+// END EPIC MOD
 
 class MainFrame;
 class DirectoryCache;
 class LiveModule;
+class LiveModuleOrbis;
 class LiveProcess;
 
 class ServerCommandThread
@@ -29,17 +31,18 @@ public:
 	~ServerCommandThread(void);
 
 	void RestartTargets(void);
+	void CompileChanges(void);
 
 	std::wstring GetProcessImagePath(void) const;
 
 private:
-	scheduler::Task<LiveModule*>* LoadModule(unsigned int processId, void* moduleBase, const wchar_t* modulePath, scheduler::TaskBase* taskRoot);
-	bool UnloadModule(unsigned int processId, const wchar_t* modulePath);
+	scheduler::Task<LiveModule*>* LoadModule(Process::Id processId, void* moduleBase, const wchar_t* modulePath, scheduler::TaskBase* taskRoot);
+	bool UnloadModule(Process::Id processId, const wchar_t* modulePath);
 
 	void PrewarmCompilerEnvironmentCache(void);
 
-	unsigned int ServerThread(void);
-	unsigned int CompileThread(void);
+	Thread::ReturnValue ServerThread(void);
+	Thread::ReturnValue CompileThread(void);
 
 	// BEGIN EPIC MOD - Add the ability for pre and post compile notifications
 	void CallPrecompileHooks(bool didAllProcessesMakeProgress);
@@ -50,18 +53,18 @@ private:
 	{
 		DuplexPipeServer pipe;
 		Event* readyEvent;
-		thread::Handle commandThread;
+		Thread::Handle commandThread;
 
 		DuplexPipeServer exceptionPipe;
-		thread::Handle exceptionCommandThread;
+		Thread::Handle exceptionCommandThread;
 	};
 
-	unsigned int CommandThread(DuplexPipeServer* pipe, Event* readyEvent);
-	unsigned int ExceptionCommandThread(DuplexPipeServer* exceptionPipe);
+	Thread::ReturnValue CommandThread(DuplexPipeServer* pipe, Event* readyEvent);
+	Thread::ReturnValue ExceptionCommandThread(DuplexPipeServer* exceptionPipe);
 
 	void RemoveCommandThread(const DuplexPipe* pipe);
 
-	LiveProcess* FindProcessById(unsigned int processId);
+	LiveProcess* FindProcessById(Process::Id processId);
 
 	void CompileChanges(bool didAllProcessesMakeProgress);
 
@@ -121,8 +124,8 @@ private:
 	RunMode::Enum m_runMode;
 
 	MainFrame* m_mainFrame;
-	thread::Handle m_serverThread;
-	thread::Handle m_compileThread;
+	Thread::Handle m_serverThread;
+	Thread::Handle m_compileThread;
 
 	types::vector<LiveModule*> m_liveModules;
 	types::vector<LiveProcess*> m_liveProcesses;
@@ -154,13 +157,17 @@ private:
 
 	// for triggering recompiles using the API
 	bool m_manualRecompileTriggered;
-	types::unordered_map<std::wstring, types::vector<LiveModule::ModifiedObjFile>> m_liveModuleToModifiedOrNewObjFiles;
+	types::unordered_map<std::wstring, types::vector<symbols::ModifiedObjFile>> m_liveModuleToModifiedOrNewObjFiles;
 
 	// restart mechanism
 	CriticalSection m_restartCS;
 	void* m_restartJob;
 	unsigned int m_restartedProcessCount;
+// BEGIN EPIC MOD
 #if WITH_VISUALSTUDIO_DTE
+// END EPIC MOD
 	types::unordered_map<unsigned int, EnvDTE::DebuggerPtr> m_restartedProcessIdToDebugger;
+// BEGIN EPIC MOD
 #endif
+// END EPIC MOD
 };
diff --git a/Engine/Source/Developer/Windows/LiveCodingServer/Private/External/LC_Settings.cpp b/Engine/Source/Developer/Windows/LiveCodingServer/Private/External/LC_Settings.cpp
--- a/Engine/Source/Developer/Windows/LiveCodingServer/Private/External/LC_Settings.cpp
+++ b/Engine/Source/Developer/Windows/LiveCodingServer/Private/External/LC_Settings.cpp
@@ -1,8 +1,11 @@
-// Copyright 2011-2019 Molecular Matters GmbH, all rights reserved.
+// Copyright 2011-2020 Molecular Matters GmbH, all rights reserved.
 
+// BEGIN EPIC MOD
+//#include PCH_INCLUDE
+// END EPIC MOD
 #include "LC_Settings.h"
 #include "LC_AppSettings.h"
-#include "LC_FileUtil.h"
+#include "LC_Filesystem.h"
 #include "LC_StringUtil.h"
 
 
@@ -20,8 +23,8 @@ namespace
 #else
 		// first try loading the setting from the project settings file
 		const std::wstring& projectSettingsPath = appSettings::GetProjectSettingsPath();
-		const file::Attributes& attributes = file::GetAttributes(projectSettingsPath.c_str());
-		if (file::DoesExist(attributes))
+		const Filesystem::PathAttributes& attributes = Filesystem::GetAttributes(projectSettingsPath.c_str());
+		if (Filesystem::DoesExist(attributes))
 		{
 			// a file is there, so try loading the setting
 			const UINT ILLEGAL_VALUE = static_cast<UINT>(-1);
@@ -56,8 +59,8 @@ namespace
 #else
 		// first try loading the setting from the project settings file
 		const std::wstring& projectSettingsPath = appSettings::GetProjectSettingsPath();
-		const file::Attributes& attributes = file::GetAttributes(projectSettingsPath.c_str());
-		if (file::DoesExist(attributes))
+		const Filesystem::PathAttributes& attributes = Filesystem::GetAttributes(projectSettingsPath.c_str());
+		if (Filesystem::DoesExist(attributes))
 		{
 			// a file is there, so try loading the setting
 			const wchar_t* ILLEGAL_VALUE = L"__ILLEGAL_STRING__";
diff --git a/Engine/Source/Developer/Windows/LiveCodingServer/Private/External/LC_SymbolPatterns.cpp b/Engine/Source/Developer/Windows/LiveCodingServer/Private/External/LC_SymbolPatterns.cpp
--- a/Engine/Source/Developer/Windows/LiveCodingServer/Private/External/LC_SymbolPatterns.cpp
+++ b/Engine/Source/Developer/Windows/LiveCodingServer/Private/External/LC_SymbolPatterns.cpp
@@ -1,5 +1,8 @@
-// Copyright 2011-2019 Molecular Matters GmbH, all rights reserved.
+// Copyright 2011-2020 Molecular Matters GmbH, all rights reserved.
 
+// BEGIN EPIC MOD
+//#include PCH_INCLUDE
+// END EPIC MOD
 #include "LC_SymbolPatterns.h"
 
 
@@ -65,15 +68,6 @@ namespace symbolPatterns
 		"$initializer$"
 	};
 
-	extern const char* const WEAK_SYMBOL_PATTERNS[4] =
-	{
-		// these weak symbols are allowed per the standard and need to be special-cased in code
-		"??2",					// operator new
-		"??3",					// operator delete
-		"??_U",					// operator new[]
-		"??_V"					// operator delete[]
-	};
-	
 	extern const char* const STRING_LITERAL_PATTERNS[2] =
 	{
 		// in decorated form, a string literal is denoted by "??_C@_".
@@ -106,9 +100,9 @@ namespace symbolPatterns
 	};
 
 #if LC_64_BIT
-	extern const char* const EXCEPTION_RELATED_PATTERNS[16] =
+	extern const char* const EXCEPTION_RELATED_PATTERNS[17] =
 #else
-	extern const char* const EXCEPTION_RELATED_PATTERNS[10] =
+	extern const char* const EXCEPTION_RELATED_PATTERNS[11] =
 #endif
 	{
 		// used for C++ exception handling
@@ -120,14 +114,16 @@ namespace symbolPatterns
 		"?catch$",
 		"?fin$",
 		"?filt$",
-		"__catch$",
+		"__catch$",	// CL
+		"?catch$",	// Clang
 		"_CxxThrowException",
 		"__CxxFrameHandler",
 		"__GSHandlerCheck",
 #else
 		"__ehhandler$",
 		"__unwindfunclet$",
-		"__catch$",
+		"__catch$",	// CL
+		"?catch$",	// Clang
 		"__except_handler3",
 		"__except_handler4",
 #endif
@@ -160,9 +156,10 @@ namespace symbolPatterns
 #endif
 	};
 
-	extern const char* const EXCEPTION_CLAUSE_PATTERNS[1] =
+	extern const char* const EXCEPTION_CLAUSE_PATTERNS[2] =
 	{
-		"__catch$"
+		"__catch$",	// CL
+		"?catch$",	// Clang
 	};
 
 	extern const char* const RTC_PATTERNS[8] =
diff --git a/Engine/Source/Developer/Windows/LiveCodingServer/Private/External/LC_SymbolPatterns.h b/Engine/Source/Developer/Windows/LiveCodingServer/Private/External/LC_SymbolPatterns.h
--- a/Engine/Source/Developer/Windows/LiveCodingServer/Private/External/LC_SymbolPatterns.h
+++ b/Engine/Source/Developer/Windows/LiveCodingServer/Private/External/LC_SymbolPatterns.h
@@ -1,9 +1,11 @@
-// Copyright 2011-2019 Molecular Matters GmbH, all rights reserved.
+// Copyright 2011-2020 Molecular Matters GmbH, all rights reserved.
 
 #pragma once
 
+// BEGIN EPIC MOD
 #include "CoreTypes.h"
 #include "LC_Platform.h"
+// END EPIC MOD
 
 namespace symbolPatterns
 {
@@ -14,19 +16,18 @@ namespace symbolPatterns
 	extern const char* const DYNAMIC_INITIALIZER_PATTERNS[1];
 	extern const char* const DYNAMIC_ATEXIT_DESTRUCTORS[1];
 	extern const char* const POINTER_TO_DYNAMIC_INITIALIZER_PATTERNS[1];
-	extern const char* const WEAK_SYMBOL_PATTERNS[4];
 	extern const char* const STRING_LITERAL_PATTERNS[2];
 	extern const char* const LINE_NUMBER_PATTERNS[1];
 	extern const char* const FLOATING_POINT_CONSTANT_PATTERNS[4];
 
 #if LC_64_BIT
-	extern const char* const EXCEPTION_RELATED_PATTERNS[16];
+	extern const char* const EXCEPTION_RELATED_PATTERNS[17];
 #else
-	extern const char* const EXCEPTION_RELATED_PATTERNS[10];
+	extern const char* const EXCEPTION_RELATED_PATTERNS[11];
 #endif
 	extern const char* const EXCEPTION_UNWIND_PATTERNS[1];
 
-	extern const char* const EXCEPTION_CLAUSE_PATTERNS[1];
+	extern const char* const EXCEPTION_CLAUSE_PATTERNS[2];
 
 	extern const char* const RTC_PATTERNS[8];
 	extern const char* const SDL_CHECK_PATTERNS[2];
diff --git a/Engine/Source/Developer/Windows/LiveCodingServer/Private/External/LC_SymbolReconstruction.cpp b/Engine/Source/Developer/Windows/LiveCodingServer/Private/External/LC_SymbolReconstruction.cpp
--- a/Engine/Source/Developer/Windows/LiveCodingServer/Private/External/LC_SymbolReconstruction.cpp
+++ b/Engine/Source/Developer/Windows/LiveCodingServer/Private/External/LC_SymbolReconstruction.cpp
@@ -1,7 +1,9 @@
-// Copyright 2011-2019 Molecular Matters GmbH, all rights reserved.
+// Copyright 2011-2020 Molecular Matters GmbH, all rights reserved.
 
+// BEGIN EPIC MOD
+//#include PCH_INCLUDE
+// END EPIC MOD
 #include "LC_SymbolReconstruction.h"
-#include "LC_Allocators.h"
 
 #include "LC_StringUtil.h"
 #include "LC_PointerUtil.h"
@@ -69,8 +71,16 @@ void symbols::ReconstructFromExecutableCoff
 				}
 				else
 				{
-					LC_LOG_DEV("Unknown symbol %s", symbolName.c_str());
-					++unknownSymbolsToFind;
+					const coff::Section& coffSection = coffDb->sections[symbol->sectionIndex];
+					if (coff::IsMSVCJustMyCodeSection(coffSection.name.c_str()))
+					{
+						LC_LOG_DEV("JustMyCode symbol %s", symbolName.c_str());
+					}
+					else
+					{
+						LC_LOG_DEV("Unknown symbol %s", symbolName.c_str());
+						++unknownSymbolsToFind;
+					}
 				}
 			}
 			else
@@ -101,6 +111,10 @@ void symbols::ReconstructFromExecutableCoff
 				{
 					LC_LOG_DEV("Compiler-generated symbol %s", symbolName.c_str());
 				}
+				else if (symbols::IsSectionSymbol(symbolName))
+				{
+					LC_LOG_DEV("Section symbol %s", symbolName.c_str());
+				}
 				else if (strippedSymbols.find(symbolName) != strippedSymbols.end())
 				{
 					LC_LOG_DEV("Stripped symbol %s", symbolName.c_str());
@@ -175,6 +189,19 @@ walkOpenSymbols:
 		for (size_t i = 0u; i < relocationCount; ++i)
 		{
 			const coff::Relocation* relocation = symbol->relocations[i];
+
+			// ignore relocations to symbols in .msvcjmc (MSVC JustMyCode) sections
+			if (relocation->dstSectionIndex >= 0)
+			{
+				const uint32_t index = static_cast<uint32_t>(relocation->dstSectionIndex);
+				const coff::Section& section = coffDb->sections[index];
+				if (coff::IsMSVCJustMyCodeSection(section.name.c_str()))
+				{
+					LC_LOG_DEV("Ignoring relocation to symbol in section %s", section.name.c_str());
+					continue;
+				}
+			}
+
 			const ImmutableString& dstSymbolName = coff::GetRelocationDstSymbolName(coffDb, relocation);
 
 			// the symbol we are looking for might already be in the database because of the public symbols gathered from the PDB
@@ -221,6 +248,11 @@ walkOpenSymbols:
 				LC_LOG_DEV("Ignoring destination symbol \"%s\"", dstSymbolName.c_str());
 				continue;
 			}
+			else if (symbols::IsSectionSymbol(dstSymbolName))
+			{
+				LC_LOG_DEV("Ignoring section symbol \"%s\"", dstSymbolName.c_str());
+				continue;
+			}
 
 			if (!srcSymbol)
 			{
@@ -414,11 +446,8 @@ walkOpenSymbols:
 			}
 
 			// we found a new symbol, add it to the database
-			symbols::Symbol* newSymbol = LC_NEW(&g_symbolAllocator, symbols::Symbol) { dstSymbolName, dstRva32 };
-
 			LC_LOG_DEV("Found new symbol %s at RVA 0x%X", dstSymbolName.c_str(), dstRva32);
-			symbolDB->symbolsByName.emplace(dstSymbolName, newSymbol);
-			symbolDB->symbolsByRva.emplace(dstRva32, newSymbol);
+			symbols::CreateNewSymbol(dstSymbolName, dstRva32, symbolDB);
 
 			// walk the relocations of the new symbol as well
 			const coff::Symbol* nextSymbol = coff::GetSymbolByIndex(coffDb, relocation->dstSymbolNameIndex);
@@ -543,7 +572,7 @@ walkOpenSymbols:
 	}
 
 	// next try finding the missing symbols.
-	// NOTE: this is carefully constructed to only run into O(N²) in rare edge cases, because the original O(N²) algorithm
+	// NOTE: this is carefully constructed to only run into O(N^2) in rare edge cases, because the original O(N^2) algorithm
 	// caused a 25-30s slowdown for some users.
 	const size_t missingSymbolCount = missingSymbols.size();
 
@@ -560,6 +589,11 @@ walkOpenSymbols:
 		const uint64_t isExceptionClauseSymbol = symbols::IsExceptionClauseSymbol(missingSymbolName) ? (1ull << 32ull) : 0ull;
 
 		const coff::Section& coffSection = coffDb->sections[symbol->sectionIndex];
+		if (coff::IsMSVCJustMyCodeSection(coffSection.name.c_str()))
+		{
+			LC_LOG_DEV("Ignoring JustMyCode symbol %s in section %s", missingSymbolName.c_str(), coffSection.name.c_str());
+			continue;
+		}
 
 		LC_LOG_DEV("Trying to find RVA for static symbol %s in section %s", missingSymbolName.c_str(), coffSection.name.c_str());
 		LC_LOG_INDENT_DEV;
@@ -649,6 +683,13 @@ walkOpenSymbols:
 		const coff::Symbol* symbol = missingSymbols[i];
 
 		const ImmutableString& missingSymbolName = coff::GetSymbolName(coffDb, symbol);
+		const coff::Section& coffSection = coffDb->sections[symbol->sectionIndex];
+		if (coff::IsMSVCJustMyCodeSection(coffSection.name.c_str()))
+		{
+			LC_LOG_DEV("Ignoring JustMyCode symbol %s in section %s", missingSymbolName.c_str(), coffSection.name.c_str());
+			continue;
+		}
+
 		const std::string& coffUndecoratedName = symbols::UndecorateSymbolName(missingSymbolName);
 
 		auto diaNameIt = diaNameToRva.find(ImmutableString(coffUndecoratedName.c_str()));
@@ -660,10 +701,7 @@ walkOpenSymbols:
 
 			LC_LOG_DEV("Fast path, found symbol %s at 0x%X", missingSymbolName.c_str(), rva);
 
-			symbols::Symbol* newSymbol = LC_NEW(&g_symbolAllocator, symbols::Symbol) { missingSymbolName, rva };
-
-			symbolDB->symbolsByName.emplace(missingSymbolName, newSymbol);
-			symbolDB->symbolsByRva.emplace(rva, newSymbol);
+			symbols::CreateNewSymbol(missingSymbolName, rva, symbolDB);
 
 			openSymbols.push_back(symbol);
 
@@ -691,7 +729,6 @@ walkOpenSymbols:
 			types::unordered_set<const symbols::Contribution*> potentialContributions;
 			potentialContributions.reserve(contributionsForThisCompiland->size());
 
-			const coff::Section& coffSection = coffDb->sections[symbol->sectionIndex];
 			const uint32_t sectionRelativeAddress = symbol->rva - coffSection.rawDataRva;
 
 			// find this section in the image
@@ -792,10 +829,7 @@ walkOpenSymbols:
 
 				LC_LOG_DEV("Slow path, found symbol %s at 0x%X", missingSymbolName.c_str(), rva);
 
-				symbols::Symbol* newSymbol = LC_NEW(&g_symbolAllocator, symbols::Symbol) { missingSymbolName, rva };
-
-				symbolDB->symbolsByName.emplace(missingSymbolName, newSymbol);
-				symbolDB->symbolsByRva.emplace(rva, newSymbol);
+				CreateNewSymbol(missingSymbolName, rva, symbolDB);
 
 				openSymbols.push_back(symbol);
 
diff --git a/Engine/Source/Developer/Windows/LiveCodingServer/Private/External/LC_SymbolReconstruction.h b/Engine/Source/Developer/Windows/LiveCodingServer/Private/External/LC_SymbolReconstruction.h
--- a/Engine/Source/Developer/Windows/LiveCodingServer/Private/External/LC_SymbolReconstruction.h
+++ b/Engine/Source/Developer/Windows/LiveCodingServer/Private/External/LC_SymbolReconstruction.h
@@ -1,11 +1,12 @@
-// Copyright 2011-2019 Molecular Matters GmbH, all rights reserved.
+// Copyright 2011-2020 Molecular Matters GmbH, all rights reserved.
 
 #pragma once
 
+// BEGIN EPIC MOD
 #include "CoreTypes.h"
+// END EPIC MOD
 #include "LC_Coff.h"
 #include "LC_Symbols.h"
-#include "LC_Process.h"
 #include "LC_Executable.h"
 
 
