Integration of Live++ 1.6.3

diff --git a/Engine/Source/Developer/Windows/LiveCodingServer/Private/External/LC_Symbols.cpp b/Engine/Source/Developer/Windows/LiveCodingServer/Private/External/LC_Symbols.cpp
--- a/Engine/Source/Developer/Windows/LiveCodingServer/Private/External/LC_Symbols.cpp
+++ b/Engine/Source/Developer/Windows/LiveCodingServer/Private/External/LC_Symbols.cpp
@@ -1,7 +1,10 @@
-// Copyright 2011-2019 Molecular Matters GmbH, all rights reserved.
+// Copyright 2011-2020 Molecular Matters GmbH, all rights reserved.
 
+// BEGIN EPIC MOD
+//#include PCH_INCLUDE
+// END EPIC MOD
 #include "LC_Symbols.h"
-#include "LC_FileUtil.h"
+#include "LC_Filesystem.h"
 #include "LC_StringUtil.h"
 #include "LC_Process.h"
 #include "LC_Telemetry.h"
@@ -19,6 +22,7 @@
 #include "LC_Amalgamation.h"
 #include "LC_CompilerOptions.h"
 #include "LC_AppSettings.h"
+// BEGIN EPIC MOD
 #include "LC_Allocators.h"
 #include <diacreate.h>
 #include <algorithm>
@@ -30,74 +34,7 @@
 #include "Dom/JsonObject.h"
 
 #include "Windows/AllowWindowsPlatformAtomics.h"
-
-namespace
-{
-	static void RecurseTypeName(IDiaSymbol* diaSymbol, types::unordered_set<uint32_t>& userDefinedTypes)
-	{
-		IDiaSymbol* typeSymbol = nullptr;
-		if (diaSymbol->get_type(&typeSymbol) == S_OK)
-		{
-			DWORD tag = 0u;
-			typeSymbol->get_symTag(&tag);
-
-			if (tag == SymTagPointerType)
-			{
-				RecurseTypeName(typeSymbol, userDefinedTypes);
-			}
-			else if (tag == SymTagUDT)
-			{
-				// found a user-defined type
-				DWORD id = 0u;
-				typeSymbol->get_symIndexId(&id);
-				userDefinedTypes.emplace(id);
-			}
-			else if (tag == SymTagArrayType)
-			{
-				RecurseTypeName(typeSymbol, userDefinedTypes);
-			}
-			else if (tag == SymTagFunctionType)
-			{
-				// the type symbol represents the function signature. recurse possible UDTs from there
-				RecurseTypeName(typeSymbol, userDefinedTypes);
-
-				// grab all argument types and find possible UDTs from there
-				const types::vector<IDiaSymbol*>& argSymbols = dia::GatherChildSymbols(typeSymbol, SymTagFunctionArgType);
-				for (size_t i = 0u; i < argSymbols.size(); ++i)
-				{
-					IDiaSymbol* arg = argSymbols[i];
-					RecurseTypeName(arg, userDefinedTypes);
-					arg->Release();
-				}
-			}
-
-			typeSymbol->Release();
-		}
-	}
-
-
-	static void FindUdtsFromData(IDiaSymbol* diaSymbol, types::unordered_set<uint32_t>& userDefinedTypes)
-	{
-		RecurseTypeName(diaSymbol, userDefinedTypes);
-	}
-
-
-	static void FindUdtsFromFunction(IDiaSymbol* diaSymbol, types::unordered_set<uint32_t>& userDefinedTypes)
-	{
-		// the function type symbol represents the function signature. recurse possible UDTs from there
-		RecurseTypeName(diaSymbol, userDefinedTypes);
-
-		// gather possible UDTs from data used in function (e.g. local variables) as well
-		const types::vector<IDiaSymbol*>& dataSymbols = dia::GatherChildSymbols(diaSymbol, SymTagData);
-		for (size_t i = 0u; i < dataSymbols.size(); ++i)
-		{
-			IDiaSymbol* dataSymbol = dataSymbols[i];
-			FindUdtsFromData(dataSymbol, userDefinedTypes);
-			dataSymbol->Release();
-		}
-	}
-}
-
+// END EPIC MOD
 
 namespace
 {
@@ -216,8 +153,8 @@ namespace
 			if (resultCode == S_OK)
 			{
 				// the PDB was successfully loaded from this path
-				const file::Attributes attributes = file::GetAttributes(pdbPath);
-				const uint64_t size = file::GetSize(attributes);
+				const Filesystem::PathAttributes attributes = Filesystem::GetAttributes(pdbPath);
+				const uint64_t size = Filesystem::GetSize(attributes);
 
 				if (m_openOptions & symbols::OpenOptions::ACCUMULATE_SIZE)
 				{
@@ -250,17 +187,66 @@ namespace
 
 	static symbols::Provider* CreateProvider(const wchar_t* filename, uint32_t openOptions)
 	{
-		IDiaDataSource* diaDataSource = nullptr;
-		HRESULT hr = NoRegCoCreate(L"msdia140.dll", CLSID_DiaSource, IID_IDiaDataSource, reinterpret_cast<void**>(&diaDataSource));
-		if (hr != S_OK)
+		// create a IDiaDataSource manually from the msdia140.dll that we ship with
+		std::wstring msdiaPath = Process::Current::GetImagePath().GetString();
+		msdiaPath = Filesystem::GetDirectory(msdiaPath.c_str()).GetString();
+		msdiaPath += L"\\";
+		msdiaPath += L"msdia140.dll";
+
+		HMODULE diaModule = ::LoadLibraryW(msdiaPath.c_str());
+		if (diaModule == NULL)
 		{
-			LC_ERROR_USER("Cannot create IDiaDataSource instance while trying to load module %S. Error: 0x%X", filename, hr);
+			LC_WARNING_USER("msdia140.dll seems to be missing in the Live++ directory at path %S (module: %S)", msdiaPath.c_str(), filename);
+
+			// let the OS find any msdia140.dll in the system
+			diaModule = ::LoadLibraryW(L"msdia140.dll");
+			if (diaModule == NULL)
+			{
+				LC_ERROR_USER("No msdia140.dll found by the operating system (module: %S)", filename);
+				return nullptr;
+			}
+		}
 
+		typedef decltype(&DllGetClassObject) DllGetClassObjectFunction;
+
+		DllGetClassObjectFunction getClassObjectFunction = reinterpret_cast<DllGetClassObjectFunction>(reinterpret_cast<uintptr_t>(::GetProcAddress(diaModule, "DllGetClassObject")));
+		if (!getClassObjectFunction)
+		{
+			LC_ERROR_USER("Could not find DllGetClassObject function in DIA DLL (module: %S)", filename);
 			return nullptr;
 		}
 
+		IClassFactory* classFactory = nullptr;
+		{
+			const HRESULT hr = getClassObjectFunction(__uuidof(DiaSource), IID_IClassFactory, reinterpret_cast<void**>(&classFactory));
+			if (FAILED(hr))
+			{
+				LC_ERROR_USER("Cannot create IClassFactory instance (module: %S, error: 0x%X)", filename, hr);
+				return nullptr;
+			}
+
+			if (!classFactory)
+			{
+				LC_ERROR_USER("IClassFactory interface not supported by DIA DLL (module: %S, error: 0x%X)", filename, hr);
+				return nullptr;
+			}
+		}
+
+		IDiaDataSource* diaDataSource = nullptr;
+		{
+			const HRESULT hr = classFactory->CreateInstance(NULL, __uuidof(IDiaDataSource), reinterpret_cast<void**>(&diaDataSource));
+			if (FAILED(hr))
+			{
+				LC_ERROR_USER("Cannot create IDiaDataSource instance (module: %S, error: 0x%X)", filename, hr);
+				return nullptr;
+			}
+
+			classFactory->Release();
+		}
+
 		LoadCallback* callback = new LoadCallback(openOptions);
 
+		HRESULT hr = S_OK;
 		if (openOptions & symbols::OpenOptions::USE_SYMBOL_SERVER)
 		{
 			// allow DIA to use a symbol server.
@@ -292,7 +278,7 @@ namespace
 			// BEGIN EPIC MOD - Show a warning if we don't have a PDB for this module. Since we only enable Live++ for modules that we built, we should always have a PDB.
 			if (hr == E_PDB_NOT_FOUND)
 			{
-				LC_WARNING_USER("No PDB file found for module %S. If this is a packaged build, make sure that debug files are being staged. Live coding will be disabled for this module.", file::GetFilename(filename).c_str());
+				LC_WARNING_USER("No PDB file found for module %S. If this is a packaged build, make sure that debug files are being staged. Live coding will be disabled for this module.", Filesystem::GetFilename(filename).GetString());
 			}
 			// END EPIC MOD
 
@@ -317,8 +303,8 @@ namespace
 			return nullptr;
 		}
 
-		const file::Attributes& attributes = file::GetAttributes(filename);
-		const uint64_t lastModification = file::GetLastModificationTime(attributes);
+		const Filesystem::PathAttributes& attributes = Filesystem::GetAttributes(filename);
+		const uint64_t lastModification = Filesystem::GetLastModificationTime(attributes);
 		symbols::Provider* provider = new symbols::Provider { diaDataSource, diaSession, globalScope, lastModification };
 		return provider;
 	}
@@ -329,12 +315,12 @@ namespace
 		if (libraryName.GetString())
 		{
 			// library names also contain .obj files, we are not interested in those
-			const std::wstring& libExtension = file::GetExtension(libraryName.GetString());
-			if (libExtension.length() != 0u)
+			const Filesystem::Path libExtension = Filesystem::GetExtension(libraryName.GetString());
+			if (libExtension.GetLength() != 0u)
 			{
 				// found an extension
-				const std::wstring& uppercaseExtensionName = string::ToUpper(libExtension);
-				if (string::Contains(uppercaseExtensionName.c_str(), L".LIB"))
+				const Filesystem::Path uppercaseExtensionName = libExtension.ToUpper();
+				if (string::Contains(uppercaseExtensionName.GetString(), L".LIB"))
 				{
 					return true;
 				}
@@ -353,36 +339,6 @@ namespace
 	}
 
 
-	static bool IsMainCompilandCpp(const std::wstring& normalizedDependencySrcPath, const std::wstring& objPath)
-	{
-		// it should suffice to only check the source filename (without extension) against the object filename (without extension).
-		// comparisons involving paths are tricky in this case, because certain build systems like FASTBuild automatically
-		// generate unity files, and can do so in different directories, e.g:
-		// OBJ: Z:\Intermediate\x64\Debug\Unity11.obj
-		// SRC: Z:\Unity\Unity11.cpp
-		const std::wstring srcFile = string::ToUpper(file::RemoveExtension(file::GetFilename(normalizedDependencySrcPath)));
-		const std::wstring objFile = string::ToUpper(file::RemoveExtension(file::GetFilename(objPath)));
-		return string::Contains(objFile.c_str(), srcFile.c_str());
-	}
-
-
-	static bool IsCppOrCFile(const std::wstring& normalizedLowercaseFilename)
-	{
-		const std::wstring& filenameExtension = file::GetExtension(normalizedLowercaseFilename);
-		const types::vector<std::wstring>& amalgamatedCppFileExtensions = appSettings::GetAmalgamatedCppFileExtensions();
-
-		for (size_t i = 0u; i < amalgamatedCppFileExtensions.size(); ++i)
-		{
-			if (string::Matches(filenameExtension.c_str(), amalgamatedCppFileExtensions[i].c_str()))
-			{
-				return true;
-			}
-		}
-
-		return false;
-	}
-
-
 	static void AddFileDependency(symbols::CompilandDB* compilandDb, const ImmutableString& changedSrcFile, const ImmutableString& recompiledObjFile, uint64_t srcFileLastModificationTime)
 	{
 		// try updating dependencies for the given file and create a new dependency in case none exists yet
@@ -395,6 +351,7 @@ namespace
 			dependency = LC_NEW(&g_dependencyAllocator, symbols::Dependency);
 			dependency->lastModification = srcFileLastModificationTime;
 			dependency->parentDirectory = nullptr;
+			dependency->hadInitialChange = false;
 		}
 
 		// update entry
@@ -405,6 +362,22 @@ namespace
 
 namespace symbols
 {
+	Symbol* CreateNewSymbol(const ImmutableString& name, uint32_t rva, SymbolDB* db)
+	{
+		Symbol* symbol = LC_NEW(&g_symbolAllocator, Symbol) { name, rva };
+		db->symbolsByName.emplace(name, symbol);
+		db->symbolsByRva.emplace(rva, symbol);
+
+		// test if this symbol contains an anonymous namespace in its name
+		if (string::Find(name.c_str(), symbolPatterns::ANONYMOUS_NAMESPACE_PATTERN))
+		{
+			db->ansSymbols.push_back(symbol);
+		}
+
+		return symbol;
+	}
+
+
 	Provider* OpenEXE(const wchar_t* filename, uint32_t openOptions)
 	{
 		return CreateProvider(filename, openOptions);
@@ -436,7 +409,7 @@ namespace symbols
 
 	bool TryGetCompilandIdFromUnityManifest(const std::wstring& objFile, uint32_t& compilandId)
 	{
-		std::wstring normalizedObjFile = file::NormalizePath(objFile.c_str());
+		std::wstring normalizedObjFile = Filesystem::NormalizePath(objFile.c_str()).GetString();
 		CriticalSection::ScopedLock lock(&g_objFileToCompilandIdCS);
 
 		// Check if it's already cached
@@ -444,7 +417,7 @@ namespace symbols
 		if(it == g_objFileToCompilandId.end())
 		{
 			// Read the manifest file
-			std::wstring BaseDir = file::GetDirectory(normalizedObjFile);
+			std::wstring BaseDir = Filesystem::GetDirectory(normalizedObjFile.c_str()).GetString();
 			std::wstring ManifestFile = BaseDir + L"\\LiveCodingInfo.json";
 
 			// If we've already tried to read this string, don't try again
@@ -478,7 +451,7 @@ namespace symbols
 
 			for (const TPair<FString, TSharedPtr<FJsonValue>>& Pair : FilesObject->Get()->Values)
 			{
-				std::wstring UnityObjectFile = file::NormalizePath((BaseDir + L"\\" + *Pair.Key).c_str());
+				std::wstring UnityObjectFile = Filesystem::NormalizePath((BaseDir + L"\\" + *Pair.Key).c_str()).GetString();
 				uint32_t UnityCompilandId = uniqueId::Generate(UnityObjectFile);
 				g_objFileToCompilandId.insert(std::make_pair(UnityObjectFile, UnityCompilandId));
 
@@ -498,7 +471,7 @@ namespace symbols
 						return false;
 					}
 
-					std::wstring MemberObjFile = file::NormalizePath((BaseDir + L"\\" + *SourceFileValue->AsString()).c_str());
+					std::wstring MemberObjFile = Filesystem::NormalizePath((BaseDir + L"\\" + *SourceFileValue->AsString()).c_str()).GetString();
 					g_objFileToCompilandId.insert(std::make_pair(MemberObjFile, UnityCompilandId));
 				}
 			}
@@ -524,7 +497,7 @@ namespace symbols
 		}
 		else
 		{
-			return uniqueId::Generate(file::NormalizePath(objPath.c_str()));
+			return uniqueId::Generate(Filesystem::NormalizePath(objPath.c_str()).GetString());
 		}
 	}
 	// END EPIC MOD
@@ -542,6 +515,7 @@ namespace symbols
 			symbolDB->symbolsByName.reserve(symbolCount);
 			symbolDB->symbolsByRva.reserve(symbolCount);
 			symbolDB->patchableFunctionSymbols.reserve(symbolCount);
+			symbolDB->ansSymbols.reserve(16u);	// we don't expect to find much of these
 
 			for (size_t i = 0u; i < symbolCount; ++i)
 			{
@@ -572,9 +546,7 @@ namespace symbols
 				}
 				else
 				{
-					Symbol* symbol = LC_NEW(&g_symbolAllocator, Symbol) { symbolName, rva };
-					symbolDB->symbolsByName.emplace(symbolName, symbol);
-					symbolDB->symbolsByRva.emplace(rva, symbol);
+					Symbol* symbol = CreateNewSymbol(symbolName, rva, symbolDB);
 
 					if (dia::IsFunction(publicSymbol))
 					{
@@ -754,40 +726,6 @@ namespace symbols
 	}
 
 
-	UserDefinedTypesDB* GatherUserDefinedTypes(const DiaCompilandDB* diaCompilandDb, const Compiland* compiland)
-	{
-		telemetry::Scope telemetryScope("Gathering user-defined types");
-
-		UserDefinedTypesDB* database = new UserDefinedTypesDB;
-
-		// due to the structure of the internal PDB format, enumerating all user-defined types is far too slow
-		// and doesn't allow grabbing types for a certain compiland only.
-		// therefore, we grab the DIA compiland instead, enumerate its data and function symbols, and reconstruct
-		// the used UDTs from there.
-		IDiaSymbol* diaSymbol = diaCompilandDb->symbols[compiland->diaSymbolIndex];
-
-		const types::vector<IDiaSymbol*>& dataSymbols = dia::GatherChildSymbols(diaSymbol, SymTagData);
-		const size_t dataSymbolCount = dataSymbols.size();
-		for (size_t i = 0u; i < dataSymbolCount; ++i)
-		{
-			IDiaSymbol* symbol = dataSymbols[i];
-			FindUdtsFromData(symbol, database->typeIds);
-			symbol->Release();
-		}
-
-		const types::vector<IDiaSymbol*>& functionSymbols = dia::GatherChildSymbols(diaSymbol, SymTagFunction);
-		const size_t functionSymbolCount = functionSymbols.size();
-		for (size_t i = 0u; i < functionSymbolCount; ++i)
-		{
-			IDiaSymbol* symbol = functionSymbols[i];
-			FindUdtsFromFunction(symbol, database->typeIds);
-			symbol->Release();
-		}
-
-		return database;
-	}
-
-
 	CompilandDB* GatherCompilands(const Provider* provider, const DiaCompilandDB* diaCompilandDb, unsigned int splitAmalgamatedFilesThreshold, uint32_t compilandOptions)
 	{
 		telemetry::Scope telemetryScope("Gathering compilands");
@@ -870,7 +808,7 @@ namespace symbols
 					else if (string::Matches(environmentName.GetString(), L"cl"))
 					{
 						// the path to the compiler is often not normalized, and contains wrong casing
-						optionsCache[3] = file::NormalizePath(environmentOption.GetString());
+						optionsCache[3] = Filesystem::NormalizePath(environmentOption.GetString()).GetString();
 						++foundOptions;
 					}
 					else if (string::Matches(environmentName.GetString(), L"cmd"))
@@ -921,7 +859,7 @@ namespace symbols
 				// optimization: ignore files stored on optical drives because they cannot be changed anyway.
 				{
 					// test the compiland path first
-					if (file::GetDriveType(compilandPath.c_str()) == file::DriveType::OPTICAL)
+					if (Filesystem::GetDriveType(compilandPath.c_str()) == Filesystem::DriveType::OPTICAL)
 					{
 						if (generateLogs)
 						{
@@ -945,11 +883,11 @@ namespace symbols
 						bool testFileExists = (environmentCompilandPath.length() != 0);
 						if (testFileExists)
 						{
-							testPath = file::GetDirectory(environmentCompilandPath);
+							testPath = Filesystem::GetDirectory(environmentCompilandPath.c_str()).GetString();
 							testPath += L"\\";
-							testPath += file::GetFilename(compilandPath);
+							testPath += Filesystem::GetFilename(compilandPath.c_str()).GetString();
 
-							if (file::GetDriveType(testPath.c_str()) == file::DriveType::OPTICAL)
+							if (Filesystem::GetDriveType(testPath.c_str()) == Filesystem::DriveType::OPTICAL)
 							{
 								if (generateLogs)
 								{
@@ -975,13 +913,13 @@ namespace symbols
 						{
 							// try the compiler working directory plus compiland path.
 							// optimization: this can only work if the compiland path is relative
-							if (file::IsRelativePath(compilandPath.c_str()))
+							if (Filesystem::IsRelativePath(compilandPath.c_str()))
 							{
 								testPath = optionsCache[2];
 								testPath += L"\\";
 								testPath += compilandPath;
 
-								if (file::GetDriveType(testPath.c_str()) == file::DriveType::OPTICAL)
+								if (Filesystem::GetDriveType(testPath.c_str()) == Filesystem::DriveType::OPTICAL)
 								{
 									if (generateLogs)
 									{
@@ -1025,7 +963,7 @@ namespace symbols
 					}
 				}
 
-				const std::wstring normalizedCompilandPath = file::NormalizePath(compilandPath.c_str());
+				const std::wstring normalizedCompilandPath = Filesystem::NormalizePath(compilandPath.c_str()).GetString();
 
 				// check for incompatible compiler/linker settings depending on enabled features
 				const bool splitAmalgamatedFiles = (splitAmalgamatedFilesThreshold > 1u);
@@ -1082,8 +1020,9 @@ namespace symbols
 					envCompilerWorkingDirectory,
 					ImmutableString(""),							// amalgamation .obj path
 					nullptr,										// file indices
+					// BEGIN EPIC MOD
 					GetCompilandIdFromPath(normalizedCompilandPath),	// unique ID
-					static_cast<uint32_t>(i),						// dia symbol index
+					// END EPIC MOD
 					Compiland::Type::SINGLE_FILE,					// type of file
 					isPartOfLibrary,								// isPartOfLibrary
 					false											// wasRecompiled
@@ -1105,7 +1044,7 @@ namespace symbols
 				// if we find a file dependency matching the given source file, we take that one instead to get
 				// full absolute file paths.
 				// optimization: ignore all files on optical drives because they cannot be changed anyway
-				const std::wstring srcFileOnlyLowercase = string::ToLower(file::GetFilename(optionsCache[0]));
+				const std::wstring srcFileOnlyLowercase = string::ToLower(Filesystem::GetFilename(optionsCache[0].c_str()).GetString());
 
 				const ObjPath objPath(string::ToUtf8String(normalizedCompilandPath));
 
@@ -1114,7 +1053,7 @@ namespace symbols
 					// we are only interested in tracking .obj files. we will never be able to recompile files
 					// and we don't know anything about source files, dependencies, etc.
 					// but we still use our dependency tracking system by letting each .obj depend on itself.
-					if (file::GetDriveType(normalizedCompilandPath.c_str()) == file::DriveType::OPTICAL)
+					if (Filesystem::GetDriveType(normalizedCompilandPath.c_str()) == Filesystem::DriveType::OPTICAL)
 					{
 						if (generateLogs)
 						{
@@ -1135,284 +1074,48 @@ namespace symbols
 					continue;
 				}
 
+				// fetch all include files for this compiland
 				types::vector<IDiaSourceFile*> sourceFiles = dia::GatherCompilandFiles(provider->diaSession, diaSymbol);
 				const size_t fileCount = sourceFiles.size();
 
-				// gather number of .cpp files first to check whether this compiland is an amalgamated/unity/batch file
-				// (i.e. a .cpp file including several other .cpp files).
-				struct FileInfo
-				{
-					std::wstring normalizedFilename;
-					bool isCppOrCFile;
-				};
+				types::vector<std::wstring> includeFilePaths;
+				includeFilePaths.reserve(fileCount);
 
-				types::vector<FileInfo> fileInfos;
-				fileInfos.reserve(fileCount);
-
-				size_t cppFileCount = 0u;
 				for (size_t j = 0u; j < fileCount; ++j)
 				{
 					IDiaSourceFile* sourceFile = sourceFiles[j];
 					const dia::SymbolName& filename = dia::GetSymbolFilename(sourceFile);
 					const std::wstring wideFilename(filename.GetString());
 
+					includeFilePaths.push_back(wideFilename);
+
+					sourceFile->Release();
+
+					// repair paths to remote machines, e.g. when using FASTBuild.
 					// we are not allowed to normalize this filename. otherwise, normalizing will resolve symbolic links
 					// and virtual drives, which means that files compiled by Live++ will use a different path than
 					// the original compilands.
 					// this could break when including header files that use #pragma once.
 					const std::wstring lowercaseFilename = string::ToLower(wideFilename);
-					const std::wstring lowercaseFilenameOnly = file::GetFilename(lowercaseFilename);
+					const std::wstring lowercaseFilenameOnly = Filesystem::GetFilename(lowercaseFilename.c_str()).GetString();
 					if (string::Matches(lowercaseFilenameOnly.c_str(), srcFileOnlyLowercase.c_str()))
 					{
 						// replace the source path with the full absolute path to make remote builds work.
 						// we convert the path to lower case to be absolutely sure it is at least consistent across PDBs of
 						// patches, executables and DLLs, given the fact that we cannot normalize it.
-						compiland->srcPath = string::ToUtf8String(lowercaseFilename);
-					}
-
-					// AMALGAMATION
-					// skip checking file names when not trying to split amalgamated files
-					const bool isCppOrCFile = splitAmalgamatedFiles ? IsCppOrCFile(lowercaseFilename) : false;
-					fileInfos.emplace_back(FileInfo { lowercaseFilename, isCppOrCFile });
-
-					if (isCppOrCFile)
-					{
-						++cppFileCount;
-					}
-
-					sourceFile->Release();
-				}
-
-				// AMALGAMATION
-				// make sure to treat single-part compilands as being non-amalgamated, i.e. we don't support
-				// recursive amalgamation.
-				// in case splitting of amalgamated files is turned off, this automatically takes care of
-				// treating every compiland as single-file compiland.
-				const bool isPartOfAmalgamation = amalgamation::IsPartOfAmalgamation(compilandPath.c_str());
-				if ((!splitAmalgamatedFiles) || (isPartOfAmalgamation) || (cppFileCount < splitAmalgamatedFilesThreshold))
-				{
-					LC_LOG_DEV("Single .cpp file compiland %s", objPath.c_str());
-
-					// only store source files when splitting amalgamated files in order to save memory
-					// in the general case.
-					if (splitAmalgamatedFiles)
-					{
-						// create array of source file indices for this compiland
-						compiland->sourceFiles = new CompilandSourceFiles;
-						compiland->sourceFiles->files.reserve(fileCount);
-					}
-
-					// this is not an amalgamated compiland
-					for (size_t j = 0u; j < fileCount; ++j)
-					{
-						const FileInfo& fileInfo = fileInfos[j];
-						const std::wstring& normalizedFilename = fileInfo.normalizedFilename;
-
-						if (file::GetDriveType(normalizedFilename.c_str()) == file::DriveType::OPTICAL)
+						// note that when compiling with Clang, there may be more than one file that matches the filename,
+						// e.g. one file pointing to the local path, and one file pointing to a remote path.
+						// therefore, we need to make sure to only take files that really exist on disk.
+						const FileAttributeCache::Data& cacheData = fileCache.UpdateCacheData(wideFilename);
+						if (cacheData.exists)
 						{
-							if (generateLogs)
-							{
-								LC_LOG_DEV("Ignoring file %S on optical drive", normalizedFilename.c_str());
-							}
-						}
-						else
-						{
-							const FileAttributeCache::Data& cacheData = fileCache.UpdateCacheData(normalizedFilename);
-							if (cacheData.exists)
-							{
-								if (generateLogs)
-								{
-									LC_LOG_DEV("Dependency %S", normalizedFilename.c_str());
-								}
-
-								const ImmutableString& sourceFilePath = string::ToUtf8String(normalizedFilename);
-								AddFileDependency(compilandDb, sourceFilePath, objPath, cacheData.lastModificationTime);
-
-								if (splitAmalgamatedFiles)
-								{
-									compiland->sourceFiles->files.push_back(sourceFilePath);
-								}
-							}
-							else if (generateLogs)
-							{
-								LC_LOG_DEV("Missing dependency %S", normalizedFilename.c_str());
-							}
+							compiland->srcPath = string::ToUtf8String(lowercaseFilename);
 						}
 					}
-
-					compilandDb->compilands.insert(std::make_pair(objPath, compiland));
-					compilandDb->compilandNameToObjOnDisk.emplace(string::ToUtf8String(diaCompilandPath.GetString()), objPath);
 				}
-				else
-				{
-					// this is an amalgamated compiland
-					LC_LOG_DEV("Amalgamated .cpp file compiland %s", objPath.c_str());
-
-					// always add a main compiland for the .obj file.
-					// some amalgamated files don't store their main .cpp as dependency.
-					{
-						compiland->type = Compiland::Type::AMALGAMATION;
-						compilandDb->compilands.insert(std::make_pair(objPath, compiland));
-						compilandDb->compilandNameToObjOnDisk.insert(std::make_pair(string::ToUtf8String(diaCompilandPath.GetString()), objPath));
-					}
-
-					for (size_t j = 0u; j < fileCount; ++j)
-					{
-						const FileInfo& fileInfo = fileInfos[j];
-						const std::wstring& normalizedFilename = fileInfo.normalizedFilename;
-						const ImmutableString& sourceFilePath = string::ToUtf8String(normalizedFilename);
-
-						if (fileInfo.isCppOrCFile)
-						{
-							if (IsMainCompilandCpp(normalizedFilename, compilandPath))
-							{
-								LC_LOG_DEV("Main .cpp %S", normalizedFilename.c_str());
-
-								if (file::GetDriveType(normalizedFilename.c_str()) == file::DriveType::OPTICAL)
-								{
-									if (generateLogs)
-									{
-										LC_LOG_DEV("Ignoring file %S on optical drive", normalizedFilename.c_str());
-									}
-								}
-								else
-								{
-									const FileAttributeCache::Data& cacheData = fileCache.UpdateCacheData(normalizedFilename);
-									if (cacheData.exists)
-									{
-										if (generateLogs)
-										{
-											LC_LOG_DEV("Dependency %S", normalizedFilename.c_str());
-										}
-
-										AddFileDependency(compilandDb, sourceFilePath, objPath, cacheData.lastModificationTime);
-									}
-									else if (generateLogs)
-									{
-										LC_LOG_DEV("Missing dependency %S", normalizedFilename.c_str());
-									}
-								}
-							}
-							else
-							{
-								// this is a .cpp file included by the amalgamated file.
-								// add a separate compiland and .obj for this file, and update dependencies so that changing
-								// this source file will not trigger a build of the amalgamated file.
-								LC_LOG_DEV("Included .cpp %S", normalizedFilename.c_str());
-
-								if (file::GetDriveType(normalizedFilename.c_str()) == file::DriveType::OPTICAL)
-								{
-									if (generateLogs)
-									{
-										LC_LOG_DEV("Ignoring file %S on optical drive", normalizedFilename.c_str());
-									}
-								}
-								else
-								{
-									const FileAttributeCache::Data& cacheData = fileCache.UpdateCacheData(normalizedFilename);
-									if (cacheData.exists)
-									{
-										if (generateLogs)
-										{
-											LC_LOG_DEV("Dependency %S", normalizedFilename.c_str());
-										}
-
-										// create new .obj path by appending this file name to the real .obj, e.g.
-										// Amalgamated.obj turns into Amalgamated.lpp_part.ASingleFile.obj.
-										const std::wstring newObjPart = amalgamation::CreateObjPart(file::NormalizePath(normalizedFilename.c_str()));
-										const std::wstring newObjPath = amalgamation::CreateObjPath(compilandPath, newObjPart);
-										const std::wstring normalizedNewObjPath = file::NormalizePath(newObjPath.c_str());
 
-										AddFileDependency(compilandDb, sourceFilePath, string::ToUtf8String(normalizedNewObjPath), cacheData.lastModificationTime);
-
-										// create a new compiland matching this .obj.
-										// we could use different PDBs for different files when no PCHs are being used, but with
-										// our automatic multi-processor compilation this doesn't really gain anything performance-wise
-										// and just complicates things.
-
-										// adapt command line to accommodate new .obj path
-										const std::wstring& newCommandLine = string::Replace(optionsCache[4], L".obj", newObjPart);
-
-										Compiland* newCompiland = LC_NEW(&g_compilandAllocator, Compiland)
-										{
-											string::ToUtf8String(newObjPath),
-											string::ToUtf8String(normalizedFilename),
-											envPdbPath,
-											envCompilerPath,
-											string::ToUtf8String(newCommandLine),
-											envCompilerWorkingDirectory,
-											objPath,										// .obj of the amalgamation
-											nullptr,										// file indices
-
-											// note that for the purpose of disambiguating symbols in COFF files,
-											// we treat these files as being the amalgamated file.
-											// symbols originally coming from amalgamated files need to have the same
-											// name as symbols from individual files.
-											GetCompilandIdFromPath(normalizedCompilandPath),
-
-											// same for the DIA symbol index, for the same reason
-											static_cast<uint32_t>(i),
-
-											Compiland::Type::PART_OF_AMALGAMATION,			// type of file
-											isPartOfLibrary,								// isPartOfLibrary
-											false											// wasRecompiled
-										};
-
-										compilandDb->compilands.insert(std::make_pair(string::ToUtf8String(normalizedNewObjPath), newCompiland));
-
-										// try updating the amalgamated compiland for the given file and create a new one in case none exists yet
-										{
-											const auto& insertPair = compilandDb->amalgamatedCompilands.emplace(objPath, nullptr);
-											symbols::AmalgamatedCompiland*& amalgamatedCompiland = insertPair.first->second;
-
-											if (insertPair.second)
-											{
-												// insertion was successful, create a new amalgamated compiland
-												amalgamatedCompiland = LC_NEW(&g_amalgamatedCompilandAllocator, symbols::AmalgamatedCompiland);
-												amalgamatedCompiland->isSplit = false;
-											}
-
-											// update entry
-											amalgamatedCompiland->singleParts.push_back(string::ToUtf8String(normalizedNewObjPath));
-										}
-									}
-									else if (generateLogs)
-									{
-										LC_LOG_DEV("Missing dependency %S", normalizedFilename.c_str());
-									}
-								}
-							}
-						}
-						else
-						{
-							// this is a header file. add it as regular dependency for the main amalgamated .obj file
-							if (file::GetDriveType(normalizedFilename.c_str()) == file::DriveType::OPTICAL)
-							{
-								if (generateLogs)
-								{
-									LC_LOG_DEV("Ignoring file %S on optical drive", normalizedFilename.c_str());
-								}
-							}
-							else
-							{
-								const FileAttributeCache::Data& cacheData = fileCache.UpdateCacheData(normalizedFilename);
-								if (cacheData.exists)
-								{
-									if (generateLogs)
-									{
-										LC_LOG_DEV("Dependency %S", normalizedFilename.c_str());
-									}
-
-									AddFileDependency(compilandDb, sourceFilePath, objPath, cacheData.lastModificationTime);
-								}
-								else if (generateLogs)
-								{
-									LC_LOG_DEV("Missing dependency %S", normalizedFilename.c_str());
-								}
-							}
-						}
-					}
-				}
+				amalgamation::SplitAmalgamatedCompilands(splitAmalgamatedFilesThreshold,
+					compilandDb, compiland, includeFilePaths, diaCompilandPath.GetString(), compilandPath, fileCache, generateLogs, provider->lastModificationTime);
 			}
 		}
 
@@ -1582,7 +1285,7 @@ namespace symbols
 			else if (string::Matches(environmentName.GetString(), L"exe"))
 			{
 				// the path to the linker is often not normalized, and contains wrong casing
-				linkerDb->linkerPath = string::ToUtf8String(file::NormalizePath(environmentOption.GetString()));
+				linkerDb->linkerPath = string::ToUtf8String(Filesystem::NormalizePath(environmentOption.GetString()).GetString());
 				++foundOptions;
 			}
 			else if (string::Matches(environmentName.GetString(), L"cmd"))
@@ -1788,9 +1491,7 @@ namespace symbols
 
 							// note that symbols coming from COFFs have already been disambiguated, so we can
 							// directly use their name
-							symbols::Symbol* newSymbol = LC_NEW(&g_symbolAllocator, symbols::Symbol) { symbolName, rva };
-							symbolDb->symbolsByName.emplace(symbolName, newSymbol);
-							symbolDb->symbolsByRva.emplace(rva, newSymbol);
+							symbols::Symbol* newSymbol = CreateNewSymbol(symbolName, rva, symbolDb);
 
 							initializerDb.dynamicInitializers.push_back(newSymbol);
 						}
@@ -1847,10 +1548,23 @@ namespace symbols
 								// differently than the undecorated names for COFF symbols when using nameMangling::UndecorateSymbol
 								// without flags.
 								// however, using the correct (undocumented) flags yields the same name as stored in DIA.
+								size_t crtSectionSymbolCount = 0u;
 								const size_t crtSectionCount = crtSections.size();
 								for (size_t i = 0u; i < crtSectionCount; ++i)
 								{
 									const coff::CrtSection* crtSection = crtSections[i];
+									crtSectionSymbolCount += crtSection->symbols.size();
+
+									// when using Clang, .CRT sections often have no symbols assigned in the COFF symbol table
+									if (crtSection->symbols.size() == 0u)
+									{
+										continue;
+									}
+									if (symbolIndex >= crtSection->symbols.size())
+									{
+										continue;
+									}
+
 									const coff::Symbol* coffSymbol = crtSection->symbols[symbolIndex];
 									const size_t relocationCount = coffSymbol->relocations.size();
 
@@ -1875,9 +1589,7 @@ namespace symbols
 												dstSymbolName.c_str(),
 												dynamicInitializerRva);
 
-											symbols::Symbol* newSymbol = LC_NEW(&g_symbolAllocator, symbols::Symbol) { coffSymbolName, initializerRva };
-											symbolDb->symbolsByName.emplace(coffSymbolName, newSymbol);
-											symbolDb->symbolsByRva.emplace(initializerRva, newSymbol);
+											symbols::Symbol* newSymbol = CreateNewSymbol(coffSymbolName, initializerRva, symbolDb);
 
 											initializerDb.dynamicInitializers.push_back(newSymbol);
 
@@ -1886,7 +1598,12 @@ namespace symbols
 									}
 								}
 
-								LC_ERROR_DEV("Could not find dynamic initializer symbol %S for compiland %s", diaSymbolName.c_str(), compilandName.c_str());
+								// only output an error in case we had a real chance to find the initializer. when code is compiled with Clang,
+								// there will be no symbols in the CRT sections.
+								if (crtSectionSymbolCount != 0u)
+								{
+									LC_ERROR_DEV("Could not find dynamic initializer symbol %S for compiland %s", diaSymbolName.c_str(), compilandName.c_str());
+								}
 
 							symbolFound:
 								dynamicInitializerSymbol->Release();
@@ -1921,7 +1638,7 @@ namespace symbols
 						// we don't have a COFF database for this compiland. the compiland is not part of the module and
 						// must be part of e.g. an external library.
 						// in this case, the name of an initializer's symbol doesn't really matter, as long as it is unique
-						// and the same during a live coding session.
+						// and the same during a live code session.
 						// the reason for that is that these files cannot be changed and recompiled anyway, but will only be used for
 						// linking. therefore, the COFF used for linking is always the same, and we only need to assign unique
 						// names for these initializers.
@@ -1943,9 +1660,7 @@ namespace symbols
 							ImmutableString fullPath(symbolName.c_str());
 							LC_LOG_DEV("Found dynamic initializer %s at 0x%X", fullPath.c_str(), rva);
 
-							symbols::Symbol* newSymbol = LC_NEW(&g_symbolAllocator, symbols::Symbol) { fullPath, rva };
-							symbolDb->symbolsByName.emplace(std::move(fullPath), newSymbol);
-							symbolDb->symbolsByRva.emplace(rva, newSymbol);
+							symbols::Symbol* newSymbol = CreateNewSymbol(fullPath, rva, symbolDb);
 
 							initializerDb.dynamicInitializers.push_back(newSymbol);
 
@@ -2007,12 +1722,6 @@ namespace symbols
 	}
 
 
-	void DestroyUserDefinedTypesDB(UserDefinedTypesDB* db)
-	{
-		delete db;
-	}
-
-
 	void MergeCompilandsAndDependencies(CompilandDB* existingDb, CompilandDB* mergedDb)
 	{
 		// merge compilands
@@ -2303,6 +2012,138 @@ namespace symbols
 	}
 
 
+	uint32_t GetCompilandId(const Compiland* compiland, const wchar_t* const objPath, const types::vector<ModifiedObjFile>& modifiedObjFiles)
+	{
+		// try to find the given .obj path in the array of modified object files to check if there's an original amalgamated object path for it
+		for (size_t i = 0u; i < modifiedObjFiles.size(); ++i)
+		{
+			const ModifiedObjFile& objFile = modifiedObjFiles[i];
+
+			// don't bother checking strings if the amalgamated object path is empty anyway
+			if (!objFile.amalgamatedObjPath.empty())
+			{
+				if (string::Matches(objPath, objFile.objPath.c_str()))
+				{
+					return uniqueId::Generate(Filesystem::NormalizePath(objFile.amalgamatedObjPath.c_str()).GetString());
+				}
+			}
+		}
+
+		if (compiland)
+		{
+			// the compiland already exists
+			return compiland->uniqueId;
+		}
+		else
+		{
+			// BEGIN EPIC MOD - Fixes problems with such things as statics failing to resolve
+			return GetCompilandIdFromPath(objPath);
+			// END EPIC MD
+		}
+	}
+
+	ImmutableString TransformAnonymousNamespacePattern(const ImmutableString& immutableString, uint32_t uniqueId)
+	{
+		// an ANS symbol name is always of the form ?identifier@?A0x12345678, where the hex code following the "@?A0x" part is most likely a hash of the
+		// filename the ANS appears in, generated by the compiler.
+
+		// when splitting amalgamated files, we need to make sure that symbols in anonymous namespaces compiled into
+		// those files are also found when compiled into single-part files.
+		// however, single-part files get assigned a different hash by the compiler, leading to different
+		// symbol names for symbols that reside in anonymous namespaces.
+		// in order to "correct" this, we generate our own hex identifier for ANS symbols, making sure that this
+		// identifier yields the same result for both amalgamated as well as single-part files.
+		// this is done by using the uniqueId as identifier, which is the same for amalgamated files as well
+		// as their split single-file counterparts.
+
+		// the same issue occurs when compiling files with Clang, because relative paths and absolute paths get assigned
+		// different hashes, even though the resulting file is the same.
+		const char* str = immutableString.c_str();
+		const char* anonNamespaceCursor = string::Find(str, symbolPatterns::ANONYMOUS_NAMESPACE_PATTERN);
+		if (!anonNamespaceCursor)
+		{
+			return immutableString;
+		}
+
+		// convert the unique ID to a hex string with 8 characters
+		const char hexUniqueId[8u] =
+		{
+			"0123456789ABCDEF"[(uniqueId >> 28u) & 15u],
+			"0123456789ABCDEF"[(uniqueId >> 24u) & 15u],
+			"0123456789ABCDEF"[(uniqueId >> 20u) & 15u],
+			"0123456789ABCDEF"[(uniqueId >> 16u) & 15u],
+			"0123456789ABCDEF"[(uniqueId >> 12u) & 15u],
+			"0123456789ABCDEF"[(uniqueId >> 8u) & 15u],
+			"0123456789ABCDEF"[(uniqueId >> 4u) & 15u],
+			"0123456789ABCDEF"[(uniqueId >> 0u) & 15u]
+		};
+
+		const size_t patternLength = strlen(symbolPatterns::ANONYMOUS_NAMESPACE_PATTERN);
+
+		// Clang does not add leading zeros to the hash, so in theory, the hash after ?A0x can be of any length.
+		// we therefore cannot simply replace the hash in the original string, but must build a new one from individual pieces.
+		std::string newStr;
+		newStr.reserve(immutableString.GetLength());
+
+		const char* lastFound = str;
+		do
+		{
+			// skip anonymous namespace pattern
+			anonNamespaceCursor += patternLength;
+
+			// copy until current cursor
+			newStr.append(lastFound, static_cast<size_t>(anonNamespaceCursor - lastFound));
+
+			// append hex ID
+			newStr.append(hexUniqueId, 8u);
+
+			// skip until the end of the anonymous namespace pattern
+			anonNamespaceCursor = string::Find(anonNamespaceCursor, "@");
+			lastFound = anonNamespaceCursor;
+
+			// the identifier could contain several more anonymous namespace patterns
+			anonNamespaceCursor = string::Find(anonNamespaceCursor, symbolPatterns::ANONYMOUS_NAMESPACE_PATTERN);
+		}
+		while (anonNamespaceCursor);
+
+		newStr.append(lastFound, static_cast<size_t>(str + immutableString.GetLength() - lastFound));
+
+		return ImmutableString(newStr.c_str());
+	}
+
+
+	void TransformAnonymousNamespaceSymbols(SymbolDB* symbolDb, ContributionDB* contributionDb, CompilandDB* compilandDb, const types::vector<ModifiedObjFile>& modifiedObjFiles)
+	{
+		const size_t count = symbolDb->ansSymbols.size();
+		for (size_t i = 0u; i < count; ++i)
+		{
+			symbols::Symbol* symbol = symbolDb->ansSymbols[i];
+
+			// generate a new name that is based on the (original amalgamated) filename or the compiland ID
+			const symbols::Contribution* contribution = symbols::FindContributionByRVA(contributionDb, symbol->rva);
+			if (contribution)
+			{
+				const ImmutableString compilandName = symbols::GetContributionCompilandName(contributionDb, contribution);
+				const symbols::Compiland* compiland = symbols::FindCompiland(compilandDb, compilandName);
+				if (compiland)
+				{
+					// remove the symbol with the old name from the database
+					symbolDb->symbolsByName.erase(symbol->name);
+
+					const uint32_t uniqueId = GetCompilandId(compiland, string::ToWideString(compilandName).c_str(), modifiedObjFiles);
+					symbol->name = TransformAnonymousNamespacePattern(symbol->name, uniqueId);
+
+					// store the symbol with the new name in the database again
+					symbolDb->symbolsByName.emplace(symbol->name, symbol);
+				}
+			}
+		}
+
+		// all symbols have been transformed, no need to touch them again next time
+		symbolDb->ansSymbols.clear();
+	}
+
+
 	template <typename T, size_t N>
 	static inline bool ContainsPatterns(const char* name, const T (&patterns)[N])
 	{
@@ -2369,12 +2210,6 @@ namespace symbols
 	}
 
 
-	bool IsWeakSymbol(const ImmutableString& symbolName)
-	{
-		return StartsWithPatterns(symbolName.c_str(), symbolPatterns::WEAK_SYMBOL_PATTERNS);
-	}
-
-
 	bool IsStringLiteral(const ImmutableString& symbolName)
 	{
 		return ContainsPatterns(symbolName.c_str(), symbolPatterns::STRING_LITERAL_PATTERNS);
@@ -2458,6 +2293,14 @@ namespace symbols
 	{
 		return ContainsPatterns(symbolName.c_str(), symbolPatterns::TLS_STATICS_PATTERNS);
 	}
+
+
+	bool IsSectionSymbol(const ImmutableString& symbolName)
+	{
+		return symbolName.c_str()[0] == '.';
+	}
 }
 
+// BEGIN EPIC MOD
 #include "Windows/HideWindowsPlatformAtomics.h"
+// END EPIC MOD
\ No newline at end of file
diff --git a/Engine/Source/Developer/Windows/LiveCodingServer/Private/External/LC_Symbols.h b/Engine/Source/Developer/Windows/LiveCodingServer/Private/External/LC_Symbols.h
--- a/Engine/Source/Developer/Windows/LiveCodingServer/Private/External/LC_Symbols.h
+++ b/Engine/Source/Developer/Windows/LiveCodingServer/Private/External/LC_Symbols.h
@@ -1,15 +1,19 @@
-// Copyright 2011-2019 Molecular Matters GmbH, all rights reserved.
+// Copyright 2011-2020 Molecular Matters GmbH, all rights reserved.
 
 #pragma once
 
+// BEGIN EPIC MOD
 #include "CoreTypes.h"
+// END EPIC MOD
 #include "LC_ImmutableString.h"
-#include "LC_Process.h"
+#include "LC_ProcessTypes.h"
 #include "LC_Types.h"
 #include "LC_Executable.h"
 #include "LC_ChangeNotification.h"
 #include "LC_DirectoryCache.h"
+// BEGIN EPIC MOD
 #include <dia2.h>
+// END EPIC MOD
 
 template <typename T> class CoffCache;
 
@@ -41,6 +45,9 @@ namespace symbols
 
 		// public function symbols that can be patched
 		types::vector<Symbol*> patchableFunctionSymbols;
+
+		// symbols that contain an anonymous namespace in their mangled name
+		types::vector<Symbol*> ansSymbols;
 	};
 
 	struct Contribution
@@ -102,7 +109,6 @@ namespace symbols
 		ImmutableString amalgamationPath;		// full path to the amalgamation .obj in case this is part of an amalgamation
 		CompilandSourceFiles* sourceFiles;
 		uint32_t uniqueId;
-		uint32_t diaSymbolIndex;				// the index to the DIA symbol from which this compiland originated
 		Type::Enum type;
 		bool isPartOfLibrary;
 		bool wasRecompiled;
@@ -117,6 +123,13 @@ namespace symbols
 		bool isSplit;
 	};
 
+	struct ModifiedObjFile
+	{
+		std::wstring objPath;
+		std::wstring amalgamatedObjPath;	// optional
+	};
+
+
 	// full path to .obj file
 	typedef ImmutableString ObjPath;
 	
@@ -130,6 +143,7 @@ namespace symbols
 		DirectoryCache::Directory* parentDirectory;
 		TimeStamp lastModification;
 		types::vector<ObjPath> objPaths;
+		bool hadInitialChange;
 	};
 
 	struct CompilandDB
@@ -143,7 +157,9 @@ namespace symbols
 		// stores an array of .obj files that need to be recompiled when a certain file changes
 		types::StringMap<Dependency*> dependencies;
 
-		// stores the real name of the .obj on disk for each original DIA compiland name
+		// stores the real name of the .obj on disk for each original DIA compiland name.
+		// this lookup table is needed because DIA compilands coming from static libraries often don't contain the full path,
+		// but a path that is relative to the project instead, e.g. obj\x64\Debug\SimpleFunction.obj
 		types::StringMap<ImmutableString> compilandNameToObjOnDisk;
 	};
 
@@ -228,6 +244,11 @@ namespace symbols
 		};
 	};
 
+
+	// TODO: helper function to make fixing and refactoring easier for now
+	Symbol* CreateNewSymbol(const ImmutableString& name, uint32_t rva, SymbolDB* db);
+
+
 	Provider* OpenEXE(const wchar_t* filename, uint32_t openOptions);
 	void Close(Provider* provider);
 
@@ -253,9 +274,6 @@ namespace symbols
 	// CACHE: temporary, throw away after use
 	DiaCompilandDB* GatherDiaCompilands(Provider* provider);
 
-	// gathers user-defined types for a given compiland
-	UserDefinedTypesDB* GatherUserDefinedTypes(const DiaCompilandDB* diaCompilandDb, const Compiland* compiland);
-
 
 	struct CompilandOptions
 	{
@@ -299,8 +317,6 @@ namespace symbols
 
 	void DestroyCompilandDB(CompilandDB* db);
 
-	void DestroyUserDefinedTypesDB(UserDefinedTypesDB* db);
-
 	void MergeCompilandsAndDependencies(CompilandDB* existingDb, CompilandDB* mergedDb);
 
 
@@ -355,6 +371,18 @@ namespace symbols
 
 	const ImmutableString& GetImageSectionName(const ImageSectionDB* db, const ImageSection* imageSection);
 
+
+
+	// helper function that returns or generates the unique ID of an optional compiland.
+	// for files split off from amalgamated files, we need to use the original object path of the amalgamated file here,
+	// otherwise names of symbols would differ, leading to constructors of global instances being called again.
+	uint32_t GetCompilandId(const Compiland* compiland, const wchar_t* const objPath, const types::vector<ModifiedObjFile>& modifiedObjFiles);
+
+	ImmutableString TransformAnonymousNamespacePattern(const ImmutableString& immutableString, uint32_t uniqueId);
+
+	void TransformAnonymousNamespaceSymbols(SymbolDB* symbolDb, ContributionDB* contributionDb, CompilandDB* compilandDb, const types::vector<ModifiedObjFile>& modifiedObjFiles);
+
+
 	bool IsPchSymbol(const ImmutableString& symbolName);
 
 	bool IsVTable(const ImmutableString& symbolName);
@@ -364,8 +392,6 @@ namespace symbols
 	bool IsDynamicAtexitDestructor(const ImmutableString& functionSymbolName);
 	bool IsPointerToDynamicInitializer(const ImmutableString& symbolName);
 
-	bool IsWeakSymbol(const ImmutableString& symbolName);
-
 	bool IsStringLiteral(const ImmutableString& symbolName);
 	bool IsLineNumber(const ImmutableString& symbolName);
 
@@ -384,4 +410,6 @@ namespace symbols
 	bool IsTlsIndexRelatedSymbol(const ImmutableString& symbolName);
 	bool IsTlsInitRelatedSymbol(const ImmutableString& symbolName);
 	bool IsTlsStaticsRelatedSymbol(const ImmutableString& symbolName);
+
+	bool IsSectionSymbol(const ImmutableString& symbolName);
 }
diff --git a/Engine/Source/Developer/Windows/LiveCodingServer/Private/External/LC_Telemetry.h b/Engine/Source/Developer/Windows/LiveCodingServer/Private/External/LC_Telemetry.h
--- a/Engine/Source/Developer/Windows/LiveCodingServer/Private/External/LC_Telemetry.h
+++ b/Engine/Source/Developer/Windows/LiveCodingServer/Private/External/LC_Telemetry.h
@@ -1,10 +1,14 @@
-// Copyright 2011-2019 Molecular Matters GmbH, all rights reserved.
+// Copyright 2011-2020 Molecular Matters GmbH, all rights reserved.
 
 #pragma once
 
+// BEGIN EPIC MOD
 #include "CoreTypes.h"
-#include <chrono>
+// END EPIC MOD
 #include "LC_CriticalSection.h"
+// BEGIN EPIC MOD
+#include <chrono>
+// END EPIC MOD
 
 namespace telemetry
 {
diff --git a/Engine/Source/Developer/Windows/LiveCodingServer/Private/External/LC_TimeDate.cpp b/Engine/Source/Developer/Windows/LiveCodingServer/Private/External/LC_TimeDate.cpp
--- a/Engine/Source/Developer/Windows/LiveCodingServer/Private/External/LC_TimeDate.cpp
+++ b/Engine/Source/Developer/Windows/LiveCodingServer/Private/External/LC_TimeDate.cpp
@@ -1,7 +1,12 @@
-// Copyright 2011-2019 Molecular Matters GmbH, all rights reserved.
+// Copyright 2011-2020 Molecular Matters GmbH, all rights reserved.
 
+// BEGIN EPIC MOD
+//#include PCH_INCLUDE
+// END EPIC MOD
 #include "LC_TimeDate.h"
+// BEGIN EPIC MOD
 #include "Windows/WindowsHWrapper.h"
+// END EPIC MOD
 
 TimeDate::TimeDate(unsigned short year, unsigned char month, unsigned char day, unsigned char hour, unsigned char minute, unsigned char second, unsigned short milliSecond)
 	: m_year(year)
diff --git a/Engine/Source/Developer/Windows/LiveCodingServer/Private/External/LC_Types.h b/Engine/Source/Developer/Windows/LiveCodingServer/Private/External/LC_Types.h
--- a/Engine/Source/Developer/Windows/LiveCodingServer/Private/External/LC_Types.h
+++ b/Engine/Source/Developer/Windows/LiveCodingServer/Private/External/LC_Types.h
@@ -1,11 +1,15 @@
-// Copyright 2011-2019 Molecular Matters GmbH, all rights reserved.
+// Copyright 2011-2020 Molecular Matters GmbH, all rights reserved.
 
 #pragma once
 
+// BEGIN EPIC MOD
+#include "CoreTypes.h"
+// END EPIC MOD
 #include "LC_ImmutableString.h"
+// BEGIN EPIC MOD
 #include <unordered_set>
 #include <unordered_map>
-
+// END EPIC MOD
 
 namespace types
 {
diff --git a/Engine/Source/Developer/Windows/LiveCodingServer/Private/External/LC_UniqueId.cpp b/Engine/Source/Developer/Windows/LiveCodingServer/Private/External/LC_UniqueId.cpp
--- a/Engine/Source/Developer/Windows/LiveCodingServer/Private/External/LC_UniqueId.cpp
+++ b/Engine/Source/Developer/Windows/LiveCodingServer/Private/External/LC_UniqueId.cpp
@@ -1,10 +1,14 @@
-// Copyright 2011-2019 Molecular Matters GmbH, all rights reserved.
+// Copyright 2011-2020 Molecular Matters GmbH, all rights reserved.
 
+// BEGIN EPIC MOD
+//#include PCH_INCLUDE
+// END EPIC MOD
 #include "LC_UniqueId.h"
+#include "LC_Hashing.h"
+// BEGIN EPIC MOD
 #include "LC_CriticalSection.h"
 #include "LC_Types.h"
-#include "xxhash.h"
-
+// END EPIC MOD
 
 namespace detail
 {
@@ -14,7 +18,7 @@ namespace detail
 		{
 			inline size_t operator()(const std::wstring& key) const
 			{
-				return XXH32(key.c_str(), key.length() * sizeof(wchar_t), 0u);
+				return Hashing::Hash32(key.c_str(), key.length() * sizeof(wchar_t), 0u);
 			}
 		};
 
diff --git a/Engine/Source/Developer/Windows/LiveCodingServer/Private/External/LC_VirtualMemoryRange.cpp b/Engine/Source/Developer/Windows/LiveCodingServer/Private/External/LC_VirtualMemoryRange.cpp
new file mode 100644
--- /dev/null
+++ b/Engine/Source/Developer/Windows/LiveCodingServer/Private/External/LC_VirtualMemoryRange.cpp
@@ -0,0 +1,71 @@
+// Copyright 2011-2020 Molecular Matters GmbH, all rights reserved.
+
+// BEGIN EPIC MOD
+//#include PCH_INCLUDE
+// END EPIC MOD
+#include "LC_VirtualMemoryRange.h"
+#include "LC_PointerUtil.h"
+// BEGIN EPIC MOD
+#include "LC_Logging.h"
+#include <cinttypes>
+// END EPIC MOD
+
+VirtualMemoryRange::VirtualMemoryRange(Process::Handle processHandle, const void* addressStart, const void* addressEnd, size_t alignment)
+	: m_processHandle(processHandle)
+	, m_addressStart(addressStart)
+	, m_addressEnd(addressEnd)
+	, m_alignment(alignment)
+{
+	m_pageData.reserve(32u);
+}
+
+
+void VirtualMemoryRange::ReservePages(void)
+{
+	// reserve all free pages in the virtual memory range.
+	// pages must be aligned to the given alignment.
+	for (const void* address = m_addressStart; address < m_addressEnd; /* nothing */)
+	{
+		// align address to be scanned
+		address = pointer::AlignTop<const void*>(address, m_alignment);
+
+		::MEMORY_BASIC_INFORMATION memoryInfo = {};
+		const size_t bytesReturned = ::VirtualQueryEx(+m_processHandle, address, &memoryInfo, sizeof(::MEMORY_BASIC_INFORMATION));
+
+		// we are only interested in free pages
+		if ((bytesReturned > 0u) && (memoryInfo.State == MEM_FREE))
+		{
+			// work out the maximum size of the page allocation.
+			// we should not allocate past the end of the range.
+			const size_t bytesLeft = pointer::Displacement<size_t>(memoryInfo.BaseAddress, m_addressEnd);
+			const size_t size = std::min<size_t>(memoryInfo.RegionSize, bytesLeft);
+
+			// try to reserve this page.
+			// if we are really unlucky, the process might have allocated this region in the meantime.
+			void* baseAddress = ::VirtualAllocEx(+m_processHandle, memoryInfo.BaseAddress, size, MEM_RESERVE, PAGE_NOACCESS);
+			if (baseAddress)
+			{
+				LC_LOG_DEV("Found virtual memory region at 0x%p with size 0x%" PRIX64, baseAddress, size);
+				m_pageData.emplace_back(PageData { baseAddress });
+			}
+		}
+
+		// keep on searching
+		address = pointer::Offset<const void*>(memoryInfo.BaseAddress, memoryInfo.RegionSize);
+	}
+}
+
+
+void VirtualMemoryRange::FreeReservedPages(void)
+{
+	for (const auto& it : m_pageData)
+	{
+		const bool success = ::VirtualFreeEx(+m_processHandle, it.address, 0u, MEM_RELEASE);
+		if (!success)
+		{
+			LC_WARNING_USER("Cannot free virtual memory region at 0x%p", it.address);
+		}
+	}
+
+	m_pageData.clear();
+}
diff --git a/Engine/Source/Developer/Windows/LiveCodingServer/Private/External/LC_VirtualMemoryRange.h b/Engine/Source/Developer/Windows/LiveCodingServer/Private/External/LC_VirtualMemoryRange.h
new file mode 100644
--- /dev/null
+++ b/Engine/Source/Developer/Windows/LiveCodingServer/Private/External/LC_VirtualMemoryRange.h
@@ -0,0 +1,34 @@
+// Copyright 2011-2020 Molecular Matters GmbH, all rights reserved.
+
+#pragma once
+
+// BEGIN EPIC MOD
+#include "CoreTypes.h"
+// END EPIC MOD
+#include "LC_ProcessTypes.h"
+// BEGIN EPIC MOD
+#include "LC_Types.h"
+// END EPIC MOD
+
+
+class VirtualMemoryRange
+{
+public:
+	VirtualMemoryRange(Process::Handle processHandle, const void* addressStart, const void* addressEnd, size_t alignment);
+
+	void ReservePages(void);
+	void FreeReservedPages(void);
+
+private:
+	struct PageData
+	{
+		void* address;
+	};
+
+	Process::Handle m_processHandle;
+	const void* m_addressStart;
+	const void* m_addressEnd;
+	size_t m_alignment;
+
+	types::vector<PageData> m_pageData;
+};
diff --git a/Engine/Source/Developer/Windows/LiveCodingServer/Private/External/LC_VisualStudioAutomation.cpp b/Engine/Source/Developer/Windows/LiveCodingServer/Private/External/LC_VisualStudioAutomation.cpp
--- a/Engine/Source/Developer/Windows/LiveCodingServer/Private/External/LC_VisualStudioAutomation.cpp
+++ b/Engine/Source/Developer/Windows/LiveCodingServer/Private/External/LC_VisualStudioAutomation.cpp
@@ -1,28 +1,61 @@
-// Copyright 2011-2019 Molecular Matters GmbH, all rights reserved.
+// Copyright 2011-2020 Molecular Matters GmbH, all rights reserved.
 
+// BEGIN EPIC MOD
+//#include PCH_INCLUDE
+// END EPIC MOD
 #include "LC_VisualStudioAutomation.h"
 #include "LC_COMThread.h"
 #include "LC_StringUtil.h"
+#include "LC_Thread.h"
+// BEGIN EPIC MOD
 #include "LC_Logging.h"
 #include "Windows/COMPointer.h"
+// END EPIC MOD
 
+// BEGIN EPIC MOD
 #if WITH_VISUALSTUDIO_DTE
+// END EPIC MOD
 
 namespace
 {
 	static COMThread* g_comThread = nullptr;
 
 
+	// helper function that checks whether the debugger is currently held for debugging, e.g. at a breakpoint, an exception
+	static bool IsHeldForDebugging(const EnvDTE::DebuggerPtr& debugger)
+	{
+		EnvDTE::dbgDebugMode mode = EnvDTE::dbgDesignMode;
+		debugger->get_CurrentMode(&mode);
+
+		return (mode == EnvDTE::dbgBreakMode);
+	}
+
+
 	// helper function that waits until the debugger finished executing a command and is in break mode again
-	static void WaitUntilBreakMode(const EnvDTE::DebuggerPtr& debugger)
+	static bool WaitUntilBreakMode(const EnvDTE::DebuggerPtr& debugger)
 	{
+		const unsigned int RETRY_WAIT_TIME = 5u;
+		const unsigned int RETRY_TIMEOUT = 500u;
+
 		EnvDTE::dbgDebugMode mode = EnvDTE::dbgDesignMode;
+		debugger->get_CurrentMode(&mode);
 
+		unsigned int msWaited = 0u;
 		while (mode != EnvDTE::dbgBreakMode)
 		{
+			Thread::Current::SleepMilliSeconds(RETRY_WAIT_TIME);
+			msWaited += RETRY_WAIT_TIME;
+
 			debugger->get_CurrentMode(&mode);
-			thread::Sleep(5u);
+
+			if (msWaited >= RETRY_TIMEOUT)
+			{
+				// enough time has passed, bail out
+				return false;
+			}
 		}
+
+		return true;
 	}
 }
 
@@ -41,10 +74,12 @@ void visualStudio::Shutdown(void)
 }
 
 
-static EnvDTE::DebuggerPtr FindDebuggerAttachedToProcess(unsigned int processId)
+static EnvDTE::DebuggerPtr FindDebuggerAttachedToProcess(Process::Id processId)
 {
 	// get all objects from the running object table (ROT)
+	// BEGIN EPIC MOD
 	TComPtr<IRunningObjectTable> rot = nullptr;
+	// END EPIC MOD
 	HRESULT result = ::GetRunningObjectTable(0u, &rot);
 	if (FAILED(result) || (!rot))
 	{
@@ -52,7 +87,9 @@ static EnvDTE::DebuggerPtr FindDebuggerAttachedToProcess(unsigned int processId)
 		return nullptr;
 	}
 
+	// BEGIN EPIC MOD
 	TComPtr<IEnumMoniker> enumMoniker = nullptr;
+	// END EPIC MOD
 	result = rot->EnumRunning(&enumMoniker);
 	if (FAILED(result) || (!enumMoniker))
 	{
@@ -63,12 +100,16 @@ static EnvDTE::DebuggerPtr FindDebuggerAttachedToProcess(unsigned int processId)
 	enumMoniker->Reset();
 	do
 	{
+		// BEGIN EPIC MOD
 		TComPtr<IMoniker> next = nullptr;
+		// END EPIC MOD
 		result = enumMoniker->Next(1u, &next, NULL);
 
 		if (SUCCEEDED(result) && next)
 		{
+			// BEGIN EPIC MOD
 			TComPtr<IBindCtx> context = nullptr;
+			// END EPIC MOD
 			result = ::CreateBindCtx(0, &context);
 
 			if (FAILED(result) || (!context))
@@ -89,15 +130,19 @@ static EnvDTE::DebuggerPtr FindDebuggerAttachedToProcess(unsigned int processId)
 			if (string::Contains(displayName, L"VisualStudio.DTE."))
 			{
 				// free display name using COM allocator
+				// BEGIN EPIC MOD
 				TComPtr<IMalloc> comMalloc = nullptr;
+				// END EPIC MOD
 				result = ::CoGetMalloc(1u, &comMalloc);
 				if (SUCCEEDED(result) && comMalloc)
 				{
 					comMalloc->Free(displayName);
 				}
 
+				// BEGIN EPIC MOD
 				TComPtr<IUnknown> unknown = nullptr;
 				result = rot->GetObject(next, &unknown);
+				// END EPIC MOD
 				if (FAILED(result) || (!unknown))
 				{
 					LC_ERROR_DEV("Could not retrieve COM object from running object table. Error: 0x%X", result);
@@ -160,7 +205,7 @@ static EnvDTE::DebuggerPtr FindDebuggerAttachedToProcess(unsigned int processId)
 					}
 
 					// we got a valid processId
-					if (static_cast<unsigned int>(debuggerProcessId) == processId)
+					if (static_cast<unsigned int>(debuggerProcessId) == +processId)
 					{
 						// found debugger attached to our process
 						return debugger;
@@ -175,16 +220,18 @@ static EnvDTE::DebuggerPtr FindDebuggerAttachedToProcess(unsigned int processId)
 }
 
 
-EnvDTE::DebuggerPtr visualStudio::FindDebuggerAttachedToProcess(unsigned int processId)
+EnvDTE::DebuggerPtr visualStudio::FindDebuggerAttachedToProcess(Process::Id processId)
 {
 	return g_comThread->CallInThread(&::FindDebuggerAttachedToProcess, processId);
 }
 
 
-static EnvDTE::DebuggerPtr FindDebuggerForProcess(unsigned int processId)
+static EnvDTE::DebuggerPtr FindDebuggerForProcess(Process::Id processId)
 {
 	// get all objects from the running object table (ROT)
+	// BEGIN EPIC MOD
 	TComPtr<IRunningObjectTable> rot = nullptr;
+	// END EPIC MOD
 	HRESULT result = ::GetRunningObjectTable(0u, &rot);
 	if (FAILED(result) || (!rot))
 	{
@@ -192,7 +239,9 @@ static EnvDTE::DebuggerPtr FindDebuggerForProcess(unsigned int processId)
 		return nullptr;
 	}
 
+	// BEGIN EPIC MOD
 	TComPtr<IEnumMoniker> enumMoniker = nullptr;
+	// END EPIC MOD
 	result = rot->EnumRunning(&enumMoniker);
 	if (FAILED(result) || (!enumMoniker))
 	{
@@ -203,12 +252,16 @@ static EnvDTE::DebuggerPtr FindDebuggerForProcess(unsigned int processId)
 	enumMoniker->Reset();
 	do
 	{
+		// BEGIN EPIC MOD
 		TComPtr<IMoniker> next = nullptr;
+		// END EPIC MOD
 		result = enumMoniker->Next(1u, &next, NULL);
 
 		if (SUCCEEDED(result) && next)
 		{
+			// BEGIN EPIC MOD
 			TComPtr<IBindCtx> context = nullptr;
+			// END EPIC MOD
 			result = ::CreateBindCtx(0, &context);
 
 			if (FAILED(result) || (!context))
@@ -229,15 +282,19 @@ static EnvDTE::DebuggerPtr FindDebuggerForProcess(unsigned int processId)
 			if (string::Contains(displayName, L"VisualStudio.DTE."))
 			{
 				// free display name using COM allocator
+				// BEGIN EPIC MOD
 				TComPtr<IMalloc> comMalloc = nullptr;
+				// END EPIC MOD
 				result = ::CoGetMalloc(1u, &comMalloc);
 				if (SUCCEEDED(result) && comMalloc)
 				{
 					comMalloc->Free(displayName);
 				}
 
+				// BEGIN EPIC MOD
 				TComPtr<IUnknown> unknown = nullptr;
 				result = rot->GetObject(next, &unknown);
+				// END EPIC MOD
 				if (FAILED(result) || (!unknown))
 				{
 					LC_ERROR_DEV("Could not retrieve COM object from running object table. Error: 0x%X", result);
@@ -280,7 +337,7 @@ static EnvDTE::DebuggerPtr FindDebuggerForProcess(unsigned int processId)
 				}
 
 				// we got a valid processId
-				if (static_cast<unsigned int>(debuggerProcessId) == processId)
+				if (static_cast<unsigned int>(debuggerProcessId) == +processId)
 				{
 					// found debugger debugging our process
 					return debugger;
@@ -294,13 +351,13 @@ static EnvDTE::DebuggerPtr FindDebuggerForProcess(unsigned int processId)
 }
 
 
-EnvDTE::DebuggerPtr visualStudio::FindDebuggerForProcess(unsigned int processId)
+EnvDTE::DebuggerPtr visualStudio::FindDebuggerForProcess(Process::Id processId)
 {
 	return g_comThread->CallInThread(&::FindDebuggerForProcess, processId);
 }
 
 
-static bool AttachToProcess(const EnvDTE::DebuggerPtr& debugger, unsigned int processId)
+static bool AttachToProcess(const EnvDTE::DebuggerPtr& debugger, Process::Id processId)
 {
 	// fetch all local processes running on this machine
 	EnvDTE::ProcessesPtr allProcesses = nullptr;
@@ -340,7 +397,7 @@ static bool AttachToProcess(const EnvDTE::DebuggerPtr& debugger, unsigned int pr
 		}
 
 		// we got a valid processId
-		if (static_cast<unsigned int>(localProcessId) == processId)
+		if (static_cast<unsigned int>(localProcessId) == +processId)
 		{
 			// this is the process we want to attach to
 			result = singleProcess->Attach();
@@ -358,7 +415,7 @@ static bool AttachToProcess(const EnvDTE::DebuggerPtr& debugger, unsigned int pr
 }
 
 
-bool visualStudio::AttachToProcess(const EnvDTE::DebuggerPtr& debugger, unsigned int processId)
+bool visualStudio::AttachToProcess(const EnvDTE::DebuggerPtr& debugger, Process::Id processId)
 {
 	return g_comThread->CallInThread(&::AttachToProcess, debugger, processId);
 }
@@ -419,6 +476,12 @@ types::vector<EnvDTE::ThreadPtr> visualStudio::EnumerateThreads(const EnvDTE::De
 
 static bool FreezeThreads(const EnvDTE::DebuggerPtr& debugger, const types::vector<EnvDTE::ThreadPtr>& threads)
 {
+	if (!IsHeldForDebugging(debugger))
+	{
+		LC_WARNING_USER("FreezeThreads: Debugger is no longer held for debugging");
+		return false;
+	}
+
 	bool success = true;
 
 	const size_t count = threads.size();
@@ -426,7 +489,17 @@ static bool FreezeThreads(const EnvDTE::DebuggerPtr& debugger, const types::vect
 	{
 		EnvDTE::ThreadPtr singleThread = threads[i];
 		const HRESULT result = singleThread->Freeze();
-		WaitUntilBreakMode(debugger);
+		if (FAILED(result))
+		{
+			LC_ERROR_USER("FreezeThreads: Failed to freeze thread");
+			return false;
+		}
+
+		if (!WaitUntilBreakMode(debugger))
+		{
+			LC_WARNING_USER("FreezeThreads: Cannot wait for debugger break mode");
+			return false;
+		}
 
 		success &= SUCCEEDED(result);
 	}
@@ -441,8 +514,14 @@ bool visualStudio::FreezeThreads(const EnvDTE::DebuggerPtr& debugger, const type
 }
 
 
-static bool FreezeThread(const EnvDTE::DebuggerPtr& debugger, const types::vector<EnvDTE::ThreadPtr>& threads, unsigned int threadId)
+static bool FreezeThread(const EnvDTE::DebuggerPtr& debugger, const types::vector<EnvDTE::ThreadPtr>& threads, Thread::Id threadId)
 {
+	if (!IsHeldForDebugging(debugger))
+	{
+		LC_WARNING_USER("FreezeThread: Debugger is no longer held for debugging");
+		return false;
+	}
+
 	const size_t count = threads.size();
 	for (size_t i = 0u; i < count; ++i)
 	{
@@ -455,11 +534,21 @@ static bool FreezeThread(const EnvDTE::DebuggerPtr& debugger, const types::vecto
 			continue;
 		}
 
-		if (static_cast<unsigned int>(id) == threadId)
+		if (static_cast<unsigned int>(id) == +threadId)
 		{
 			// found the thread we're looking for
 			result = singleThread->Freeze();
-			WaitUntilBreakMode(debugger);
+			if (FAILED(result))
+			{
+				LC_ERROR_USER("FreezeThread: Failed to freeze thread");
+				return false;
+			}
+
+			if (!WaitUntilBreakMode(debugger))
+			{
+				LC_WARNING_USER("FreezeThread: Cannot wait for debugger break mode");
+				return false;
+			}
 
 			return SUCCEEDED(result);
 		}
@@ -469,7 +558,7 @@ static bool FreezeThread(const EnvDTE::DebuggerPtr& debugger, const types::vecto
 }
 
 
-bool visualStudio::FreezeThread(const EnvDTE::DebuggerPtr& debugger, const types::vector<EnvDTE::ThreadPtr>& threads, unsigned int threadId)
+bool visualStudio::FreezeThread(const EnvDTE::DebuggerPtr& debugger, const types::vector<EnvDTE::ThreadPtr>& threads, Thread::Id threadId)
 {
 	return g_comThread->CallInThread(&::FreezeThread, debugger, threads, threadId);
 }
@@ -477,6 +566,12 @@ bool visualStudio::FreezeThread(const EnvDTE::DebuggerPtr& debugger, const types
 
 static bool ThawThreads(const EnvDTE::DebuggerPtr& debugger, const types::vector<EnvDTE::ThreadPtr>& threads)
 {
+	if (!IsHeldForDebugging(debugger))
+	{
+		LC_WARNING_USER("ThawThreads: Debugger is no longer held for debugging");
+		return false;
+	}
+
 	bool success = true;
 
 	const size_t count = threads.size();
@@ -484,7 +579,17 @@ static bool ThawThreads(const EnvDTE::DebuggerPtr& debugger, const types::vector
 	{
 		EnvDTE::ThreadPtr singleThread = threads[i];
 		const HRESULT result = singleThread->Thaw();
-		WaitUntilBreakMode(debugger);
+		if (FAILED(result))
+		{
+			LC_ERROR_USER("ThawThreads: Failed to thaw thread");
+			return false;
+		}
+
+		if (!WaitUntilBreakMode(debugger))
+		{
+			LC_WARNING_USER("ThawThreads: Cannot wait for debugger break mode");
+			return false;
+		}
 
 		success &= SUCCEEDED(result);
 	}
@@ -499,8 +604,14 @@ bool visualStudio::ThawThreads(const EnvDTE::DebuggerPtr& debugger, const types:
 }
 
 
-static bool ThawThread(const EnvDTE::DebuggerPtr& debugger, const types::vector<EnvDTE::ThreadPtr>& threads, unsigned int threadId)
+static bool ThawThread(const EnvDTE::DebuggerPtr& debugger, const types::vector<EnvDTE::ThreadPtr>& threads, Thread::Id threadId)
 {
+	if (!IsHeldForDebugging(debugger))
+	{
+		LC_WARNING_USER("ThawThread: Debugger is no longer held for debugging");
+		return false;
+	}
+
 	const size_t count = threads.size();
 	for (size_t i = 0u; i < count; ++i)
 	{
@@ -513,11 +624,21 @@ static bool ThawThread(const EnvDTE::DebuggerPtr& debugger, const types::vector<
 			continue;
 		}
 
-		if (static_cast<unsigned int>(id) == threadId)
+		if (static_cast<unsigned int>(id) == +threadId)
 		{
 			// found the thread we're looking for
 			result = singleThread->Thaw();
-			WaitUntilBreakMode(debugger);
+			if (FAILED(result))
+			{
+				LC_ERROR_USER("ThawThread: Failed to thaw thread");
+				return false;
+			}
+
+			if (!WaitUntilBreakMode(debugger))
+			{
+				LC_WARNING_USER("ThawThread: Cannot wait for debugger break mode");
+				return false;
+			}
 
 			return SUCCEEDED(result);
 		}
@@ -527,7 +648,7 @@ static bool ThawThread(const EnvDTE::DebuggerPtr& debugger, const types::vector<
 }
 
 
-bool visualStudio::ThawThread(const EnvDTE::DebuggerPtr& debugger, const types::vector<EnvDTE::ThreadPtr>& threads, unsigned int threadId)
+bool visualStudio::ThawThread(const EnvDTE::DebuggerPtr& debugger, const types::vector<EnvDTE::ThreadPtr>& threads, Thread::Id threadId)
 {
 	return g_comThread->CallInThread(&::ThawThread, debugger, threads, threadId);
 }
@@ -535,7 +656,9 @@ bool visualStudio::ThawThread(const EnvDTE::DebuggerPtr& debugger, const types::
 
 static bool Resume(const EnvDTE::DebuggerPtr& debugger)
 {
+	// BEGIN EPIC MOD
 	const HRESULT result = debugger->Go(variant_t(Windows::FALSE));
+	// END EPIC MOD
 	return SUCCEEDED(result);
 }
 
@@ -549,7 +672,9 @@ bool visualStudio::Resume(const EnvDTE::DebuggerPtr& debugger)
 static bool Break(const EnvDTE::DebuggerPtr& debugger)
 {
 	// wait until the debugger really enters break-mode
+	// BEGIN EPIC MOD
 	const HRESULT result = debugger->Break(variant_t(Windows::TRUE));
+	// END EPIC MOD
 	return SUCCEEDED(result);
 }
 
@@ -559,4 +684,6 @@ bool visualStudio::Break(const EnvDTE::DebuggerPtr& debugger)
 	return g_comThread->CallInThread(&::Break, debugger);
 }
 
+// BEGIN EPIC MOD
 #endif
+// END EPIC MOD
\ No newline at end of file
diff --git a/Engine/Source/Developer/Windows/LiveCodingServer/Private/External/LC_VisualStudioAutomation.h b/Engine/Source/Developer/Windows/LiveCodingServer/Private/External/LC_VisualStudioAutomation.h
--- a/Engine/Source/Developer/Windows/LiveCodingServer/Private/External/LC_VisualStudioAutomation.h
+++ b/Engine/Source/Developer/Windows/LiveCodingServer/Private/External/LC_VisualStudioAutomation.h
@@ -1,11 +1,15 @@
-// Copyright 2011-2019 Molecular Matters GmbH, all rights reserved.
+// Copyright 2011-2020 Molecular Matters GmbH, all rights reserved.
 
 #pragma once
 
 #include "LC_Types.h"
+#include "LC_ProcessTypes.h"
+#include "LC_ThreadTypes.h"
+// BEGIN EPIC MOD
 #include "VisualStudioDTE.h"
 
 #if WITH_VISUALSTUDIO_DTE
+// END EPIC MOD
 
 namespace visualStudio
 {
@@ -14,13 +18,13 @@ namespace visualStudio
 
 
 	// Finds a Visual Studio debugger instance currently attached to the process with the given ID
-	EnvDTE::DebuggerPtr FindDebuggerAttachedToProcess(unsigned int processId);
+	EnvDTE::DebuggerPtr FindDebuggerAttachedToProcess(Process::Id processId);
 
 	// Finds a Visual Studio debugger instance currently debugging the process with the given ID
-	EnvDTE::DebuggerPtr FindDebuggerForProcess(unsigned int processId);
+	EnvDTE::DebuggerPtr FindDebuggerForProcess(Process::Id processId);
 
 	// Attaches a Visual Studio debugger instance to the process with the given ID
-	bool AttachToProcess(const EnvDTE::DebuggerPtr& debugger, unsigned int processId);
+	bool AttachToProcess(const EnvDTE::DebuggerPtr& debugger, Process::Id processId);
 
 	// Enumerates all threads of a debugger instance
 	types::vector<EnvDTE::ThreadPtr> EnumerateThreads(const EnvDTE::DebuggerPtr& debugger);
@@ -29,13 +33,13 @@ namespace visualStudio
 	bool FreezeThreads(const EnvDTE::DebuggerPtr& debugger, const types::vector<EnvDTE::ThreadPtr>& threads);
 
 	// Freezes a single thread with the given thread ID
-	bool FreezeThread(const EnvDTE::DebuggerPtr& debugger, const types::vector<EnvDTE::ThreadPtr>& threads, unsigned int threadId);
+	bool FreezeThread(const EnvDTE::DebuggerPtr& debugger, const types::vector<EnvDTE::ThreadPtr>& threads, Thread::Id threadId);
 
 	// Thaws all given threads
 	bool ThawThreads(const EnvDTE::DebuggerPtr& debugger, const types::vector<EnvDTE::ThreadPtr>& threads);
 
 	// Thaws a single thread with the given thread ID
-	bool ThawThread(const EnvDTE::DebuggerPtr& debugger, const types::vector<EnvDTE::ThreadPtr>& threads, unsigned int threadId);
+	bool ThawThread(const EnvDTE::DebuggerPtr& debugger, const types::vector<EnvDTE::ThreadPtr>& threads, Thread::Id threadId);
 
 	// Resumes the process in the debugger
 	bool Resume(const EnvDTE::DebuggerPtr& debugger);
@@ -44,4 +48,6 @@ namespace visualStudio
 	bool Break(const EnvDTE::DebuggerPtr& debugger);
 }
 
+// BEGIN EPIC MOD
 #endif
+// END EPIC MOD
\ No newline at end of file
diff --git a/Engine/Source/Developer/Windows/LiveCodingServer/Private/LiveCodingServer.cpp b/Engine/Source/Developer/Windows/LiveCodingServer/Private/LiveCodingServer.cpp
--- a/Engine/Source/Developer/Windows/LiveCodingServer/Private/LiveCodingServer.cpp
+++ b/Engine/Source/Developer/Windows/LiveCodingServer/Private/LiveCodingServer.cpp
@@ -3,11 +3,12 @@
 #include "LiveCodingServer.h"
 #include "External/LC_Scheduler.h"
 #include "External/LC_UniqueId.h"
-#include "External/LC_FileUtil.h"
+#include "External/LC_Filesystem.h"
 #include "External/LC_AppSettings.h"
 #include "External/LC_ServerCommandThread.h"
 #include "External/LC_Compiler.h"
 #include "External/LC_Environment.h"
+#include "External/LC_Process.h"
 
 FLiveCodingServer* GLiveCodingServer = nullptr;
 
@@ -28,7 +29,7 @@ void FLiveCodingServer::Start(const wchar_t* InProcessGroupName)
 	ProcessGroupName = InProcessGroupName;
 
 	scheduler::Startup();
-	file::Startup();
+	Filesystem::Startup();
 	uniqueId::Startup();
 	appSettings::Startup(ProcessGroupName.c_str());
 
@@ -43,7 +44,7 @@ void FLiveCodingServer::Stop()
 
 	appSettings::Shutdown();
 	uniqueId::Shutdown();
-	file::Shutdown();
+	Filesystem::Shutdown();
 	scheduler::Shutdown();
 
 	ProcessGroupName.clear();
@@ -64,7 +65,7 @@ void FLiveCodingServer::SetLinkerPath(const wchar_t* LinkerPath, const TMap<FStr
 
 	if (LinkerEnvironment.Num() > 0)
 	{
-		process::Environment* environment = process::CreateEnvironmentFromMap(LinkerEnvironment);
+		Process::Environment* environment = Process::CreateEnvironmentFromMap(LinkerEnvironment);
 		compiler::AddEnvironmentToCache(LinkerPath, environment);
 	}
 }
diff --git a/Engine/Source/Developer/Windows/LiveCodingServer/Private/LiveCodingServerModule.cpp b/Engine/Source/Developer/Windows/LiveCodingServer/Private/LiveCodingServerModule.cpp
--- a/Engine/Source/Developer/Windows/LiveCodingServer/Private/LiveCodingServerModule.cpp
+++ b/Engine/Source/Developer/Windows/LiveCodingServer/Private/LiveCodingServerModule.cpp
@@ -10,19 +10,19 @@ IMPLEMENT_MODULE(FLiveCodingServerModule, LiveCodingServer)
 
 DEFINE_LOG_CATEGORY_STATIC(LogLiveCodingServer, Display, All);
 
-static void ServerOutputHandler(logging::Channel::Enum Channel, logging::Type::Enum Type, const wchar_t* const Text)
+static void ServerOutputHandler(Logging::Channel::Enum Channel, Logging::Type::Enum Type, const wchar_t* const Text)
 {
 	FString TrimText = FString(Text).TrimEnd();
 	switch (Type)
 	{
-	case logging::Type::LOG_ERROR:
+	case Logging::Type::LOG_ERROR:
 		UE_LOG(LogLiveCodingServer, Error, TEXT("%s"), *TrimText);
 		break;
-	case logging::Type::LOG_WARNING:
+	case Logging::Type::LOG_WARNING:
 		// There are some warnings generated in the dev channel that aren't really actionable by the users.
 		// For example, warnings about symbols being eliminated by the linker.  It would be nice to just 
 		// filter that specific warning, but we can't.
-		if (Channel == logging::Channel::DEV)
+		if (Channel == Logging::Channel::DEV)
 		{
 			UE_LOG(LogLiveCodingServer, Verbose, TEXT("%s"), *TrimText);
 		}
@@ -36,18 +36,18 @@ static void ServerOutputHandler(logging::Channel::Enum Channel, logging::Type::E
 		break;
 	}
 
-	if (Channel == logging::Channel::USER)
+	if (Channel == Logging::Channel::USER)
 	{
 		ELiveCodingLogVerbosity Verbosity;
 		switch (Type)
 		{
-		case logging::Type::LOG_SUCCESS:
+		case Logging::Type::LOG_SUCCESS:
 			Verbosity = ELiveCodingLogVerbosity::Success;
 			break;
-		case logging::Type::LOG_ERROR:
+		case Logging::Type::LOG_ERROR:
 			Verbosity = ELiveCodingLogVerbosity::Failure;
 			break;
-		case logging::Type::LOG_WARNING:
+		case Logging::Type::LOG_WARNING:
 			Verbosity = ELiveCodingLogVerbosity::Warning;
 			break;
 		default:
@@ -60,7 +60,7 @@ static void ServerOutputHandler(logging::Channel::Enum Channel, logging::Type::E
 
 void FLiveCodingServerModule::StartupModule()
 {
-	logging::SetOutputHandler(&ServerOutputHandler);
+	Logging::SetOutputHandler(&ServerOutputHandler);
 
 	GLiveCodingServer = new FLiveCodingServer();
 
@@ -77,5 +77,5 @@ void FLiveCodingServerModule::ShutdownModule()
 		GLiveCodingServer = nullptr;
 	}
 
-	logging::SetOutputHandler(nullptr);
+	Logging::SetOutputHandler(nullptr);
 }
