Crash fix in LC_ClientStartupThread.cpp.  If LPP is initialized but not started, the editor/game/cooker could crash on shutdown.

diff --git a/Engine/Source/Developer/Windows/LiveCoding/Private/External/LC_NamedSharedMemory.cpp b/Engine/Source/Developer/Windows/LiveCoding/Private/External/LC_NamedSharedMemory.cpp
--- a/Engine/Source/Developer/Windows/LiveCoding/Private/External/LC_NamedSharedMemory.cpp
+++ b/Engine/Source/Developer/Windows/LiveCoding/Private/External/LC_NamedSharedMemory.cpp
@@ -62,6 +62,12 @@ Process::NamedSharedMemory* Process::CreateNamedSharedMemory(const wchar_t* name
 
 void Process::DestroyNamedSharedMemory(NamedSharedMemory*& memory)
 {
+	// BEGIN EPIC MOD
+	if (memory == nullptr)
+	{
+		return;
+	}
+	// END EPIC MOD
 	::UnmapViewOfFile(memory->memoryView);
 	::CloseHandle(memory->memoryMapping);
 	delete memory;
