Fixing a patching issue with Live coding.  If multiple modules are loaded, then depending on which module is being updated, it might fail locating an address region.  Disabling a new Live++ 1.6.3 feature for now to prevent this from happening.

#rnx

diff --git a/Engine/Source/Developer/Windows/LiveCodingServer/Private/External/LC_VirtualMemoryRange.cpp b/Engine/Source/Developer/Windows/LiveCodingServer/Private/External/LC_VirtualMemoryRange.cpp
--- a/Engine/Source/Developer/Windows/LiveCodingServer/Private/External/LC_VirtualMemoryRange.cpp
+++ b/Engine/Source/Developer/Windows/LiveCodingServer/Private/External/LC_VirtualMemoryRange.cpp
@@ -22,6 +22,10 @@ VirtualMemoryRange::VirtualMemoryRange(Process::Handle processHandle, const void
 
 void VirtualMemoryRange::ReservePages(void)
 {
+// BEGIN EPIC MOD
+// If we have two modules in the same process, then they might overlap in the 2GB range. Thus when we try to update one module
+// the other module might prevent us from finding address space.  Disable this feature for now.
+#if 0
 	// reserve all free pages in the virtual memory range.
 	// pages must be aligned to the given alignment.
 	for (const void* address = m_addressStart; address < m_addressEnd; /* nothing */)
@@ -53,6 +57,8 @@ void VirtualMemoryRange::ReservePages(void)
 		// keep on searching
 		address = pointer::Offset<const void*>(memoryInfo.BaseAddress, memoryInfo.RegionSize);
 	}
+#endif
+// END EPIC MOD
 }
 
 
