Fixing a few PVS warning.  The memcpy issue is intentional by the 3rd party library.

diff --git a/Engine/Source/Developer/Windows/LiveCoding/Private/External/LC_Process.cpp b/Engine/Source/Developer/Windows/LiveCoding/Private/External/LC_Process.cpp
--- a/Engine/Source/Developer/Windows/LiveCoding/Private/External/LC_Process.cpp
+++ b/Engine/Source/Developer/Windows/LiveCoding/Private/External/LC_Process.cpp
@@ -596,7 +596,9 @@ std::vector<Thread::Id> Process::EnumerateThreads(Id processId)
 		else if (status < 0)
 		{
 			// something went wrong
-			LC_ERROR_USER("Cannot enumerate threads in process (PID: %d)", processId);
+			// BEGIN EPIC MOD
+			LC_ERROR_USER("Cannot enumerate threads in process (PID: %d)", +processId);
+			// END EPIC MOD
 			::free(processSnapshot);
 
 			return threadIds;
diff --git a/Engine/Source/Developer/Windows/LiveCoding/Private/External/xxhash.h b/Engine/Source/Developer/Windows/LiveCoding/Private/External/xxhash.h
--- a/Engine/Source/Developer/Windows/LiveCoding/Private/External/xxhash.h
+++ b/Engine/Source/Developer/Windows/LiveCoding/Private/External/xxhash.h
@@ -1310,7 +1310,9 @@ XXH_PUBLIC_API XXH_errorcode XXH32_reset(XXH32_state_t* statePtr, XXH32_hash_t s
     state.v3 = seed + 0;
     state.v4 = seed - PRIME32_1;
     /* do not write into reserved, planned to be removed in a future version */
-    memcpy(statePtr, &state, sizeof(state) - sizeof(state.reserved));
+	// BEGIN EPIC MOD
+    memcpy(statePtr, &state, sizeof(state) - sizeof(state.reserved)); //-V512_UNDERFLOW_OFF PVS
+	// END EPIC MOD
     return XXH_OK;
 }
 
@@ -1822,7 +1824,9 @@ XXH_PUBLIC_API XXH_errorcode XXH64_reset(XXH64_state_t* statePtr, XXH64_hash_t s
     state.v3 = seed + 0;
     state.v4 = seed - PRIME64_1;
      /* do not write into reserved64, might be removed in a future version */
-    memcpy(statePtr, &state, sizeof(state) - sizeof(state.reserved64));
+	// BEGIN EPIC MOD
+	memcpy(statePtr, &state, sizeof(state) - sizeof(state.reserved64)); //-V512_UNDERFLOW_OFF PVS
+	// END EPIC MOD
     return XXH_OK;
 }
 
