Fixing a PVS warning.  Changing the code to avoid the + operator since PVS is having issues interpreting the code in this instance.

diff --git a/Engine/Source/Developer/Windows/LiveCoding/Private/External/LC_Process.cpp b/Engine/Source/Developer/Windows/LiveCoding/Private/External/LC_Process.cpp
--- a/Engine/Source/Developer/Windows/LiveCoding/Private/External/LC_Process.cpp
+++ b/Engine/Source/Developer/Windows/LiveCoding/Private/External/LC_Process.cpp
@@ -596,8 +596,8 @@ std::vector<Thread::Id> Process::EnumerateThreads(Id processId)
 		else if (status < 0)
 		{
 			// something went wrong
-			// BEGIN EPIC MOD
-			LC_ERROR_USER("Cannot enumerate threads in process (PID: %d)", +processId);
+			// BEGIN EPIC MOD - PVS is having problems dealing with + here. 
+			LC_ERROR_USER("Cannot enumerate threads in process (PID: %d)", static_cast<DWORD>(processId));
 			// END EPIC MOD
 			::free(processSnapshot);
 
@@ -621,7 +621,9 @@ std::vector<Thread::Id> Process::EnumerateThreads(Id processId)
 			if (processInfo->NextEntryOffset == 0u)
 			{
 				// we couldn't find our process
-				LC_ERROR_USER("Cannot enumerate threads, process not found (PID: %d)", +processId);
+				// BEGIN EPIC MOD - PVS is having problems dealing with + here. 
+				LC_ERROR_USER("Cannot enumerate threads, process not found (PID: %d)", static_cast<DWORD>(processId));
+				// END EPIC MOD
 				::free(processSnapshot);
 
 				return threadIds;
