Fixed static analysis warning. The warning is correct; _setjmp is declared with and without the _Out_ annotation in setjmp.h and intrin.h respectively. It is however, harmless in this instance.

#rnx

diff --git a/Engine/Source/Developer/Windows/LiveCoding/Private/External/xxhash.c b/Engine/Source/Developer/Windows/LiveCoding/Private/External/xxhash.c
--- a/Engine/Source/Developer/Windows/LiveCoding/Private/External/xxhash.c
+++ b/Engine/Source/Developer/Windows/LiveCoding/Private/External/xxhash.c
@@ -45,6 +45,7 @@
 #pragma warning(push)
 #pragma warning(disable:6385) // warning C6385: Reading invalid data from '&state'
 #pragma warning(disable:4668) // warning C4668: '__cplusplus' is not defined as a preprocessor macro, replacing with '0' for '#if/#elif'
+#pragma warning(disable:28251) // warning C28251: Inconsistent annotation for '_setjmp'
 // END EPIC MODS
 
 #include "xxhash.h"
