Integration of Live++ [REDACTED_PHONE]/18/01)

This fixes and re-enables a feature of 1.6 that didn't work in an Unreal environment.

#rnx

diff --git a/Engine/Source/Developer/Windows/LiveCoding/Private/External/LC_ClientStartupThread.cpp b/Engine/Source/Developer/Windows/LiveCoding/Private/External/LC_ClientStartupThread.cpp
--- a/Engine/Source/Developer/Windows/LiveCoding/Private/External/LC_ClientStartupThread.cpp
+++ b/Engine/Source/Developer/Windows/LiveCoding/Private/External/LC_ClientStartupThread.cpp
@@ -114,7 +114,10 @@ ClientStartupThread::~ClientStartupThread(void)
 	::CloseHandle(m_job);
 
 	// clean up interprocess objects
-	Process::DestroyNamedSharedMemory(m_sharedMemory);
+	if (m_sharedMemory)
+	{
+		Process::DestroyNamedSharedMemory(m_sharedMemory);
+	}
 }
 
 
diff --git a/Engine/Source/Developer/Windows/LiveCoding/Private/External/LC_Process.cpp b/Engine/Source/Developer/Windows/LiveCoding/Private/External/LC_Process.cpp
--- a/Engine/Source/Developer/Windows/LiveCoding/Private/External/LC_Process.cpp
+++ b/Engine/Source/Developer/Windows/LiveCoding/Private/External/LC_Process.cpp
@@ -596,9 +596,7 @@ std::vector<Thread::Id> Process::EnumerateThreads(Id processId)
 		else if (status < 0)
 		{
 			// something went wrong
-			// BEGIN EPIC MOD - PVS is having problems dealing with + here. 
-			LC_ERROR_USER("Cannot enumerate threads in process (PID: %d)", static_cast<DWORD>(processId));
-			// END EPIC MOD
+			LC_ERROR_USER("Cannot enumerate threads in process (PID: %d)", +processId);
 			::free(processSnapshot);
 
 			return threadIds;
diff --git a/Engine/Source/Developer/Windows/LiveCoding/Private/External/LPP_API.h b/Engine/Source/Developer/Windows/LiveCoding/Private/External/LPP_API.h
--- a/Engine/Source/Developer/Windows/LiveCoding/Private/External/LPP_API.h
+++ b/Engine/Source/Developer/Windows/LiveCoding/Private/External/LPP_API.h
@@ -105,7 +105,7 @@
 /******************************************************************************/
 
 // version string
-#define LPP_VERSION "1.6.3"
+#define LPP_VERSION "1.6.4"
 
 // macros to temporarily enable/disable optimizations
 #ifdef __clang__
diff --git a/Engine/Source/Developer/Windows/LiveCodingServer/Private/External/LC_LiveModule.cpp b/Engine/Source/Developer/Windows/LiveCodingServer/Private/External/LC_LiveModule.cpp
--- a/Engine/Source/Developer/Windows/LiveCodingServer/Private/External/LC_LiveModule.cpp
+++ b/Engine/Source/Developer/Windows/LiveCodingServer/Private/External/LC_LiveModule.cpp
@@ -85,7 +85,6 @@ namespace
 	}
 #endif
 
-
 	// common linker options:
 	//	*) create x86/x64 code
 	//	*) don't echo command-line options
@@ -1248,16 +1247,9 @@ void LiveModule::RegisterProcess(LiveProcess* liveProcess, void* moduleBase, con
 {
 	m_moduleCache->RegisterProcess(m_mainModuleToken, liveProcess, moduleBase);
 
-#if LC_64_BIT
-	const void* lowerBound = GetLowerBoundIn4GBRange(moduleBase);
-	const void* upperBound = GetUpperBoundIn4GBRange(moduleBase);
-	VirtualMemoryRange* virtualMemoryRange = new VirtualMemoryRange(liveProcess->GetProcessHandle(), lowerBound, upperBound, 64u * 1024u);
-	virtualMemoryRange->ReservePages();
-#else
-	VirtualMemoryRange* virtualMemoryRange = new VirtualMemoryRange(liveProcess->GetProcessHandle(), moduleBase, moduleBase, 64u * 1024u);
-#endif
+	liveProcess->ReserveVirtualMemoryPages(moduleBase);
 
-	PerProcessData perProcessData = { liveProcess, moduleBase, modulePath, virtualMemoryRange };
+	PerProcessData perProcessData = { liveProcess, moduleBase, modulePath };
 	m_perProcessData.emplace_back(perProcessData);
 }
 
@@ -1274,8 +1266,7 @@ void LiveModule::UnregisterProcess(LiveProcess* liveProcess)
 		const PerProcessData& data = *it;
 		if (data.liveProcess == liveProcess)
 		{
-			data.virtualMemoryRange->FreeReservedPages();
-			delete data.virtualMemoryRange;
+			liveProcess->FreeVirtualMemoryPages(data.originalModuleBase);
 
 			m_perProcessData.erase(it);
 			break;
@@ -3241,7 +3232,7 @@ LiveModule::ErrorType::Enum LiveModule::Update(FileAttributeCache* fileCache, Di
 			Process::Suspend(data.liveProcess->GetProcessHandle());
 
 			// free our page reservations in the +-2GB range of the main module
-			data.virtualMemoryRange->FreeReservedPages();
+			data.liveProcess->FreeVirtualMemoryPages(data.originalModuleBase);
 
 			const executable::PreferredBase preferredImageBase = FindPreferredImageBase(patchImageSize, data.liveProcess->GetProcessId(), data.liveProcess->GetProcessHandle(), data.originalModuleBase);
 
@@ -3292,7 +3283,7 @@ LiveModule::ErrorType::Enum LiveModule::Update(FileAttributeCache* fileCache, Di
 	for (size_t i = 0u; i < processCount; ++i)
 	{
 		const PerProcessData& data = processData[i];
-		data.virtualMemoryRange->ReservePages();
+		data.liveProcess->ReserveVirtualMemoryPages(data.originalModuleBase);
 	}
 #endif
 
diff --git a/Engine/Source/Developer/Windows/LiveCodingServer/Private/External/LC_LiveModule.h b/Engine/Source/Developer/Windows/LiveCodingServer/Private/External/LC_LiveModule.h
--- a/Engine/Source/Developer/Windows/LiveCodingServer/Private/External/LC_LiveModule.h
+++ b/Engine/Source/Developer/Windows/LiveCodingServer/Private/External/LC_LiveModule.h
@@ -30,7 +30,6 @@ public:
 		LiveProcess* liveProcess;
 		void* originalModuleBase;
 		std::wstring modulePath;
-		VirtualMemoryRange* virtualMemoryRange;
 	};
 
 	struct ErrorType
diff --git a/Engine/Source/Developer/Windows/LiveCodingServer/Private/External/LC_LiveProcess.cpp b/Engine/Source/Developer/Windows/LiveCodingServer/Private/External/LC_LiveProcess.cpp
--- a/Engine/Source/Developer/Windows/LiveCodingServer/Private/External/LC_LiveProcess.cpp
+++ b/Engine/Source/Developer/Windows/LiveCodingServer/Private/External/LC_LiveProcess.cpp
@@ -15,6 +15,42 @@
 // END EPIC MOD
 
 
+#if LC_64_BIT
+namespace
+{
+	static const void* GetLowerBoundIn4GBRange(const void* moduleBase)
+	{
+		// nothing can be loaded at address 0x0, 64 KB seems like a realistic minimum
+		const uint64_t LOWEST_ADDRESS = 64u * 1024u;
+		const uint64_t base = pointer::AsInteger<uint64_t>(moduleBase);
+
+		// make sure we don't underflow
+		if (base >= 0x80000000ull + LOWEST_ADDRESS)
+		{
+			return pointer::FromInteger<const void*>(base - 0x80000000ull);
+		}
+
+		// operation would underflow
+		return pointer::FromInteger<const void*>(LOWEST_ADDRESS);
+	}
+
+	static const void* GetUpperBoundIn4GBRange(const void* moduleBase)
+	{
+		const uint64_t base = pointer::AsInteger<uint64_t>(moduleBase);
+
+		// make sure we don't overflow
+		if (base <= 0xFFFFFFFF7FFFFFFFull)
+		{
+			return pointer::FromInteger<const void*>(base + 0x80000000ull);
+		}
+
+		// operation would overflow
+		return pointer::FromInteger<const void*>(0xFFFFFFFFFFFFFFFFull);
+	}
+}
+#endif
+
+
 LiveProcess::LiveProcess(Process::Handle processHandle, Process::Id processId, Thread::Id commandThreadId, const void* jumpToSelf, const DuplexPipe* pipe,
 	const wchar_t* imagePath, const wchar_t* commandLine, const wchar_t* workingDirectory, const void* environment, size_t environmentSize)
 	: m_processHandle(processHandle)
@@ -38,6 +74,7 @@ LiveProcess::LiveProcess(Process::Handle processHandle, Process::Id processId, T
 // END EPIC MOD
 	, m_codeCave(nullptr)
 	, m_restartState(RestartState::DEFAULT)
+	, m_virtualMemoryRange(processHandle)
 {
 	m_imagesTriedToLoad.reserve(256u);
 }
@@ -278,6 +315,29 @@ bool LiveProcess::WasSuccessfulRestart(void) const
 }
 
 
+void LiveProcess::ReserveVirtualMemoryPages(void* moduleBase)
+{
+#if LC_64_BIT
+	const void* lowerBound = GetLowerBoundIn4GBRange(moduleBase);
+	const void* upperBound = GetUpperBoundIn4GBRange(moduleBase);
+	m_virtualMemoryRange.ReservePages(lowerBound, upperBound, 64u * 1024u);
+#else
+	LC_UNUSED(moduleBase);
+#endif
+}
+
+
+void LiveProcess::FreeVirtualMemoryPages(void* moduleBase)
+{
+#if LC_64_BIT
+	const void* lowerBound = GetLowerBoundIn4GBRange(moduleBase);
+	const void* upperBound = GetUpperBoundIn4GBRange(moduleBase);
+	m_virtualMemoryRange.FreePages(lowerBound, upperBound);
+#else
+	LC_UNUSED(moduleBase);
+#endif
+}
+
 // BEGIN EPIC MOD - Allow lazy-loading modules
 void LiveProcess::AddLazyLoadedModule(const std::wstring moduleName, Windows::HMODULE moduleBase)
 {
diff --git a/Engine/Source/Developer/Windows/LiveCodingServer/Private/External/LC_LiveProcess.h b/Engine/Source/Developer/Windows/LiveCodingServer/Private/External/LC_LiveProcess.h
--- a/Engine/Source/Developer/Windows/LiveCodingServer/Private/External/LC_LiveProcess.h
+++ b/Engine/Source/Developer/Windows/LiveCodingServer/Private/External/LC_LiveProcess.h
@@ -6,6 +6,7 @@
 #include "LC_ThreadTypes.h"
 #include "LC_Executable.h"
 #include "LC_MemoryBlock.h"
+#include "LC_VirtualMemoryRange.h"
 // BEGIN EPIC MOD
 #include "LC_Types.h"
 #include "VisualStudioDTE.h"
@@ -48,6 +49,10 @@ public:
 	bool WasSuccessfulRestart(void) const;
 
 
+	void ReserveVirtualMemoryPages(void* moduleBase);
+	void FreeVirtualMemoryPages(void* moduleBase);
+
+
 	inline Process::Handle GetProcessHandle(void) const
 	{
 		return m_processHandle;
@@ -153,6 +158,8 @@ private:
 
 	RestartState::Enum m_restartState;
 
+	VirtualMemoryRange m_virtualMemoryRange;
+
 	LC_DISABLE_COPY(LiveProcess);
 	LC_DISABLE_MOVE(LiveProcess);
 	LC_DISABLE_ASSIGNMENT(LiveProcess);
diff --git a/Engine/Source/Developer/Windows/LiveCodingServer/Private/External/LC_VirtualMemoryRange.cpp b/Engine/Source/Developer/Windows/LiveCodingServer/Private/External/LC_VirtualMemoryRange.cpp
--- a/Engine/Source/Developer/Windows/LiveCodingServer/Private/External/LC_VirtualMemoryRange.cpp
+++ b/Engine/Source/Developer/Windows/LiveCodingServer/Private/External/LC_VirtualMemoryRange.cpp
@@ -10,28 +10,43 @@
 #include <cinttypes>
 // END EPIC MOD
 
-VirtualMemoryRange::VirtualMemoryRange(Process::Handle processHandle, const void* addressStart, const void* addressEnd, size_t alignment)
+VirtualMemoryRange::VirtualMemoryRange(Process::Handle processHandle)
 	: m_processHandle(processHandle)
-	, m_addressStart(addressStart)
-	, m_addressEnd(addressEnd)
-	, m_alignment(alignment)
+	, m_pageData()
+	, m_cs()
 {
-	m_pageData.reserve(32u);
+	m_pageData.reserve(256u);
 }
 
 
-void VirtualMemoryRange::ReservePages(void)
+VirtualMemoryRange::~VirtualMemoryRange(void)
 {
-// BEGIN EPIC MOD
-// If we have two modules in the same process, then they might overlap in the 2GB range. Thus when we try to update one module
-// the other module might prevent us from finding address space.  Disable this feature for now.
-#if 0
+	for (const auto& it : m_pageData)
+	{
+		const bool success = ::VirtualFreeEx(+m_processHandle, it.address, 0u, MEM_RELEASE);
+		if (!success)
+		{
+			LC_WARNING_DEV("Cannot free virtual memory region at 0x%p", it.address);
+		}
+	}
+
+	m_pageData.clear();
+}
+
+
+void VirtualMemoryRange::ReservePages(const void* addressStart, const void* addressEnd, size_t alignment)
+{
+	CriticalSection::ScopedLock lock(&m_cs);
+
+	LC_LOG_DEV("Reserving pages in the range 0x%p to 0x%p", addressStart, addressEnd);
+	LC_LOG_INDENT_DEV;
+
 	// reserve all free pages in the virtual memory range.
 	// pages must be aligned to the given alignment.
-	for (const void* address = m_addressStart; address < m_addressEnd; /* nothing */)
+	for (const void* address = addressStart; address < addressEnd; /* nothing */)
 	{
 		// align address to be scanned
-		address = pointer::AlignTop<const void*>(address, m_alignment);
+		address = pointer::AlignTop<const void*>(address, alignment);
 
 		::MEMORY_BASIC_INFORMATION memoryInfo = {};
 		const size_t bytesReturned = ::VirtualQueryEx(+m_processHandle, address, &memoryInfo, sizeof(::MEMORY_BASIC_INFORMATION));
@@ -41,7 +56,7 @@ void VirtualMemoryRange::ReservePages(void)
 		{
 			// work out the maximum size of the page allocation.
 			// we should not allocate past the end of the range.
-			const size_t bytesLeft = pointer::Displacement<size_t>(memoryInfo.BaseAddress, m_addressEnd);
+			const size_t bytesLeft = pointer::Displacement<size_t>(memoryInfo.BaseAddress, addressEnd);
 			const size_t size = std::min<size_t>(memoryInfo.RegionSize, bytesLeft);
 
 			// try to reserve this page.
@@ -49,7 +64,7 @@ void VirtualMemoryRange::ReservePages(void)
 			void* baseAddress = ::VirtualAllocEx(+m_processHandle, memoryInfo.BaseAddress, size, MEM_RESERVE, PAGE_NOACCESS);
 			if (baseAddress)
 			{
-				LC_LOG_DEV("Found virtual memory region at 0x%p with size 0x%" PRIX64, baseAddress, size);
+				LC_LOG_DEV("Reserving virtual memory region at 0x%p with size 0x%" PRIX64, baseAddress, size);
 				m_pageData.emplace_back(PageData { baseAddress });
 			}
 		}
@@ -57,21 +72,36 @@ void VirtualMemoryRange::ReservePages(void)
 		// keep on searching
 		address = pointer::Offset<const void*>(memoryInfo.BaseAddress, memoryInfo.RegionSize);
 	}
-#endif
-// END EPIC MOD
 }
 
 
-void VirtualMemoryRange::FreeReservedPages(void)
+void VirtualMemoryRange::FreePages(const void* addressStart, const void* addressEnd)
 {
-	for (const auto& it : m_pageData)
+	CriticalSection::ScopedLock lock(&m_cs);
+
+	LC_LOG_DEV("Freeing pages in the range 0x%p to 0x%p", addressStart, addressEnd);
+	LC_LOG_INDENT_DEV;
+
+	for (auto it = m_pageData.begin(); it != m_pageData.end(); /* nothing */)
 	{
-		const bool success = ::VirtualFreeEx(+m_processHandle, it.address, 0u, MEM_RELEASE);
-		if (!success)
+		const PageData& data = *it;
+		if ((data.address >= addressStart) && (data.address < addressEnd))
 		{
-			LC_WARNING_USER("Cannot free virtual memory region at 0x%p", it.address);
+			const bool success = ::VirtualFreeEx(+m_processHandle, data.address, 0u, MEM_RELEASE);
+			if (success)
+			{
+				LC_LOG_DEV("Freeing virtual memory region at 0x%p", data.address);
+			}
+			else
+			{
+				LC_WARNING_DEV("Cannot free virtual memory region at 0x%p", data.address);
+			}
+
+			it = m_pageData.erase(it);
+		}
+		else
+		{
+			++it;
 		}
 	}
-
-	m_pageData.clear();
 }
diff --git a/Engine/Source/Developer/Windows/LiveCodingServer/Private/External/LC_VirtualMemoryRange.h b/Engine/Source/Developer/Windows/LiveCodingServer/Private/External/LC_VirtualMemoryRange.h
--- a/Engine/Source/Developer/Windows/LiveCodingServer/Private/External/LC_VirtualMemoryRange.h
+++ b/Engine/Source/Developer/Windows/LiveCodingServer/Private/External/LC_VirtualMemoryRange.h
@@ -14,10 +14,11 @@
 class VirtualMemoryRange
 {
 public:
-	VirtualMemoryRange(Process::Handle processHandle, const void* addressStart, const void* addressEnd, size_t alignment);
+	explicit VirtualMemoryRange(Process::Handle processHandle);
+	~VirtualMemoryRange(void);
 
-	void ReservePages(void);
-	void FreeReservedPages(void);
+	void ReservePages(const void* addressStart, const void* addressEnd, size_t alignment);
+	void FreePages(const void* addressStart, const void* addressEnd);
 
 private:
 	struct PageData
@@ -26,9 +27,6 @@ private:
 	};
 
 	Process::Handle m_processHandle;
-	const void* m_addressStart;
-	const void* m_addressEnd;
-	size_t m_alignment;
-
 	types::vector<PageData> m_pageData;
+	CriticalSection m_cs;
 };
