Refixing a PVS warning stomped by the integration

#rnx

diff --git a/Engine/Source/Developer/Windows/LiveCoding/Private/External/LC_Process.cpp b/Engine/Source/Developer/Windows/LiveCoding/Private/External/LC_Process.cpp
--- a/Engine/Source/Developer/Windows/LiveCoding/Private/External/LC_Process.cpp
+++ b/Engine/Source/Developer/Windows/LiveCoding/Private/External/LC_Process.cpp
@@ -596,7 +596,9 @@ std::vector<Thread::Id> Process::EnumerateThreads(Id processId)
 		else if (status < 0)
 		{
 			// something went wrong
-			LC_ERROR_USER("Cannot enumerate threads in process (PID: %d)", +processId);
+			// BEGIN EPIC MOD - PVS is having problems dealing with + here. 
+			LC_ERROR_USER("Cannot enumerate threads in process (PID: %d)", static_cast<DWORD>(processId));
+			// END EPIC MOD
 			::free(processSnapshot);
 
 			return threadIds;
