Fixing compile issues with LiveCodingServer.

#rnx

diff --git a/Engine/Source/Developer/Windows/LiveCodingServer/Private/External/LC_LiveProcess.cpp b/Engine/Source/Developer/Windows/LiveCodingServer/Private/External/LC_LiveProcess.cpp
--- a/Engine/Source/Developer/Windows/LiveCodingServer/Private/External/LC_LiveProcess.cpp
+++ b/Engine/Source/Developer/Windows/LiveCodingServer/Private/External/LC_LiveProcess.cpp
@@ -16,7 +16,9 @@
 
 
 #if LC_64_BIT
-namespace
+// EPIC BEGIN MOD
+namespace liveProcess // Duplicate definitions 
+// EPIC END MOD
 {
 	static const void* GetLowerBoundIn4GBRange(const void* moduleBase)
 	{
@@ -318,8 +320,10 @@ bool LiveProcess::WasSuccessfulRestart(void) const
 void LiveProcess::ReserveVirtualMemoryPages(void* moduleBase)
 {
 #if LC_64_BIT
-	const void* lowerBound = GetLowerBoundIn4GBRange(moduleBase);
-	const void* upperBound = GetUpperBoundIn4GBRange(moduleBase);
+	// BEGIN EPIC MOD
+	const void* lowerBound = liveProcess::GetLowerBoundIn4GBRange(moduleBase);
+	const void* upperBound = liveProcess::GetUpperBoundIn4GBRange(moduleBase);
+	// END EPIC MOD
 	m_virtualMemoryRange.ReservePages(lowerBound, upperBound, 64u * 1024u);
 #else
 	LC_UNUSED(moduleBase);
@@ -330,8 +334,10 @@ void LiveProcess::ReserveVirtualMemoryPages(void* moduleBase)
 void LiveProcess::FreeVirtualMemoryPages(void* moduleBase)
 {
 #if LC_64_BIT
-	const void* lowerBound = GetLowerBoundIn4GBRange(moduleBase);
-	const void* upperBound = GetUpperBoundIn4GBRange(moduleBase);
+	// BEGIN EPIC MOD
+	const void* lowerBound = liveProcess::GetLowerBoundIn4GBRange(moduleBase);
+	const void* upperBound = liveProcess::GetUpperBoundIn4GBRange(moduleBase);
+	// END EPIC MOD
 	m_virtualMemoryRange.FreePages(lowerBound, upperBound);
 #else
 	LC_UNUSED(moduleBase);
diff --git a/Engine/Source/Developer/Windows/LiveCodingServer/Private/External/LC_VirtualMemoryRange.cpp b/Engine/Source/Developer/Windows/LiveCodingServer/Private/External/LC_VirtualMemoryRange.cpp
--- a/Engine/Source/Developer/Windows/LiveCodingServer/Private/External/LC_VirtualMemoryRange.cpp
+++ b/Engine/Source/Developer/Windows/LiveCodingServer/Private/External/LC_VirtualMemoryRange.cpp
@@ -50,6 +50,10 @@ void VirtualMemoryRange::ReservePages(const void* addressStart, const void* addr
 
 		::MEMORY_BASIC_INFORMATION memoryInfo = {};
 		const size_t bytesReturned = ::VirtualQueryEx(+m_processHandle, address, &memoryInfo, sizeof(::MEMORY_BASIC_INFORMATION));
+		// BEGIN EPIC MOD
+		if (bytesReturned == 0)
+			break;
+		// END EPIC MOD
 
 		// we are only interested in free pages
 		if ((bytesReturned > 0u) && (memoryInfo.State == MEM_FREE))
diff --git a/Engine/Source/Developer/Windows/LiveCodingServer/Private/External/LC_VirtualMemoryRange.h b/Engine/Source/Developer/Windows/LiveCodingServer/Private/External/LC_VirtualMemoryRange.h
--- a/Engine/Source/Developer/Windows/LiveCodingServer/Private/External/LC_VirtualMemoryRange.h
+++ b/Engine/Source/Developer/Windows/LiveCodingServer/Private/External/LC_VirtualMemoryRange.h
@@ -7,6 +7,7 @@
 // END EPIC MOD
 #include "LC_ProcessTypes.h"
 // BEGIN EPIC MOD
+#include "LC_CriticalSection.h"
 #include "LC_Types.h"
 // END EPIC MOD
 
