Include shader compiler DLL version in FShaderFormatD3D/Vulkan::GetVersion(), similar to how it's done for other platforms.

diff --git a/Engine/Source/Developer/ShaderCompilerCommon/Private/DXCWrapper.cpp b/Engine/Source/Developer/ShaderCompilerCommon/Private/DXCWrapper.cpp
--- a/Engine/Source/Developer/ShaderCompilerCommon/Private/DXCWrapper.cpp
+++ b/Engine/Source/Developer/ShaderCompilerCommon/Private/DXCWrapper.cpp
@@ -4,10 +4,51 @@
 #include "HAL/PlatformProcess.h"
 #include "Templates/RefCounting.h"
 #include "Misc/Paths.h"
+#include "Hash/CityHash.h"
+
+#if PLATFORM_WINDOWS
+#include "Windows/WindowsHWrapper.h"
+#endif // PLATFORM_WINDOWS
 
 static TRefCountPtr<FDllHandle> GDxcHandle;
 static TRefCountPtr<FDllHandle> GShaderConductorHandle;
 
+static uint64 GetLoadedModuleVersion(const TCHAR* ModuleName)
+{
+#if PLATFORM_WINDOWS
+	HMODULE ModuleDll = ::GetModuleHandleW(ModuleName);
+	if (ModuleDll == INVALID_HANDLE_VALUE)
+	{
+		return 0;
+	}
+
+	TCHAR DllPath[4096] = {};
+	int32 PathLen = ::GetModuleFileNameW(ModuleDll, DllPath, UE_ARRAY_COUNT(DllPath));
+	if (PathLen == 0)
+	{
+		return 0;
+	}
+	check(PathLen < UE_ARRAY_COUNT(DllPath));
+
+	// Official shader compiler binaries will have the the version, but anything that we compile ourselves won't
+	uint64 FileVersion = FPlatformMisc::GetFileVersion(DllPath);
+	if (FileVersion == 0)
+	{
+		TUniquePtr<FArchive> FileReader(IFileManager::Get().CreateFileReader(DllPath));
+		if (FileReader)
+		{
+			TArray<char> FileContents;
+			FileContents.SetNumUninitialized(FileReader->TotalSize());
+			FileReader->Serialize(FileContents.GetData(), FileContents.Num());
+			FileVersion = CityHash64(FileContents.GetData(), FileContents.Num());
+		}
+	}
+	return FileVersion;
+#else // PLATFORM_WINDOWS
+	return 0;
+#endif // PLATFORM_WINDOWS
+}
+
 FDllHandle::FDllHandle(const TCHAR* InFilename)
 {
 #if PLATFORM_WINDOWS
@@ -39,6 +80,15 @@ FDxcModuleWrapper::FDxcModuleWrapper()
 	{
 		GDxcHandle->AddRef();
 	}
+
+	static uint64 DxcVersion = GetLoadedModuleVersion(TEXT("dxcompiler.dll"));
+
+	// If dxil.dll is present, it's automatically loaded by dxcompiler.dll and used for validation and signing.
+	// If dxil.dll is not present, shaders will silently be unsigned and will fail to load outside of development environment.
+	// This must be taken into account when computing DDC key for D3D shaders.
+	static uint64 DxilVersion = GetLoadedModuleVersion(TEXT("dxil.dll"));
+
+	ModuleVersionHash = HashCombine(GetTypeHash(DxcVersion), GetTypeHash(DxilVersion));
 }
 
 FDxcModuleWrapper::~FDxcModuleWrapper()
@@ -57,6 +107,10 @@ FShaderConductorModuleWrapper::FShaderConductorModuleWrapper()
 	{
 		GShaderConductorHandle->AddRef();
 	}
+
+	static uint64 DllVersion = GetLoadedModuleVersion(TEXT("ShaderConductor.dll"));
+
+	ModuleVersionHash = HashCombine(GetTypeHash(DllVersion), FDxcModuleWrapper::ModuleVersionHash);
 }
 
 FShaderConductorModuleWrapper::~FShaderConductorModuleWrapper()
diff --git a/Engine/Source/Developer/ShaderCompilerCommon/Public/DXCWrapper.h b/Engine/Source/Developer/ShaderCompilerCommon/Public/DXCWrapper.h
--- a/Engine/Source/Developer/ShaderCompilerCommon/Public/DXCWrapper.h
+++ b/Engine/Source/Developer/ShaderCompilerCommon/Public/DXCWrapper.h
@@ -19,7 +19,7 @@ class SHADERCOMPILERCOMMON_API FDxcModuleWrapper
 {
 protected:
 	TRefCountPtr<FDllHandle> Dxc;
-
+	uint32 ModuleVersionHash = 0;
 public:
 	FDxcModuleWrapper();
 	virtual ~FDxcModuleWrapper();
@@ -29,7 +29,7 @@ class SHADERCOMPILERCOMMON_API FShaderConductorModuleWrapper : public FDxcModule
 {
 protected:
 	TRefCountPtr<FDllHandle> ShaderConductor;
-
+	uint32 ModuleVersionHash = 0;
 public:
 	FShaderConductorModuleWrapper();
 	virtual ~FShaderConductorModuleWrapper();
diff --git a/Engine/Source/Developer/VulkanShaderFormat/Private/VulkanShaderFormat.cpp b/Engine/Source/Developer/VulkanShaderFormat/Private/VulkanShaderFormat.cpp
--- a/Engine/Source/Developer/VulkanShaderFormat/Private/VulkanShaderFormat.cpp
+++ b/Engine/Source/Developer/VulkanShaderFormat/Private/VulkanShaderFormat.cpp
@@ -48,14 +48,25 @@ class FShaderFormatVulkan : public IShaderFormat
 		return -1;
 	}
 
+	uint32 ShaderConductorVersionHash;
+
 public:
+
+	FShaderFormatVulkan(uint32 InShaderConductorVersionHash)
+		: ShaderConductorVersionHash(InShaderConductorVersionHash)
+	{
+	}
+
 	virtual uint32 GetVersion(FName Format) const override
 	{
-		const uint8 HLSLCCVersion = ((HLSLCC_VersionMajor & 0x0f) << 4) | (HLSLCC_VersionMinor & 0x0f);
-		uint16 Version = ((HLSLCCVersion & 0xff) << 8) | (InternalGetVersion(Format) & 0xff);
-#if VULKAN_ENABLE_BINDING_DEBUG_NAMES
-		Version = (Version << 1) + Version;
-#endif
+		uint32 Version = HashCombine(GetTypeHash(HLSLCC_VersionMajor), GetTypeHash(HLSLCC_VersionMinor));
+		Version = HashCombine(Version, GetTypeHash(InternalGetVersion(Format)));
+		Version = HashCombine(Version, GetTypeHash(ShaderConductorVersionHash));
+
+	#if VULKAN_ENABLE_BINDING_DEBUG_NAMES
+		Version = HashCombine(Version, 0xFC0848E2);
+	#endif
+
 		return Version;
 	}
 	virtual void GetSupportedFormats(TArray<FName>& OutFormats) const
@@ -140,7 +151,7 @@ public:
 	{
 		if (!Singleton)
 		{
-			Singleton = new FShaderFormatVulkan();
+			Singleton = new FShaderFormatVulkan(FShaderConductorModuleWrapper::ModuleVersionHash);
 		}
 
 		return Singleton;
diff --git a/Engine/Source/Developer/Windows/ShaderFormatD3D/Private/ShaderFormatD3D.cpp b/Engine/Source/Developer/Windows/ShaderFormatD3D/Private/ShaderFormatD3D.cpp
--- a/Engine/Source/Developer/Windows/ShaderFormatD3D/Private/ShaderFormatD3D.cpp
+++ b/Engine/Source/Developer/Windows/ShaderFormatD3D/Private/ShaderFormatD3D.cpp
@@ -5,9 +5,7 @@
 #include "Modules/ModuleManager.h"
 #include "Interfaces/IShaderFormat.h"
 #include "Interfaces/IShaderFormatModule.h"
-#include "DxcWrapper.h"
-#include "HAL/PlatformProcess.h"
-#include "Misc/Paths.h"
+#include "DXCWrapper.h"
 
 static FName NAME_PCD3D_SM6(TEXT("PCD3D_SM6"));
 static FName NAME_PCD3D_SM5(TEXT("PCD3D_SM5"));
@@ -28,20 +26,31 @@ class FShaderFormatD3D : public IShaderFormat
 		check(Format == NAME_PCD3D_SM6 || Format == NAME_PCD3D_SM5 || Format == NAME_PCD3D_ES3_1);
 	}
 
+	uint32 DxcVersionHash = 0;
+
 public:
+
+	FShaderFormatD3D(uint32 InDxcVersionHash)
+		: DxcVersionHash(InDxcVersionHash)
+	{
+	}
+
 	virtual uint32 GetVersion(FName Format) const override
 	{
 		CheckFormat(Format);
 		if (Format == NAME_PCD3D_SM6)
 		{
-			return UE_SHADER_PCD3D_SM6_VER;
+			return HashCombine(DxcVersionHash, GetTypeHash(UE_SHADER_PCD3D_SM6_VER));
 		}
 		else if (Format == NAME_PCD3D_SM5)
 		{
-			return UE_SHADER_PCD3D_SM5_VER;
+			// Technically not needed for regular SM5 commpiled with legacy compiler,
+			// but PCD3D_SM5 currently includes ray tracing shaders that are compiled with new compiler stack.
+			return HashCombine(DxcVersionHash, GetTypeHash(UE_SHADER_PCD3D_SM5_VER));
 		}
 		else if (Format == NAME_PCD3D_ES3_1) 
 		{
+			// Shader DXC signature is intentionally not included, as ES3_1 target always uses legacy compiler.
 			return UE_SHADER_PCD3D_ES3_1_VER;
 		}
 		checkf(0, TEXT("Unknown Format %s"), *Format.ToString());
@@ -89,55 +98,9 @@ static IShaderFormat* Singleton = nullptr;
 
 class FShaderFormatD3DModule : public IShaderFormatModule, public FDxcModuleWrapper
 {
-	void* DllHandleSC = nullptr;
-	void* DllHandleDXC = nullptr;
-	void* DllHandleDXIL = nullptr;
-
 public:
-	FShaderFormatD3DModule()
-	{
-		FString DllPath;
-		{
-#if PLATFORM_WINDOWS
-			DllPath = FPaths::EngineDir() / TEXT("Binaries/ThirdParty/ShaderConductor/Win64/ShaderConductor.dll");
-#endif
-			checkf(DllPath.Len() > 0, TEXT("Unable to load %s"), *DllPath);
-			DllHandleSC = FPlatformProcess::GetDllHandle(*DllPath);
-		}
-		{
-#if PLATFORM_WINDOWS
-			DllPath = FPaths::EngineDir() / TEXT("Binaries/ThirdParty/ShaderConductor/Win64/dxcompiler.dll");
-#endif
-			checkf(DllPath.Len() > 0, TEXT("Unable to load %s"), *DllPath);
-			DllHandleDXC = FPlatformProcess::GetDllHandle(*DllPath);
-		}
-		{
-#if PLATFORM_WINDOWS
-			DllPath = FPaths::EngineDir() / TEXT("Binaries/ThirdParty/Windows/DirectX/x64/dxil.dll");
-#endif
-			checkf(DllPath.Len() > 0, TEXT("Unable to load %s"), *DllPath);
-			DllHandleDXIL = FPlatformProcess::GetDllHandle(*DllPath);
-		}
-	}
-
 	virtual ~FShaderFormatD3DModule()
 	{
-		if (DllHandleDXC)
-		{
-			FPlatformProcess::FreeDllHandle(DllHandleDXC);
-			DllHandleDXC = nullptr;
-		}
-		if (DllHandleDXIL)
-		{
-			FPlatformProcess::FreeDllHandle(DllHandleDXIL);
-			DllHandleDXIL = nullptr;
-		}
-		if (DllHandleSC)
-		{
-			FPlatformProcess::FreeDllHandle(DllHandleSC);
-			DllHandleSC = nullptr;
-		}
-
 		delete Singleton;
 		Singleton = nullptr;
 	}
@@ -146,7 +109,7 @@ public:
 	{
 		if (!Singleton)
 		{
-			Singleton = new FShaderFormatD3D();
+			Singleton = new FShaderFormatD3D(FDxcModuleWrapper::ModuleVersionHash);
 		}
 		return Singleton;
 	}
