Use custom built dxc.exe for D3D shader compilation batch scripts.

#rnx

diff --git a/Engine/Source/Developer/Windows/ShaderFormatD3D/Private/D3DShaderCompilerDXC.cpp b/Engine/Source/Developer/Windows/ShaderFormatD3D/Private/D3DShaderCompilerDXC.cpp
--- a/Engine/Source/Developer/Windows/ShaderFormatD3D/Private/D3DShaderCompilerDXC.cpp
+++ b/Engine/Source/Developer/Windows/ShaderFormatD3D/Private/D3DShaderCompilerDXC.cpp
@@ -370,18 +370,33 @@ static HRESULT D3DCompileToDxil(const char* SourceText, FDxcArguments& Arguments
 
 static FString D3DCreateDXCCompileBatchFile(const FDxcArguments& Args, const FString& ShaderPath)
 {
-	FString BatchFileHeader = TEXT("@ECHO OFF\nSET DXC=\"C:\\Program Files (x86)\\Windows Kits\\10\\bin\\10.0.18362.0\\x64\\dxc.exe\"\n"\
-		"IF EXIST %DXC% (\nREM\n) ELSE (\nECHO Couldn't find Windows 10.0.18362 SDK, falling back to dxc.exe in PATH...\n"\
-		"SET DXC=dxc.exe)\n");
-
-	FString DXCCommandline = FString(TEXT("%DXC%"));
-
-	DXCCommandline += Args.GetBatchCommandLineString(ShaderPath);
-
-	DXCCommandline += TEXT(" ");
-	DXCCommandline += ShaderPath;
-
-	return BatchFileHeader + DXCCommandline + TEXT("\npause\n");
+	FString DxcPath = FPaths::ConvertRelativePathToFull(FPaths::EngineDir());
+
+	DxcPath = FPaths::Combine(DxcPath, TEXT("Binaries/ThirdParty/ShaderConductor/Win64"));
+	FPaths::MakePlatformFilename(DxcPath);
+
+	FString DxcFilename = FPaths::Combine(DxcPath, TEXT("dxc.exe"));
+	FPaths::MakePlatformFilename(DxcFilename);
+
+	const FString BatchCmdLineArgs = Args.GetBatchCommandLineString(ShaderPath);
+
+	return FString::Printf(
+		TEXT(
+			"@ECHO OFF\n"
+			"SET DXC=\"%s\"\n"
+			"IF NOT EXIST %%DXC%% (\n"
+			"\tECHO Couldn't find dxc.exe under \"%s\"\n"
+			"\tGOTO :END\n"
+			")\n"
+			"%%DXC%%%s %s\n"
+			":END\n"
+			"PAUSE\n"
+		),
+		*DxcFilename,
+		*DxcPath,
+		*BatchCmdLineArgs,
+		*ShaderPath
+	);
 }
 
 inline bool IsCompatibleBinding(const D3D12_SHADER_INPUT_BIND_DESC& BindDesc, uint32 BindingSpace)
