Moved dxil.dll from DirectX folder into ShaderConductor binary folder and load dxil.dll manually to ensure it's always loaded from that location.

#rnx

diff --git a/Engine/Shaders/Public/Platform/D3D/D3DCommon.ush b/Engine/Shaders/Public/Platform/D3D/D3DCommon.ush
--- a/Engine/Shaders/Public/Platform/D3D/D3DCommon.ush
+++ b/Engine/Shaders/Public/Platform/D3D/D3DCommon.ush
@@ -5,7 +5,7 @@ D3DCommon.usf: Common defines for the d3d shader compiler
 =============================================================================*/
 
 // Update this GUID to invalidate D3D shaders and force a shader recompilation for D3D RHIs.
-// GUID: 420707B5-DB86-4011-97F5-AEAAF78928DD
+// GUID: B8728C4B-C1E8-45AD-B3E4-A6761FAFA247
 
 precise float MakePrecise(precise float v) { return v; }
 precise float2 MakePrecise(precise float2 v) { return v; }
diff --git a/Engine/Source/Developer/ShaderCompilerCommon/Private/DXCWrapper.cpp b/Engine/Source/Developer/ShaderCompilerCommon/Private/DXCWrapper.cpp
--- a/Engine/Source/Developer/ShaderCompilerCommon/Private/DXCWrapper.cpp
+++ b/Engine/Source/Developer/ShaderCompilerCommon/Private/DXCWrapper.cpp
@@ -10,6 +10,7 @@
 #include "Windows/WindowsHWrapper.h"
 #endif // PLATFORM_WINDOWS
 
+static TRefCountPtr<FDllHandle> GDxilHandle;
 static TRefCountPtr<FDllHandle> GDxcHandle;
 static TRefCountPtr<FDllHandle> GShaderConductorHandle;
 
@@ -17,7 +18,7 @@ static uint64 GetLoadedModuleVersion(const TCHAR* ModuleName)
 {
 #if PLATFORM_WINDOWS
 	HMODULE ModuleDll = ::GetModuleHandleW(ModuleName);
-	if (ModuleDll == INVALID_HANDLE_VALUE)
+	if (ModuleDll == nullptr)
 	{
 		return 0;
 	}
@@ -74,10 +75,12 @@ FDxcModuleWrapper::FDxcModuleWrapper()
 {
 	if (GDxcHandle.GetRefCount() == 0)
 	{
+		GDxilHandle = new FDllHandle(TEXT("dxil.dll"));
 		GDxcHandle = new FDllHandle(TEXT("dxcompiler.dll"));
 	}
 	else
 	{
+		GDxilHandle->AddRef();
 		GDxcHandle->AddRef();
 	}
 
@@ -94,6 +97,7 @@ FDxcModuleWrapper::FDxcModuleWrapper()
 FDxcModuleWrapper::~FDxcModuleWrapper()
 {
 	GDxcHandle.SafeRelease();
+	GDxilHandle.SafeRelease();
 }
 
 
@@ -110,7 +114,7 @@ FShaderConductorModuleWrapper::FShaderConductorModuleWrapper()
 
 	static uint64 DllVersion = GetLoadedModuleVersion(TEXT("ShaderConductor.dll"));
 
-	ModuleVersionHash = HashCombine(GetTypeHash(DllVersion), FDxcModuleWrapper::ModuleVersionHash);
+	ModuleVersionHash = HashCombine(GetTypeHash(DllVersion), FDxcModuleWrapper::GetModuleVersionHash());
 }
 
 FShaderConductorModuleWrapper::~FShaderConductorModuleWrapper()
diff --git a/Engine/Source/Developer/ShaderCompilerCommon/Public/DXCWrapper.h b/Engine/Source/Developer/ShaderCompilerCommon/Public/DXCWrapper.h
--- a/Engine/Source/Developer/ShaderCompilerCommon/Public/DXCWrapper.h
+++ b/Engine/Source/Developer/ShaderCompilerCommon/Public/DXCWrapper.h
@@ -7,7 +7,7 @@
 
 struct FDllHandle : public FRefCountedObject
 {
-protected:
+private:
 	void* Handle = nullptr;
 
 public:
@@ -17,19 +17,27 @@ public:
 
 class SHADERCOMPILERCOMMON_API FDxcModuleWrapper
 {
-protected:
-	TRefCountPtr<FDllHandle> Dxc;
+private:
 	uint32 ModuleVersionHash = 0;
+protected:
+	FORCEINLINE uint32 GetModuleVersionHash() const
+	{
+		return ModuleVersionHash;
+	}
 public:
 	FDxcModuleWrapper();
 	virtual ~FDxcModuleWrapper();
 };
 
-class SHADERCOMPILERCOMMON_API FShaderConductorModuleWrapper : public FDxcModuleWrapper
+class SHADERCOMPILERCOMMON_API FShaderConductorModuleWrapper : private FDxcModuleWrapper
 {
-protected:
-	TRefCountPtr<FDllHandle> ShaderConductor;
+private:
 	uint32 ModuleVersionHash = 0;
+protected:
+	FORCEINLINE uint32 GetModuleVersionHash() const
+	{
+		return ModuleVersionHash;
+	}
 public:
 	FShaderConductorModuleWrapper();
 	virtual ~FShaderConductorModuleWrapper();
diff --git a/Engine/Source/Developer/VulkanShaderFormat/Private/VulkanShaderFormat.cpp b/Engine/Source/Developer/VulkanShaderFormat/Private/VulkanShaderFormat.cpp
--- a/Engine/Source/Developer/VulkanShaderFormat/Private/VulkanShaderFormat.cpp
+++ b/Engine/Source/Developer/VulkanShaderFormat/Private/VulkanShaderFormat.cpp
@@ -151,7 +151,7 @@ public:
 	{
 		if (!Singleton)
 		{
-			Singleton = new FShaderFormatVulkan(FShaderConductorModuleWrapper::ModuleVersionHash);
+			Singleton = new FShaderFormatVulkan(FShaderConductorModuleWrapper::GetModuleVersionHash());
 		}
 
 		return Singleton;
diff --git a/Engine/Source/Developer/Windows/ShaderFormatD3D/Private/ShaderFormatD3D.cpp b/Engine/Source/Developer/Windows/ShaderFormatD3D/Private/ShaderFormatD3D.cpp
--- a/Engine/Source/Developer/Windows/ShaderFormatD3D/Private/ShaderFormatD3D.cpp
+++ b/Engine/Source/Developer/Windows/ShaderFormatD3D/Private/ShaderFormatD3D.cpp
@@ -109,7 +109,7 @@ public:
 	{
 		if (!Singleton)
 		{
-			Singleton = new FShaderFormatD3D(FDxcModuleWrapper::ModuleVersionHash);
+			Singleton = new FShaderFormatD3D(FDxcModuleWrapper::GetModuleVersionHash());
 		}
 		return Singleton;
 	}
diff --git a/Engine/Source/Developer/Windows/ShaderFormatD3D/ShaderFormatD3D.Build.cs b/Engine/Source/Developer/Windows/ShaderFormatD3D/ShaderFormatD3D.Build.cs
--- a/Engine/Source/Developer/Windows/ShaderFormatD3D/ShaderFormatD3D.Build.cs
+++ b/Engine/Source/Developer/Windows/ShaderFormatD3D/ShaderFormatD3D.Build.cs
@@ -1,6 +1,7 @@
 // Copyright Epic Games, Inc. All Rights Reserved.
 
 using UnrealBuildTool;
+using System.IO;
 
 public class ShaderFormatD3D : ModuleRules
 {
@@ -27,14 +28,13 @@ public class ShaderFormatD3D : ModuleRules
 			string AmdAgsPath = Target.UEThirdPartySourceDirectory + "AMD/AMD_AGS/";
 			PublicSystemIncludePaths.Add(AmdAgsPath + "inc/");  // For amd_ags.h, to get AGS_DX12_SHADER_INSTRINSICS_SPACE_ID
 
-            string DxDllsPath = "$(EngineDir)/Binaries/ThirdParty/Windows/DirectX/x64/";
+            string DxDllsPath = "$(EngineDir)/Binaries/ThirdParty/Windows/DirectX/x64";
+			RuntimeDependencies.Add(Path.Combine(DxDllsPath, "d3dcompiler_47.dll"));
 
-			RuntimeDependencies.Add("$(TargetOutputDir)/d3dcompiler_47.dll", DxDllsPath + "d3dcompiler_47.dll");
-			RuntimeDependencies.Add("$(TargetOutputDir)/dxil.dll", DxDllsPath + "dxil.dll");
-
-			string BinaryFolder = Target.UEThirdPartyBinariesDirectory + "ShaderConductor/Win64";
-			RuntimeDependencies.Add(BinaryFolder + "/dxcompiler.dll");
-			RuntimeDependencies.Add(BinaryFolder + "/ShaderConductor.dll");
+			string ShaderConductorDllsPath = Path.Combine(Target.UEThirdPartyBinariesDirectory, "ShaderConductor/Win64");
+			RuntimeDependencies.Add(Path.Combine(ShaderConductorDllsPath, "dxcompiler.dll"));
+			RuntimeDependencies.Add(Path.Combine(ShaderConductorDllsPath, "ShaderConductor.dll"));
+			RuntimeDependencies.Add(Path.Combine(ShaderConductorDllsPath, "dxil.dll"));
 
 			AddEngineThirdPartyPrivateStaticDependencies(Target,
 				"ShaderConductor"
diff --git a/Engine/Source/Developer/Windows/WindowsTargetPlatform/Private/GenericWindowsTargetPlatform.h b/Engine/Source/Developer/Windows/WindowsTargetPlatform/Private/GenericWindowsTargetPlatform.h
--- a/Engine/Source/Developer/Windows/WindowsTargetPlatform/Private/GenericWindowsTargetPlatform.h
+++ b/Engine/Source/Developer/Windows/WindowsTargetPlatform/Private/GenericWindowsTargetPlatform.h
@@ -146,12 +146,9 @@ public:
 	virtual void GetShaderCompilerDependencies(TArray<FString>& OutDependencies) const override
 	{		
 		FTargetPlatformBase::AddDependencySCArrayHelper(OutDependencies, TEXT("Binaries/ThirdParty/Windows/DirectX/x64/d3dcompiler_47.dll"));
-		FTargetPlatformBase::AddDependencySCArrayHelper(OutDependencies, TEXT("Binaries/ThirdParty/Windows/DirectX/x64/dxil.dll"));
 		FTargetPlatformBase::AddDependencySCArrayHelper(OutDependencies, TEXT("Binaries/ThirdParty/ShaderConductor/Win64/ShaderConductor.dll"));
 		FTargetPlatformBase::AddDependencySCArrayHelper(OutDependencies, TEXT("Binaries/ThirdParty/ShaderConductor/Win64/dxcompiler.dll"));
-		FTargetPlatformBase::AddDependencySCArrayHelper(OutDependencies, TEXT("Binaries/Win64/d3dcompiler_47.dll"));
-		FTargetPlatformBase::AddDependencySCArrayHelper(OutDependencies, TEXT("Binaries/Win64/dxcompiler.dll"));
-		FTargetPlatformBase::AddDependencySCArrayHelper(OutDependencies, TEXT("Binaries/Win64/dxil.dll"));
+		FTargetPlatformBase::AddDependencySCArrayHelper(OutDependencies, TEXT("Binaries/ThirdParty/ShaderConductor/Win64/dxil.dll"));
 	}
 
 	virtual bool SupportsFeature( ETargetPlatformFeatures Feature ) const override
