More verbose error message on DXC compile failure. Helps to differentiate shader compile errors from compiler crashes.

diff --git a/Engine/Source/Developer/Windows/ShaderFormatD3D/Private/D3DShaderCompilerDXC.cpp b/Engine/Source/Developer/Windows/ShaderFormatD3D/Private/D3DShaderCompilerDXC.cpp
--- a/Engine/Source/Developer/Windows/ShaderFormatD3D/Private/D3DShaderCompilerDXC.cpp
+++ b/Engine/Source/Developer/Windows/ShaderFormatD3D/Private/D3DShaderCompilerDXC.cpp
@@ -891,7 +891,13 @@ bool CompileAndProcessD3DShaderDXC(FString& PreprocessedShaderSource,
 
 	if (FAILED(Result))
 	{
-		FilteredErrors.Add(TEXT("D3DCompileToDxil failed"));
+		TCHAR ErrorMsg[1024];
+		FPlatformMisc::GetSystemErrorMessage(ErrorMsg, UE_ARRAY_COUNT(ErrorMsg), (int)Result);
+		const bool bKnownError = ErrorMsg[0] != TEXT('\0');
+
+		FString ErrorString = FString::Printf(TEXT("D3DCompileToDxil failed. Error code: %s (0x%08X)."), bKnownError ? ErrorMsg : TEXT("Unknown error"), (int)Result);
+
+		FilteredErrors.Add(ErrorString);
 	}
 
 	return SUCCEEDED(Result);
