LiveCoding fails to load a patch after "Quick Restart" is used.

diff --git a/Engine/Source/Developer/Windows/LiveCodingServer/Private/External/LC_LiveModule.cpp b/Engine/Source/Developer/Windows/LiveCodingServer/Private/External/LC_LiveModule.cpp
--- a/Engine/Source/Developer/Windows/LiveCodingServer/Private/External/LC_LiveModule.cpp
+++ b/Engine/Source/Developer/Windows/LiveCodingServer/Private/External/LC_LiveModule.cpp
@@ -4480,6 +4480,13 @@ bool LiveModule::InstallCompiledPatches(LiveProcess* liveProcess, void* original
 	const Process::Id processId = liveProcess->GetProcessId();
 	const DuplexPipe* pipe = liveProcess->GetPipe();
 
+// BEGIN EPIC MOD
+#if LC_64_BIT
+	// free our page reservations in the +-2GB range of the main module
+	liveProcess->FreeVirtualMemoryPages(originalModuleBase);
+#endif
+// END EPIC MOD
+
 	for (auto modulePatch : m_compiledModulePatches)
 	{
 		const std::wstring& originalExePath = modulePatch->GetExePath();
@@ -4674,6 +4681,13 @@ bool LiveModule::InstallCompiledPatches(LiveProcess* liveProcess, void* original
 		pipe->SendCommandAndWaitForAck(commands::LeaveSyncPoint {}, nullptr, 0u);
 	}
 
+// BEGIN EPIC MOD
+#if LC_64_BIT
+	// immediately reserve pages in the +-2GB range of the main module again
+	liveProcess->ReserveVirtualMemoryPages(originalModuleBase);
+#endif
+// END EPIC MOD
+
 	LC_SUCCESS_USER("Successfully installed %zu patches (%.3fs)", m_compiledModulePatches.size(), wholeScope.ReadSeconds());
 
 	return true;
