Differentiate between engine custom shader validation errors and real D3DCompileToDxil failures.

diff --git a/Engine/Source/Developer/Windows/ShaderFormatD3D/Private/D3DShaderCompilerDXC.cpp b/Engine/Source/Developer/Windows/ShaderFormatD3D/Private/D3DShaderCompilerDXC.cpp
--- a/Engine/Source/Developer/Windows/ShaderFormatD3D/Private/D3DShaderCompilerDXC.cpp
+++ b/Engine/Source/Developer/Windows/ShaderFormatD3D/Private/D3DShaderCompilerDXC.cpp
@@ -683,7 +683,8 @@ bool CompileAndProcessD3DShaderDXC(FString& PreprocessedShaderSource,
 	TRefCountPtr<IDxcBlob> ShaderBlob;
 	TRefCountPtr<IDxcBlob> ReflectionBlob;
 	TRefCountPtr<IDxcBlobEncoding> DxcErrorBlob;
-	HRESULT Result = D3DCompileToDxil(AnsiSourceFile.Get(), Args, ShaderBlob, ReflectionBlob, DxcErrorBlob);
+
+	const HRESULT D3DCompileToDxilResult = D3DCompileToDxil(AnsiSourceFile.Get(), Args, ShaderBlob, ReflectionBlob, DxcErrorBlob);
 
 	if (DxcErrorBlob && DxcErrorBlob->GetBufferSize())
 	{
@@ -691,7 +692,7 @@ bool CompileAndProcessD3DShaderDXC(FString& PreprocessedShaderSource,
 		DXCFilterShaderCompileWarnings(ErrorString, FilteredErrors);
 	}
 
-	if (SUCCEEDED(Result))
+	if (SUCCEEDED(D3DCompileToDxilResult))
 	{
 		// Gather reflection information
 		int32 NumInterpolants = 0;
@@ -721,11 +722,11 @@ bool CompileAndProcessD3DShaderDXC(FString& PreprocessedShaderSource,
 		if (bIsRayTracingShader)
 		{
 			TRefCountPtr<ID3D12LibraryReflection> LibraryReflection;
-			Result = Utils->CreateReflection(&ReflBuffer, IID_PPV_ARGS(LibraryReflection.GetInitReference()));
+			const HRESULT CreateReflectionResult = Utils->CreateReflection(&ReflBuffer, IID_PPV_ARGS(LibraryReflection.GetInitReference()));
 
-			if (FAILED(Result))
+			if (FAILED(CreateReflectionResult))
 			{
-				UE_LOG(LogD3D12ShaderCompiler, Fatal, TEXT("D3DReflectDxil failed: Result=%08x"), Result);
+				UE_LOG(LogD3D12ShaderCompiler, Fatal, TEXT("CreateReflection failed: Result=%08x"), CreateReflectionResult);
 			}
 
 			D3D12_LIBRARY_DESC LibraryDesc = {};
@@ -811,7 +812,6 @@ bool CompileAndProcessD3DShaderDXC(FString& PreprocessedShaderSource,
 					}
 
 					FilteredErrors.Add(ErrorString);
-					Result = E_FAIL;
 					Output.bSucceeded = false;
 				}
 			}
@@ -825,10 +825,10 @@ bool CompileAndProcessD3DShaderDXC(FString& PreprocessedShaderSource,
 		{
 
 			TRefCountPtr<ID3D12ShaderReflection> ShaderReflection;
-			Result = Utils->CreateReflection(&ReflBuffer, IID_PPV_ARGS(ShaderReflection.GetInitReference()));
-			if (FAILED(Result))
+			const HRESULT CreateReflectionResult = Utils->CreateReflection(&ReflBuffer, IID_PPV_ARGS(ShaderReflection.GetInitReference()));
+			if (FAILED(CreateReflectionResult))
 			{
-				UE_LOG(LogD3D12ShaderCompiler, Fatal, TEXT("D3DReflectDxil failed: Result=%08x"), Result);
+				UE_LOG(LogD3D12ShaderCompiler, Fatal, TEXT("CreateReflection failed: Result=%08x"), CreateReflectionResult);
 			}
 
 			D3D12_SHADER_DESC ShaderDesc = {};
@@ -846,7 +846,6 @@ bool CompileAndProcessD3DShaderDXC(FString& PreprocessedShaderSource,
 
 		if (!ValidateResourceCounts(NumSRVs, NumSamplers, NumUAVs, NumCBs, FilteredErrors))
 		{
-			Result = E_FAIL;
 			Output.bSucceeded = false;
 		}
 
@@ -888,19 +887,18 @@ bool CompileAndProcessD3DShaderDXC(FString& PreprocessedShaderSource,
 				AddOptionalDataCallback);
 		}
 	}
-
-	if (FAILED(Result))
+	else
 	{
 		TCHAR ErrorMsg[1024];
-		FPlatformMisc::GetSystemErrorMessage(ErrorMsg, UE_ARRAY_COUNT(ErrorMsg), (int)Result);
+		FPlatformMisc::GetSystemErrorMessage(ErrorMsg, UE_ARRAY_COUNT(ErrorMsg), (int)D3DCompileToDxilResult);
 		const bool bKnownError = ErrorMsg[0] != TEXT('\0');
 
-		FString ErrorString = FString::Printf(TEXT("D3DCompileToDxil failed. Error code: %s (0x%08X)."), bKnownError ? ErrorMsg : TEXT("Unknown error"), (int)Result);
+		FString ErrorString = FString::Printf(TEXT("D3DCompileToDxil failed. Error code: %s (0x%08X)."), bKnownError ? ErrorMsg : TEXT("Unknown error"), (int)D3DCompileToDxilResult);
 
 		FilteredErrors.Add(ErrorString);
 	}
 
-	return SUCCEEDED(Result);
+	return Output.bSucceeded;
 }
 
 #undef VERIFYHRESULT
