Fix improper extern declaration that caused issues with Clang but not MSVC (Clang was right IMHO)

diff --git a/Engine/Source/Developer/Windows/LiveCoding/Private/External/LC_ClientCommandActions.cpp b/Engine/Source/Developer/Windows/LiveCoding/Private/External/LC_ClientCommandActions.cpp
--- a/Engine/Source/Developer/Windows/LiveCoding/Private/External/LC_ClientCommandActions.cpp
+++ b/Engine/Source/Developer/Windows/LiveCoding/Private/External/LC_ClientCommandActions.cpp
@@ -116,6 +116,13 @@ class FChunkedFixedUObjectArray;
 extern FChunkedFixedUObjectArray*& GObjectArrayForDebugVisualizers;
 // END EPIC MOD
 
+// BEGIN EPIC MOD - Support for object reinstancing. We need to call a global post-patch handler, rather than just getting individual callbacks for modified modules.
+extern void LiveCodingBeginPatch();
+// END EPIC MOD
+
+// BEGIN EPIC MOD - Notification that compilation has finished
+extern void LiveCodingEndCompile();
+// END EPIC MOD
 
 bool actions::LoadPatch::Execute(const CommandType* command, const DuplexPipe* pipe, void*, const void*, size_t)
 {
@@ -135,7 +142,6 @@ bool actions::LoadPatch::Execute(const CommandType* command, const DuplexPipe* p
 	}
 	// END EPIC MOD
 	// BEGIN EPIC MOD - Support for object reinstancing. We need to call a global post-patch handler, rather than just getting individual callbacks for modified modules.
-	extern void LiveCodingBeginPatch();
 	LiveCodingBeginPatch();
 	// END EPIC MOD
 
@@ -182,7 +188,6 @@ bool actions::CompilationFinished::Execute(const CommandType*, const DuplexPipe*
 	pipe->SendAck();
 
 	// BEGIN EPIC MOD - Notification that compilation has finished
-	extern void LiveCodingEndCompile();
 	LiveCodingEndCompile();
 	// END EPIC MOD
 
