Workaround for a bug in clang-cl 11.0 when building in C++17 mode, where __builtin_wmemchr is not implemented.

#rnx

diff --git a/Engine/Source/Developer/Windows/LiveCoding/Private/External/LC_StringUtil.cpp b/Engine/Source/Developer/Windows/LiveCoding/Private/External/LC_StringUtil.cpp
--- a/Engine/Source/Developer/Windows/LiveCoding/Private/External/LC_StringUtil.cpp
+++ b/Engine/Source/Developer/Windows/LiveCoding/Private/External/LC_StringUtil.cpp
@@ -3,6 +3,25 @@
 // BEGIN EPIC MOD
 //#include PCH_INCLUDE
 // END EPIC MOD
+
+// BEGIN EPIC MOD
+#if defined(__clang__) && defined(_MSC_VER) && _MSVC_LANG > 201402L
+#if __clang_major__ > 11
+// For reference: https://bugs.llvm.org/show_bug.cgi?id=41226#c16
+#error Re-check whether this workaround is still needed for clang 12.0 and higher.
+#endif
+inline const wchar_t* wmemchr(const wchar_t* Str, wchar_t Ch, size_t Len)
+{
+	for (; 0 < Len; Str++, Len--)
+		if (*Str == Ch)
+			return Str;
+
+	return nullptr;
+}
+#endif
+// END EPIC MOD
+
+
 #include "LC_StringUtil.h"
 // BEGIN EPIC MOD
 #include "Windows/WindowsHWrapper.h"
