Adding the ability to disable re-instancing when doing live coding.

Will be disabled until blueprint issues are resolved.

diff --git a/Engine/Config/BaseEngine.ini b/Engine/Config/BaseEngine.ini
--- a/Engine/Config/BaseEngine.ini
+++ b/Engine/Config/BaseEngine.ini
@@ -2613,6 +2613,7 @@ DefaultSizeInMegabyte=64
 UnrealBuildTool=Engine/Binaries/DotNET/UnrealBuildTool/UnrealBuildTool.exe
 
 [LiveCoding]
+bEnableReinstancing=false
 bDisableActionLimit=false
 ActionLimit=100
 
diff --git a/Engine/Source/Developer/Windows/LiveCoding/Private/LiveCodingModule.cpp b/Engine/Source/Developer/Windows/LiveCoding/Private/LiveCodingModule.cpp
--- a/Engine/Source/Developer/Windows/LiveCoding/Private/LiveCodingModule.cpp
+++ b/Engine/Source/Developer/Windows/LiveCoding/Private/LiveCodingModule.cpp
@@ -84,7 +84,6 @@ public:
 
 	virtual void NotifyChange(UClass* New, UClass* Old) override
 	{
-		bGotClassChange |= Old != nullptr && Old != New;
 	}
 
 	virtual void NotifyChange(UEnum* New, UEnum* Old) override
@@ -93,13 +92,13 @@ public:
 
 	virtual void NotifyChange(UScriptStruct* New, UScriptStruct* Old) override
 	{
-		bGotClassChange |= Old != nullptr && Old != New;
 	}
 
-	virtual void Reinstance()
+	virtual bool GetEnableReinstancing(bool bHasChanged) const
 	{
-		if (bGotClassChange)
+		if (bHasChanged && !bEnabledMessage)
 		{
+			bEnabledMessage = true;
 			static const TCHAR* Message = TEXT("Object structure changes detected.  LiveCoding re-instancing isn't supported in builds without the editor");
 			UE_LOG(LogLiveCoding, Error, TEXT("%s"), Message);
 #if WITH_ENGINE
@@ -109,11 +108,16 @@ public:
 			}
 #endif
 		}
+		return false;
+	}
+
+	virtual void Reinstance()
+	{
 	}
 
 private:
 	FLiveCodingModule& LiveCodingModule;
-	bool bGotClassChange = false;
+	mutable bool bEnabledMessage = false;
 };
 #endif
 
@@ -616,7 +620,10 @@ void FLiveCodingModule::BeginReload()
 		if (!GLiveCodingModule->Reload.IsValid())
 		{
 #if WITH_EDITOR
+			bool bEnableReinstancing = true;
+			GConfig->GetBool(TEXT("LiveCoding"), TEXT("bEnableReinstancing"), bEnableReinstancing, GEngineIni);
 			GLiveCodingModule->Reload.Reset(new FReload(EActiveReloadType::LiveCoding, TEXT("LIVECODING"), *GLog));
+			GLiveCodingModule->Reload->SetEnableReinstancing(bEnableReinstancing);
 #else
 			GLiveCodingModule->Reload.Reset(new FNullReload(*GLiveCodingModule));
 #endif
diff --git a/Engine/Source/Editor/UnrealEd/Private/Kismet2/ReloadUtilities.cpp b/Engine/Source/Editor/UnrealEd/Private/Kismet2/ReloadUtilities.cpp
--- a/Engine/Source/Editor/UnrealEd/Private/Kismet2/ReloadUtilities.cpp
+++ b/Engine/Source/Editor/UnrealEd/Private/Kismet2/ReloadUtilities.cpp
@@ -765,6 +765,17 @@ FReload::~FReload()
 	}
 }
 
+bool FReload::GetEnableReinstancing(bool bHasChanged) const  
+{ 
+	if (bHasChanged && !bEnabledMessage)
+	{
+		bEnabledMessage = true;
+		Ar.Logf(ELogVerbosity::Display, TEXT("Re-instancing has been disabled.  Some changes will be ignored."));
+	}
+	return bEnableReinstancing; 
+}
+
+
 void FReload::Reset()
 {
 	FunctionRemap.Empty();
diff --git a/Engine/Source/Editor/UnrealEd/Public/Kismet2/ReloadUtilities.h b/Engine/Source/Editor/UnrealEd/Public/Kismet2/ReloadUtilities.h
--- a/Engine/Source/Editor/UnrealEd/Public/Kismet2/ReloadUtilities.h
+++ b/Engine/Source/Editor/UnrealEd/Public/Kismet2/ReloadUtilities.h
@@ -34,6 +34,7 @@ public:
 	// IReload interface 
 	virtual EActiveReloadType GetType() const override { return Type; }
 	virtual const TCHAR* GetPrefix() const override { return Prefix; };
+	virtual bool GetEnableReinstancing(bool bHasChanged) const override;
 	virtual void NotifyFunctionRemap(FNativeFuncPtr NewFunctionPointer, FNativeFuncPtr OldFunctionPointer) override;
 	virtual void NotifyChange(UClass* New, UClass* Old) override;
 	virtual void NotifyChange(UEnum* New, UEnum* Old) override;
@@ -58,6 +59,14 @@ public:
 		bSendReloadComplete = bSend;
 	}
 
+	/**
+	 * Enable/Disable the support for reinstancing
+	 */
+	void SetEnableReinstancing(bool bInEnableReinstancing)
+	{
+		bEnableReinstancing = bInEnableReinstancing;
+	}
+
 private:
 
 	struct ReinstanceStats
@@ -130,9 +139,13 @@ private:
 	/** If true, send reload complete notification */
 	bool bSendReloadComplete = true;
 
+	/** If true, reinstancing is enabled */
+	bool bEnableReinstancing = true;
+
 	ReinstanceStats ClassStats;
 	ReinstanceStats EnumStats;
 	ReinstanceStats StructStats;
 	int32 NumFunctionsRemapped = 0;
 	int32 NumScriptStructsRemapped = 0;
+	mutable bool bEnabledMessage = false;
 };
diff --git a/Engine/Source/Runtime/CoreUObject/Private/UObject/Class.cpp b/Engine/Source/Runtime/CoreUObject/Private/UObject/Class.cpp
--- a/Engine/Source/Runtime/CoreUObject/Private/UObject/Class.cpp
+++ b/Engine/Source/Runtime/CoreUObject/Private/UObject/Class.cpp
@@ -2513,11 +2513,18 @@ void UScriptStruct::DeferCppStructOps(FName Target, ICppStructOps* InCppStructOp
 
 	if (UScriptStruct::ICppStructOps* ExistingOps = DeferredStructOps.FindRef(Target))
 	{
-		if (!IsReloadActive()) // in reload, we will just leak these...they may be in use.
+		IReload* Reload = GetActiveReloadInterface();
+		if (Reload == nullptr)
 		{
 			check(ExistingOps != InCppStructOps); // if it was equal, then we would be re-adding a now stale pointer to the map
 			delete ExistingOps;
 		}
+		else if (!Reload->GetEnableReinstancing(false))
+		{
+			delete InCppStructOps;
+			return;
+		}
+		// in reload, we will just leak these...they may be in use.
 	}
 	DeferredStructOps.Add(Target,InCppStructOps);
 }
diff --git a/Engine/Source/Runtime/CoreUObject/Private/UObject/DeferredRegistry.h b/Engine/Source/Runtime/CoreUObject/Private/UObject/DeferredRegistry.h
--- a/Engine/Source/Runtime/CoreUObject/Private/UObject/DeferredRegistry.h
+++ b/Engine/Source/Runtime/CoreUObject/Private/UObject/DeferredRegistry.h
@@ -85,19 +85,30 @@ public:
 		InInfo.ReloadVersionInfo = InVersion;
 		TType* OldSingleton = ExistingInfo ? (*ExistingInfo)->InnerSingleton : nullptr;
 
+		bool bAdd = true;
 		if (ExistingInfo)
 		{
-			if (!bHasChanged && *ExistingInfo != &InInfo)
+			if (IReload* Reload = GetActiveReloadInterface())
 			{
-				CopyExistingSingletons(InInfo, **ExistingInfo);
+				bAdd = Reload->GetEnableReinstancing(bHasChanged);
+			}
+			if (bAdd)
+			{
+				if (!bHasChanged && *ExistingInfo != &InInfo)
+				{
+					CopyExistingSingletons(InInfo, **ExistingInfo);
+				}
+				*ExistingInfo = &InInfo;
 			}
-			*ExistingInfo = &InInfo;
 		}
 		else
 		{
 			InfoMap.Add(Key, &InInfo);
 		}
-		Registrations.Add(FRegistrant{ InOuterRegister, InInnerRegister, InPackageName, &InInfo, OldSingleton, DeprecatedFieldInfo, bHasChanged });
+		if (bAdd)
+		{
+			Registrations.Add(FRegistrant{ InOuterRegister, InInnerRegister, InPackageName, &InInfo, OldSingleton, DeprecatedFieldInfo, bHasChanged });
+		}
 		return ExistingInfo != nullptr;
 #else
 		Registrations.Add(FRegistrant{ InOuterRegister, InInnerRegister, InPackageName, &InInfo, DeprecatedFieldInfo });
diff --git a/Engine/Source/Runtime/CoreUObject/Public/UObject/Reload.h b/Engine/Source/Runtime/CoreUObject/Public/UObject/Reload.h
--- a/Engine/Source/Runtime/CoreUObject/Public/UObject/Reload.h
+++ b/Engine/Source/Runtime/CoreUObject/Public/UObject/Reload.h
@@ -31,6 +31,10 @@ public:
 	 */
 	virtual const TCHAR* GetPrefix() const = 0;
 
+	/**
+	 * Return if re-instancing is to be allowed
+	 */
+	virtual bool GetEnableReinstancing(bool bHasChanged) const = 0;
 
 	/**
 	 * Invoke when a duplicate function has been detected.
