Removing UE4 from code

diff --git a/Engine/Source/Developer/Windows/LiveCoding/Private/External/LC_ClientCommandActions.cpp b/Engine/Source/Developer/Windows/LiveCoding/Private/External/LC_ClientCommandActions.cpp
--- a/Engine/Source/Developer/Windows/LiveCoding/Private/External/LC_ClientCommandActions.cpp
+++ b/Engine/Source/Developer/Windows/LiveCoding/Private/External/LC_ClientCommandActions.cpp
@@ -109,7 +109,7 @@ bool actions::CallHooks::Execute(const CommandType* command, const DuplexPipe* p
 }
 
 
-// BEGIN EPIC MOD - Support for UE4 debug visualizers
+// BEGIN EPIC MOD - Support for UE debug visualizers
 extern uint8** GNameBlocksDebug;
 
 class FChunkedFixedUObjectArray;
@@ -129,7 +129,7 @@ bool actions::LoadPatch::Execute(const CommandType* command, const DuplexPipe* p
 	// load library into this process
 	HMODULE module = ::LoadLibraryW(command->path);
 
-	// BEGIN EPIC MOD - Support for UE4 debug visualizers
+	// BEGIN EPIC MOD - Support for UE debug visualizers
 	if (module != nullptr)
 	{
 		typedef void InitNatvisHelpersFunc(uint8** NameTable, FChunkedFixedUObjectArray* ObjectArray);
diff --git a/Engine/Source/Developer/Windows/LiveCoding/Private/External/LC_Restart.cpp b/Engine/Source/Developer/Windows/LiveCoding/Private/External/LC_Restart.cpp
--- a/Engine/Source/Developer/Windows/LiveCoding/Private/External/LC_Restart.cpp
+++ b/Engine/Source/Developer/Windows/LiveCoding/Private/External/LC_Restart.cpp
@@ -60,7 +60,7 @@ void restart::Execute(lpp::RestartBehaviour behaviour, unsigned int exitCode)
 
 	switch (behaviour)
 	{
-		// BEGIN EPIC MODS - Use UE4 codepath for termination to ensure logs are flushed and session analytics are sent
+		// BEGIN EPIC MODS - Use UE codepath for termination to ensure logs are flushed and session analytics are sent
 		case lpp::LPP_RESTART_BEHAVIOR_REQUEST_EXIT:
 			FPlatformMisc::RequestExit(true);
 			break;
diff --git a/Engine/Source/Developer/Windows/LiveCoding/Private/External/LPP_API.h b/Engine/Source/Developer/Windows/LiveCoding/Private/External/LPP_API.h
--- a/Engine/Source/Developer/Windows/LiveCoding/Private/External/LPP_API.h
+++ b/Engine/Source/Developer/Windows/LiveCoding/Private/External/LPP_API.h
@@ -224,7 +224,7 @@ LPP_API int lppWantsRestart(HMODULE livePP)
 
 enum RestartBehaviour
 {
-	// BEGIN EPIC MODS - Use UE4 codepath for termination to ensure logs are flushed and session analytics are sent
+	// BEGIN EPIC MODS - Use UE codepath for termination to ensure logs are flushed and session analytics are sent
 	LPP_RESTART_BEHAVIOR_REQUEST_EXIT,				// FPlatforMisc::RequestExit(true)
 	// END EPIC MODS
 	LPP_RESTART_BEHAVIOUR_DEFAULT_EXIT,				// ExitProcess()
diff --git a/Engine/Source/Developer/Windows/LiveCoding/Private/LiveCodingModule.cpp b/Engine/Source/Developer/Windows/LiveCoding/Private/LiveCodingModule.cpp
--- a/Engine/Source/Developer/Windows/LiveCoding/Private/LiveCodingModule.cpp
+++ b/Engine/Source/Developer/Windows/LiveCoding/Private/LiveCodingModule.cpp
@@ -426,7 +426,7 @@ bool FLiveCodingModule::StartLiveCoding()
 
 		// Enable the server
 		FString ProjectPath = FPaths::ConvertRelativePathToFull(FPaths::ProjectDir()).ToLower();
-		FString ProcessGroup = FString::Printf(TEXT("UE4_%s_0x%08x"), FApp::GetProjectName(), GetTypeHash(ProjectPath));
+		FString ProcessGroup = FString::Printf(TEXT("UE_%s_0x%08x"), FApp::GetProjectName(), GetTypeHash(ProjectPath));
 		LppRegisterProcessGroup(TCHAR_TO_ANSI(*ProcessGroup));
 
 		// Build the command line
