Removed references to UE4 in Epic modifications

diff --git a/Engine/Source/Developer/Windows/LiveCodingServer/Private/External/LC_LiveModule.cpp b/Engine/Source/Developer/Windows/LiveCodingServer/Private/External/LC_LiveModule.cpp
--- a/Engine/Source/Developer/Windows/LiveCodingServer/Private/External/LC_LiveModule.cpp
+++ b/Engine/Source/Developer/Windows/LiveCodingServer/Private/External/LC_LiveModule.cpp
@@ -47,7 +47,7 @@
 #include <process.h> // needed for _PVFV
 // END EPIC MOD
 
-// BEGIN EPIC MOD - Support for UE4 debug visualizers
+// BEGIN EPIC MOD - Support for UE debug visualizers
 #include "Misc/Paths.h"
 // END EPIC MOD
 
@@ -3068,7 +3068,7 @@ LiveModule::ErrorType::Enum LiveModule::Update(FileAttributeCache* fileCache, Di
 		linkerOptions += L"\n";
 	}
 
-	// BEGIN EPIC MOD - Support for UE4 debug visualizers
+	// BEGIN EPIC MOD - Support for UE debug visualizers
 	linkerOptions += L"\"";
 #if LC_64_BIT
 	linkerOptions += *FPaths::ConvertRelativePathToFull(FPaths::EngineDir() / L"Extras/NatvisHelpers/Win64/NatvisHelpers.lib");
