Instead of relying on non-conflict with library-provided wmemchr, just re-implement std::wstring::find on top of wcsstr for clang 11.0 on Windows only.

#rnx

diff --git a/Engine/Source/Developer/Windows/LiveCoding/Private/External/LC_StringUtil.cpp b/Engine/Source/Developer/Windows/LiveCoding/Private/External/LC_StringUtil.cpp
--- a/Engine/Source/Developer/Windows/LiveCoding/Private/External/LC_StringUtil.cpp
+++ b/Engine/Source/Developer/Windows/LiveCoding/Private/External/LC_StringUtil.cpp
@@ -4,26 +4,16 @@
 //#include PCH_INCLUDE
 // END EPIC MOD
 
+#include "LC_StringUtil.h"
 // BEGIN EPIC MOD
 #if defined(__clang__) && defined(_MSC_VER) && _MSVC_LANG > 201402L
 #if __clang_major__ > 11
 // For reference: https://bugs.llvm.org/show_bug.cgi?id=41226#c16
 #error Re-check whether this workaround is still needed for clang 12.0 and higher.
 #endif
-inline const wchar_t* wmemchr(const wchar_t* Str, wchar_t Ch, size_t Len)
-{
-	for (; 0 < Len; Str++, Len--)
-		if (*Str == Ch)
-			return Str;
-
-	return nullptr;
-}
+#include <wchar.h>
 #endif
-// END EPIC MOD
-
 
-#include "LC_StringUtil.h"
-// BEGIN EPIC MOD
 #include "Windows/WindowsHWrapper.h"
 // END EPIC MOD
 
@@ -173,6 +163,29 @@ namespace detail
 			}
 		}
 	}
+
+	// Temporary replacement for std::wstring::find
+	static size_t WideFindOffset(const std::wstring& haystack, const std::wstring& needle)
+	{
+#if defined(__clang__) && defined(_MSC_VER) && _MSVC_LANG > 201402L
+#if __clang_major__ > 11
+		// For reference: https://bugs.llvm.org/show_bug.cgi?id=41226#c16
+		#error Re-check whether this workaround is still needed for clang 12.0 and higher.
+#endif
+		const wchar_t* found = ::wcsstr(haystack.c_str(), needle.c_str());
+		if (found)
+		{
+			return static_cast<size_t>(found - haystack.c_str());
+		}
+		else
+		{
+			return std::wstring::npos;
+		}
+#else
+		return haystack.find(needle);
+#endif
+	}
+
 	// END EPIC MOD
 }
 
@@ -210,7 +223,9 @@ namespace string
 	{
 		std::wstring result = str;
 
-		size_t startPos = str.find(from);
+// BEGIN EPIC MOD
+		size_t startPos = detail::WideFindOffset(str, from);
+// END EPIC MOD
 		if (startPos == std::wstring::npos)
 			return result;
 
@@ -255,7 +270,9 @@ namespace string
 
 		for (;;)
 		{
-			const size_t pos = result.find(from);
+// BEGIN EPIC MOD
+			const size_t pos = detail::WideFindOffset(result, from);
+// END EPIC MOD
 			if (pos == std::wstring::npos)
 			{
 				return result;
@@ -291,7 +308,9 @@ namespace string
 
 		for (;;)
 		{
-			const size_t pos = result.find(subString);
+// BEGIN EPIC MOD
+			const size_t pos = detail::WideFindOffset(result, subString);
+// END EPIC MOD
 			if (pos == std::wstring::npos)
 			{
 				return result;
