Fixing Clang warnings in LiveCoding

diff --git a/Engine/Source/Developer/Windows/LiveCoding/Private/External/LC_Foundation_Windows.h b/Engine/Source/Developer/Windows/LiveCoding/Private/External/LC_Foundation_Windows.h
--- a/Engine/Source/Developer/Windows/LiveCoding/Private/External/LC_Foundation_Windows.h
+++ b/Engine/Source/Developer/Windows/LiveCoding/Private/External/LC_Foundation_Windows.h
@@ -5,6 +5,7 @@
 // BEGIN EPIC MOD
 #include "CoreTypes.h"
 #include "Windows/WindowsHWrapper.h"
+#include <atomic>
 // END EPIC MOD
 
 #define LC_FILE									__FILE__
@@ -23,7 +24,7 @@
 #define LC_RESTRICT								__restrict
 #define LC_UNUSED(_value)						(void)(_value)
 
-#define LC_COMPILER_FENCE						::_ReadWriteBarrier()
 // BEGIN EPIC MOD
-#define LC_MEMORY_FENCE							::_mm_sfence()
+#define LC_COMPILER_FENCE						std::atomic_signal_fence(std::memory_order_acq_rel)
+#define LC_MEMORY_FENCE							std::atomic_thread_fence(std::memory_order_acq_rel)
 // END EPIC MOD
\ No newline at end of file
diff --git a/Engine/Source/Developer/Windows/LiveCodingServer/Private/External/LC_COMThread.h b/Engine/Source/Developer/Windows/LiveCodingServer/Private/External/LC_COMThread.h
--- a/Engine/Source/Developer/Windows/LiveCodingServer/Private/External/LC_COMThread.h
+++ b/Engine/Source/Developer/Windows/LiveCodingServer/Private/External/LC_COMThread.h
@@ -27,10 +27,14 @@ public:
 	~COMThread(void);
 
 	// helper function that runs any given function in the internal COM thread and waits until execution has finished
+	// BEGIN EPIC MOD
 	template <typename F, typename... Args>
-	inline typename std::result_of<F(Args...)>::type CallInThread(F ptrToFunction, Args&&... args)
+	inline std::invoke_result_t<F, Args...> CallInThread(F ptrToFunction, Args&&... args)
+	// END EPIC MOD
 	{
-		typedef typename std::result_of<F(Args...)>::type ReturnValue;
+		// BEGIN EPIC MOD
+		typedef std::invoke_result_t<F, Args...> ReturnValue;
+		// END EPIC MOD
 
 		// a helper lambda that calls the given functions with the provided arguments.
 		// because this lambda needs to capture, it cannot be converted to a function pointer implicitly,
diff --git a/Engine/Source/Developer/Windows/LiveCodingServer/Private/External/LC_Scheduler.h b/Engine/Source/Developer/Windows/LiveCodingServer/Private/External/LC_Scheduler.h
--- a/Engine/Source/Developer/Windows/LiveCodingServer/Private/External/LC_Scheduler.h
+++ b/Engine/Source/Developer/Windows/LiveCodingServer/Private/External/LC_Scheduler.h
@@ -13,19 +13,21 @@ namespace scheduler
 	void Startup(void);
 	void Shutdown(void);
 
+	// BEGIN EPIC MOD
 	// creates a new task from any function
 	template <typename F>
-	Task<typename std::result_of<F()>::type>* CreateTask(F&& function)
+	Task<std::invoke_result_t<F>>* CreateTask(F&& function)
 	{
-		return new Task<typename std::result_of<F()>::type>(function);
+		return new Task<std::invoke_result_t<F()>>(function);
 	}
 
 	// creates a new task from any function as child of a parent task
 	template <typename F>
-	Task<typename std::result_of<F()>::type>* CreateTask(TaskBase* parent, F&& function)
+	Task<std::invoke_result_t<F>>* CreateTask(TaskBase* parent, F&& function)
 	{
-		return new Task<typename std::result_of<F()>::type>(parent, function);
+		return new Task<std::invoke_result_t<F>>(parent, function);
 	}
+	// END EPIC MOD
 
 	// creates an empty task
 	TaskBase* CreateEmptyTask(void);
