Work in progress support for D3D12 Shader Model  6.6, specifying validator version, and workarounds to disable validation/signing until validator v1.6 is available that supports the 6.6 profile (requires developer mode and experimental features).

Disabled by default (guarded by USE_SHADER_MODEL_6_6).

Includes work done by Yuriy.

diff --git a/Engine/Source/Developer/Windows/ShaderFormatD3D/Private/D3DShaderCompiler.cpp b/Engine/Source/Developer/Windows/ShaderFormatD3D/Private/D3DShaderCompiler.cpp
--- a/Engine/Source/Developer/Windows/ShaderFormatD3D/Private/D3DShaderCompiler.cpp
+++ b/Engine/Source/Developer/Windows/ShaderFormatD3D/Private/D3DShaderCompiler.cpp
@@ -92,26 +92,27 @@ static const TCHAR* GetShaderProfileName(FShaderTarget Target, bool bForceSM6)
 			checkfSlow(false, TEXT("Unexpected shader frequency"));
 			return nullptr;
 		case SF_Pixel:
-			return TEXT("ps_6_5");
+			return USE_SHADER_MODEL_6_6 ? TEXT("ps_6_6") : TEXT("ps_6_5");
 		case SF_Vertex:
-			return TEXT("vs_6_5");
+			return USE_SHADER_MODEL_6_6 ? TEXT("vs_6_6") : TEXT("vs_6_5");
 		case SF_Mesh:
-			return TEXT("ms_6_5");
+			return USE_SHADER_MODEL_6_6 ? TEXT("ms_6_6") : TEXT("ms_6_5");
 		case SF_Amplification:
-			return TEXT("as_6_5");
+			return USE_SHADER_MODEL_6_6 ? TEXT("as_6_6") : TEXT("as_6_5");
 		case SF_Hull:
-			return TEXT("hs_6_5");
+			return USE_SHADER_MODEL_6_6 ? TEXT("hs_6_6") : TEXT("hs_6_5");
 		case SF_Domain:
-			return TEXT("ds_6_5");
+			return USE_SHADER_MODEL_6_6 ? TEXT("ds_6_6") : TEXT("ds_6_5");
 		case SF_Geometry:
-			return TEXT("gs_6_5");
+			return USE_SHADER_MODEL_6_6 ? TEXT("gs_6_6") : TEXT("gs_6_5");
 		case SF_Compute:
-			return TEXT("cs_6_5");
+			return USE_SHADER_MODEL_6_6 ? TEXT("cs_6_6") : TEXT("cs_6_5");
 		case SF_RayGen:
 		case SF_RayMiss:
 		case SF_RayHitGroup:
 		case SF_RayCallable:
-			return TEXT("lib_6_5");
+			// return USE_SHADER_MODEL_6_6 ? TEXT("lib_6_6") : TEXT("lib_6_5");
+			return TEXT("lib_6_5"); // TODO: Intentionally using SM6.5 as a workaround for outstanding driver bug.
 		}
 	}
 	else if(Target.Platform == SP_PCD3D_SM5)
@@ -162,7 +163,7 @@ static const TCHAR* GetShaderProfileName(FShaderTarget Target, bool bForceSM6)
 		}
 	}
 
-	return NULL;
+	return nullptr;
 }
 
 /**
diff --git a/Engine/Source/Developer/Windows/ShaderFormatD3D/Private/D3DShaderCompiler.inl b/Engine/Source/Developer/Windows/ShaderFormatD3D/Private/D3DShaderCompiler.inl
--- a/Engine/Source/Developer/Windows/ShaderFormatD3D/Private/D3DShaderCompiler.inl
+++ b/Engine/Source/Developer/Windows/ShaderFormatD3D/Private/D3DShaderCompiler.inl
@@ -18,17 +18,23 @@ protected:
 	TArray<FString> ExtraArguments;
 
 public:
-	FDxcArguments(const FString& InEntryPoint, const TCHAR* InShaderProfile, const FString& InExports,
+	FDxcArguments(
+		const FString& InEntryPoint,
+		const TCHAR* InShaderProfile,
+		const FString& InExports,
 		const FString& InDumpDebugInfoPath,	// Optional, empty when not dumping shader debug info
 		const FString& InBaseFilename,
 		bool bInEnable16BitTypes,
 		bool bKeepDebugInfo,
-		uint32 D3DCompileFlags, uint32 AutoBindingSpace = ~0u)
-		: ShaderProfile(InShaderProfile)
-		, EntryPoint(InEntryPoint)
-		, Exports(InExports)
-		, DumpDebugInfoPath(InDumpDebugInfoPath)
-		, bEnable16BitTypes(bInEnable16BitTypes)
+		uint32 D3DCompileFlags,
+		uint32 AutoBindingSpace,
+		const TCHAR* InOptValidatorVersion
+	)
+	: ShaderProfile(InShaderProfile)
+	, EntryPoint(InEntryPoint)
+	, Exports(InExports)
+	, DumpDebugInfoPath(InDumpDebugInfoPath)
+	, bEnable16BitTypes(bInEnable16BitTypes)
 	{
 		BatchBaseFilename = FPaths::GetBaseFilename(InBaseFilename);
 
@@ -133,6 +139,11 @@ public:
 
 		checkf(D3DCompileFlags == 0, TEXT("Unhandled shader compiler flags 0x%x!"), D3DCompileFlags);
 
+		if (InOptValidatorVersion)
+		{
+			ExtraArguments.Add(L"/validator-version");
+			ExtraArguments.Add(FString(InOptValidatorVersion));
+		}
 
 		ExtraArguments.Add(L"/Zss");
 		ExtraArguments.Add(L"/Qembed_debug");
diff --git a/Engine/Source/Developer/Windows/ShaderFormatD3D/Private/D3DShaderCompilerDXC.cpp b/Engine/Source/Developer/Windows/ShaderFormatD3D/Private/D3DShaderCompilerDXC.cpp
--- a/Engine/Source/Developer/Windows/ShaderFormatD3D/Private/D3DShaderCompilerDXC.cpp
+++ b/Engine/Source/Developer/Windows/ShaderFormatD3D/Private/D3DShaderCompilerDXC.cpp
@@ -309,7 +309,7 @@ static HRESULT D3DCompileToDxil(const char* SourceText, FDxcArguments& Arguments
 
 		if (Arguments.ShouldDump())
 		{
-			// Dump dissasembly before we strip reflection out
+			// Dump disassembly before we strip reflection out
 			const FString& DisasmFilename = Arguments.GetDumpDisassemblyFilename();
 			check(DisasmFilename.Len() > 0);
 			DisassembleAndSave(Compiler, OutDxilBlob, DisasmFilename);
@@ -663,10 +663,36 @@ bool CompileAndProcessD3DShaderDXC(FString& PreprocessedShaderSource,
 		DXCFlags |= D3DCOMPILE_SKIP_OPTIMIZATION;
 	}
 
+	const TCHAR* ValidatorVersion = nullptr; // Whatever default validator chosen by DXC
+
+#if USE_SHADER_MODEL_6_6
+	if (Input.Target.GetPlatform() == SP_PCD3D_SM6)
+	{
+		// Experimental shader model 6.6 requires turning off DXIL validation and signing as of 2021-04-04.
+		// Unsigned shader will only run in Windows 10 developer mode, and only when experimental shader models are enabled.
+		// TODO: This should be removed when dxil.dll supports SM 6.6 validation.
+		DXCFlags |= D3DCOMPILE_SKIP_VALIDATION;
+
+		// Shader model 6.6 requires validator 1.6
+		ValidatorVersion = TEXT("1.6");
+	}
+#endif
+
 	const bool bKeepDebugInfo = Input.Environment.CompilerFlags.Contains(CFLAG_KeepDebugInfo);
 
-	FDxcArguments Args(EntryPointName, ShaderProfile, RayTracingExports,
-		Input.DumpDebugInfoPath, Filename, bEnable16BitTypes, bKeepDebugInfo, DXCFlags, AutoBindingSpace);
+	FDxcArguments Args
+	(
+		EntryPointName,
+		ShaderProfile,
+		RayTracingExports,
+		Input.DumpDebugInfoPath,
+		Filename,
+		bEnable16BitTypes,
+		bKeepDebugInfo,
+		DXCFlags,
+		AutoBindingSpace,
+		ValidatorVersion
+	);
 
 	if (bDumpDebugInfo)
 	{
diff --git a/Engine/Source/Developer/Windows/ShaderFormatD3D/Private/ShaderFormatD3D.cpp b/Engine/Source/Developer/Windows/ShaderFormatD3D/Private/ShaderFormatD3D.cpp
--- a/Engine/Source/Developer/Windows/ShaderFormatD3D/Private/ShaderFormatD3D.cpp
+++ b/Engine/Source/Developer/Windows/ShaderFormatD3D/Private/ShaderFormatD3D.cpp
@@ -41,11 +41,17 @@ public:
 		CheckFormat(Format);
 		if (Format == NAME_PCD3D_SM6)
 		{
-			return HashCombine(DxcVersionHash, GetTypeHash(UE_SHADER_PCD3D_SM6_VER));
+			uint32 ShaderModelHash = GetTypeHash(UE_SHADER_PCD3D_SM6_VER);
+
+		#if USE_SHADER_MODEL_6_6
+			ShaderModelHash ^= 0x96ED7F56;
+		#endif
+
+			return HashCombine(DxcVersionHash, ShaderModelHash);
 		}
 		else if (Format == NAME_PCD3D_SM5)
 		{
-			// Technically not needed for regular SM5 commpiled with legacy compiler,
+			// Technically not needed for regular SM5 compiled with legacy compiler,
 			// but PCD3D_SM5 currently includes ray tracing shaders that are compiled with new compiler stack.
 			return HashCombine(DxcVersionHash, GetTypeHash(UE_SHADER_PCD3D_SM5_VER));
 		}
diff --git a/Engine/Source/Developer/Windows/ShaderFormatD3D/Private/ShaderFormatD3D.h b/Engine/Source/Developer/Windows/ShaderFormatD3D/Private/ShaderFormatD3D.h
--- a/Engine/Source/Developer/Windows/ShaderFormatD3D/Private/ShaderFormatD3D.h
+++ b/Engine/Source/Developer/Windows/ShaderFormatD3D/Private/ShaderFormatD3D.h
@@ -4,6 +4,9 @@
 
 #include "Windows/WindowsHWrapper.h"
 
+// TODO: Lock D3D12 to SM 6.6 min spec
+#define USE_SHADER_MODEL_6_6 0
+
 struct FShaderTarget;
 
 enum class ELanguage
