Move D3DShaderCompiler fatal error over UAV counts to a shader compiler error.

diff --git a/Engine/Source/Developer/Windows/ShaderFormatD3D/Private/D3DShaderCompiler.cpp b/Engine/Source/Developer/Windows/ShaderFormatD3D/Private/D3DShaderCompiler.cpp
--- a/Engine/Source/Developer/Windows/ShaderFormatD3D/Private/D3DShaderCompiler.cpp
+++ b/Engine/Source/Developer/Windows/ShaderFormatD3D/Private/D3DShaderCompiler.cpp
@@ -810,8 +810,9 @@ bool CompileAndProcessD3DShaderFXC(FString& PreprocessedShaderSource, const FStr
 		// Check for resource limits for feature level 11.0
 		if (NumUAVs > GD3DMaximumNumUAVs)
 		{
+			FilteredErrors.Add(FString::Printf(TEXT("Number of UAVs exceeded limit: %d slots used, but limit is %d due to maximum feature level 11.0"), NumUAVs, GD3DMaximumNumUAVs));
+			Result = E_FAIL;
 			Output.bSucceeded = false;
-			UE_LOG(LogD3D11ShaderCompiler, Fatal, TEXT("Number of UAVs in \"%s\" exceeded limit: %d slots used, but limit is %d due to maximum feature level 11.0"), *Input.VirtualSourceFilePath, NumUAVs, GD3DMaximumNumUAVs);
 		}
 
 		// Save results if compilation and reflection succeeded
