Added support for more arbitrary property changes to live coding.

1) Changed the pre/post compile notifications from module notifications to process commands.
2) Added server command to notify that we want re-instance patching (two phase)
3) Added support for two phase patching to enable re-instancing without all the limitations.
4) Added a null CDO check for old blueprint classes (approved by Phillip)

#rnx

diff --git a/Engine/Source/Developer/Windows/LiveCoding/Private/External/LC_API.cpp b/Engine/Source/Developer/Windows/LiveCoding/Private/External/LC_API.cpp
--- a/Engine/Source/Developer/Windows/LiveCoding/Private/External/LC_API.cpp
+++ b/Engine/Source/Developer/Windows/LiveCoding/Private/External/LC_API.cpp
@@ -345,6 +345,12 @@ LPP_DLL_API(void*) LppEnableLazyLoadedModule(const wchar_t* nameOfExeOrDll)
 }
 // END EPIC MOD
 
+// BEGIN EPIC MOD
+LPP_DLL_API(void) LppEnableReinstancingFlow()
+{
+	g_startupThread->EnableReinstancingFlow();
+}
+// END EPIC MOD
 
 LPP_DLL_API(void) LppApplySettingBool(const char* settingName, int value)
 {
diff --git a/Engine/Source/Developer/Windows/LiveCoding/Private/External/LC_API.h b/Engine/Source/Developer/Windows/LiveCoding/Private/External/LC_API.h
--- a/Engine/Source/Developer/Windows/LiveCoding/Private/External/LC_API.h
+++ b/Engine/Source/Developer/Windows/LiveCoding/Private/External/LC_API.h
@@ -41,6 +41,7 @@ LPP_DLL_API(void) LppSetVisible(bool visible);
 LPP_DLL_API(void) LppSetActive(bool active);
 LPP_DLL_API(void) LppSetBuildArguments(const wchar_t* arguments);
 LPP_DLL_API(void*) LppEnableLazyLoadedModule(const wchar_t* nameOfExeOrDll);
+LPP_DLL_API(void) LppEnableReinstancingFlow();
 // END EPIC MOD
 LPP_DLL_API(void) LppApplySettingBool(const char* settingName, int value);
 LPP_DLL_API(void) LppApplySettingInt(const char* settingName, int value);
diff --git a/Engine/Source/Developer/Windows/LiveCoding/Private/External/LC_ClientCommandActions.cpp b/Engine/Source/Developer/Windows/LiveCoding/Private/External/LC_ClientCommandActions.cpp
--- a/Engine/Source/Developer/Windows/LiveCoding/Private/External/LC_ClientCommandActions.cpp
+++ b/Engine/Source/Developer/Windows/LiveCoding/Private/External/LC_ClientCommandActions.cpp
@@ -91,16 +91,6 @@ bool actions::CallHooks::Execute(const CommandType* command, const DuplexPipe* p
 		case hook::Type::COMPILE_ERROR_MESSAGE:
 			hook::CallHooksInRange<hook::CompileErrorMessageFunction>(command->rangeBegin, command->rangeEnd, static_cast<const wchar_t*>(payload));
 			break;
-
-		// BEGIN EPIC MOD - Add the ability for pre and post compile notifications
-		case hook::Type::PRECOMPILE:
-			hook::CallHooksInRange<hook::PrecompileFunction>(command->rangeBegin, command->rangeEnd);
-			break;
-
-		case hook::Type::POSTCOMPILE:
-			hook::CallHooksInRange<hook::PostcompileFunction>(command->rangeBegin, command->rangeEnd);
-			break;
-		// END EPIC MOD
 	}
 
 	pipe->SendAck();
@@ -124,6 +114,12 @@ extern void LiveCodingBeginPatch();
 extern void LiveCodingEndCompile();
 // END EPIC MOD
 
+// BEGIN EPIC MOD
+extern void LiveCodingPreCompile();
+extern void LiveCodingPostCompile();
+extern void LiveCodingTriggerReload();
+// END EPIC MOD
+
 bool actions::LoadPatch::Execute(const CommandType* command, const DuplexPipe* pipe, void*, const void*, size_t)
 {
 	// load library into this process
@@ -209,3 +205,34 @@ bool actions::HandleExceptionFinished::Execute(const CommandType* command, const
 	// don't continue execution
 	return false;
 }
+
+
+// BEGIN EPIC MOD
+bool actions::PreCompile::Execute(const CommandType*, const DuplexPipe* pipe, void*, const void*, size_t)
+{
+	LiveCodingPreCompile();
+	pipe->SendAck();
+
+	return true;
+}
+
+
+bool actions::PostCompile::Execute(const CommandType*, const DuplexPipe* pipe, void*, const void*, size_t)
+{
+	LiveCodingPostCompile();
+	pipe->SendAck();
+
+	return true;
+}
+
+
+bool actions::TriggerReload::Execute(const CommandType*, const DuplexPipe* pipe, void*, const void*, size_t)
+{
+	LiveCodingTriggerReload();
+	syncPoint::Leave();
+	syncPoint::Enter();
+	pipe->SendAck();
+
+	return true;
+}
+// END EPIC MOD
\ No newline at end of file
diff --git a/Engine/Source/Developer/Windows/LiveCoding/Private/External/LC_ClientCommandActions.h b/Engine/Source/Developer/Windows/LiveCoding/Private/External/LC_ClientCommandActions.h
--- a/Engine/Source/Developer/Windows/LiveCoding/Private/External/LC_ClientCommandActions.h
+++ b/Engine/Source/Developer/Windows/LiveCoding/Private/External/LC_ClientCommandActions.h
@@ -31,6 +31,11 @@ namespace actions
 	DECLARE_ACTION(LogOutput);
 	DECLARE_ACTION(CompilationFinished);
 	DECLARE_ACTION(HandleExceptionFinished);
+	// BEGIN EPIC MOD
+	DECLARE_ACTION(PreCompile);
+	DECLARE_ACTION(PostCompile);
+	DECLARE_ACTION(TriggerReload);
+	// END EPIC MOD
 
 	#undef DECLARE_ACTION
 }
diff --git a/Engine/Source/Developer/Windows/LiveCoding/Private/External/LC_ClientCommandThread.cpp b/Engine/Source/Developer/Windows/LiveCoding/Private/External/LC_ClientCommandThread.cpp
--- a/Engine/Source/Developer/Windows/LiveCoding/Private/External/LC_ClientCommandThread.cpp
+++ b/Engine/Source/Developer/Windows/LiveCoding/Private/External/LC_ClientCommandThread.cpp
@@ -60,6 +60,11 @@ Thread::ReturnValue ClientCommandThread::ThreadFunction(const std::wstring& proc
 	commandMap.RegisterAction<actions::CallHooks>();
 	commandMap.RegisterAction<actions::LogOutput>();
 	commandMap.RegisterAction<actions::CompilationFinished>();
+	// BEGIN EPIC MOD
+	commandMap.RegisterAction<actions::PreCompile>();
+	commandMap.RegisterAction<actions::PostCompile>();
+	commandMap.RegisterAction<actions::TriggerReload>();
+	// END EPIC MOD
 
 	HeartBeat heartBeat(processGroupName.c_str(), Process::Current::GetId());
 
diff --git a/Engine/Source/Developer/Windows/LiveCoding/Private/External/LC_ClientStartupThread.cpp b/Engine/Source/Developer/Windows/LiveCoding/Private/External/LC_ClientStartupThread.cpp
--- a/Engine/Source/Developer/Windows/LiveCoding/Private/External/LC_ClientStartupThread.cpp
+++ b/Engine/Source/Developer/Windows/LiveCoding/Private/External/LC_ClientStartupThread.cpp
@@ -384,6 +384,18 @@ void* ClientStartupThread::EnableLazyLoadedModule(const wchar_t* fileName, Windo
 }
 // END EPIC MOD
 
+// BEGIN EPIC MOD
+void ClientStartupThread::EnableReinstancingFlow()
+{
+	// we cannot wait for commands in the user command thread as long as startup hasn't finished
+	Join();
+
+	if (m_userCommandThread)
+	{
+		m_userCommandThread->EnableReinstancingFlow();
+	}
+}
+// END EPIC MOD
 
 void ClientStartupThread::ApplySettingBool(const char* settingName, int value)
 {
diff --git a/Engine/Source/Developer/Windows/LiveCoding/Private/External/LC_ClientStartupThread.h b/Engine/Source/Developer/Windows/LiveCoding/Private/External/LC_ClientStartupThread.h
--- a/Engine/Source/Developer/Windows/LiveCoding/Private/External/LC_ClientStartupThread.h
+++ b/Engine/Source/Developer/Windows/LiveCoding/Private/External/LC_ClientStartupThread.h
@@ -77,6 +77,10 @@ public:
 	void* EnableLazyLoadedModule(const wchar_t* fileName, Windows::HMODULE moduleBase);
 	// END EPIC MOD
 
+	// BEGIN EPIC MOD
+	void EnableReinstancingFlow();
+	// END EPIC MOD
+
 private:
 	Thread::ReturnValue ThreadFunction(const std::wstring& groupName, RunMode::Enum runMode);
 
diff --git a/Engine/Source/Developer/Windows/LiveCoding/Private/External/LC_ClientUserCommandThread.cpp b/Engine/Source/Developer/Windows/LiveCoding/Private/External/LC_ClientUserCommandThread.cpp
--- a/Engine/Source/Developer/Windows/LiveCoding/Private/External/LC_ClientUserCommandThread.cpp
+++ b/Engine/Source/Developer/Windows/LiveCoding/Private/External/LC_ClientUserCommandThread.cpp
@@ -554,6 +554,15 @@ void* ClientUserCommandThread::EnableLazyLoadedModule(const wchar_t* fileName, W
 }
 // END EPIC MOD
 
+// BEGIN EPIC MOD
+void ClientUserCommandThread::EnableReinstancingFlow()
+{
+	ProxyCommand<commands::EnableReinstancingFlow>* proxy = new ProxyCommand<commands::EnableReinstancingFlow>(false, 0u);
+	proxy->m_command.processId = Process::Current::GetId();
+
+	PushUserCommand(proxy);
+}
+// END EPIC MOD
 
 void ClientUserCommandThread::InstallExceptionHandler(void)
 {
@@ -631,6 +640,11 @@ Thread::ReturnValue ClientUserCommandThread::ThreadFunction(Event* waitForStartE
 	moduleCommandMap.RegisterAction<actions::LeaveSyncPoint>();
 	moduleCommandMap.RegisterAction<actions::CallEntryPoint>();
 	moduleCommandMap.RegisterAction<actions::CallHooks>();
+	// BEGIN EPIC MOD
+	moduleCommandMap.RegisterAction<actions::PreCompile>();
+	moduleCommandMap.RegisterAction<actions::PostCompile>();
+	moduleCommandMap.RegisterAction<actions::TriggerReload>();
+	// END EPIC MOD
 
 	for (;;)
 	{
diff --git a/Engine/Source/Developer/Windows/LiveCoding/Private/External/LC_ClientUserCommandThread.h b/Engine/Source/Developer/Windows/LiveCoding/Private/External/LC_ClientUserCommandThread.h
--- a/Engine/Source/Developer/Windows/LiveCoding/Private/External/LC_ClientUserCommandThread.h
+++ b/Engine/Source/Developer/Windows/LiveCoding/Private/External/LC_ClientUserCommandThread.h
@@ -92,6 +92,10 @@ public:
 	void* EnableLazyLoadedModule(const wchar_t* fileName, Windows::HMODULE moduleBase);
 	// END EPIC MOD
 
+	// BEGIN EPIC MOD
+	void EnableReinstancingFlow();
+	// END EPIC MOD
+
 	void ApplySettingBool(const char* settingName, int value);
 	void ApplySettingInt(const char* settingName, int value);
 	void ApplySettingString(const char* settingName, const wchar_t* value);
diff --git a/Engine/Source/Developer/Windows/LiveCoding/Private/External/LC_CommandMap.cpp b/Engine/Source/Developer/Windows/LiveCoding/Private/External/LC_CommandMap.cpp
--- a/Engine/Source/Developer/Windows/LiveCoding/Private/External/LC_CommandMap.cpp
+++ b/Engine/Source/Developer/Windows/LiveCoding/Private/External/LC_CommandMap.cpp
@@ -64,6 +64,10 @@ CommandMap::CommandMap(void)
 	RegisterDefaultAction<commands::HandleException>(m_actions);
 	RegisterDefaultAction<commands::HandleExceptionFinished>(m_actions);
 
+	RegisterDefaultAction<commands::ApplySettingBool>(m_actions);
+	RegisterDefaultAction<commands::ApplySettingInt>(m_actions);
+	RegisterDefaultAction<commands::ApplySettingString>(m_actions);
+
 	// BEGIN EPIC MOD - Adding ShowConsole command
 	RegisterDefaultAction<commands::ShowConsole>(m_actions);
 	// END EPIC MOD
@@ -80,9 +84,12 @@ CommandMap::CommandMap(void)
 	RegisterDefaultAction<commands::SetBuildArguments>(m_actions);
 	// END EPIC MOD
 
-	RegisterDefaultAction<commands::ApplySettingBool>(m_actions);
-	RegisterDefaultAction<commands::ApplySettingInt>(m_actions);
-	RegisterDefaultAction<commands::ApplySettingString>(m_actions);
+	// BEGIN EPIC MOD
+	RegisterDefaultAction<commands::PreCompile>(m_actions);
+	RegisterDefaultAction<commands::PostCompile>(m_actions);
+	RegisterDefaultAction<commands::TriggerReload>(m_actions);
+	RegisterDefaultAction<commands::EnableReinstancingFlow>(m_actions);
+	// END EPIC MOD
 }
 
 
diff --git a/Engine/Source/Developer/Windows/LiveCoding/Private/External/LC_Commands.h b/Engine/Source/Developer/Windows/LiveCoding/Private/External/LC_Commands.h
--- a/Engine/Source/Developer/Windows/LiveCoding/Private/External/LC_Commands.h
+++ b/Engine/Source/Developer/Windows/LiveCoding/Private/External/LC_Commands.h
@@ -244,32 +244,53 @@
 		bool continueExecution;
 	};
 
-	// BEGIN EPIC MOD - Adding ShowConsole command
-	struct ShowConsole
+	// tell the EXE that a bool setting needs to be changed
+	struct ApplySettingBool
 	{
 		static const uint32_t ID = HandleExceptionFinished::ID + 1u;
+
+		char settingName[256];
+		int settingValue;
 	};
-	// END EPIC MOD
 
-	// BEGIN EPIC MOD - Adding SetVisible command
+	// tell the EXE that a int setting needs to be changed
+	struct ApplySettingInt
+	{
+		static const uint32_t ID = ApplySettingBool::ID + 1u;
+
+		char settingName[256];
+		int settingValue;
+	};
+
+	// tell the EXE that a bool setting needs to be changed
+	struct ApplySettingString
+	{
+		static const uint32_t ID = ApplySettingInt::ID + 1u;
+
+		char settingName[256];
+		wchar_t settingValue[256];
+	};
+
+	// BEGIN EPIC MOD
+	struct ShowConsole
+	{
+		static const uint32_t ID = ApplySettingString::ID + 1u;
+	};
+
 	struct SetVisible
 	{
 		static const uint32_t ID = ShowConsole::ID + 1u;
 
 		bool visible;
 	};
-	// END EPIC MOD
 
-	// BEGIN EPIC MOD - Adding SetActive command
 	struct SetActive
 	{
 		static const uint32_t ID = SetVisible::ID + 1u;
 
 		bool active;
 	};
-	// END EPIC MOD
 
-	// BEGIN EPIC MOD - Adding SetBuildArguments command
 	struct SetBuildArguments
 	{
 		static const uint32_t ID = SetActive::ID + 1u;
@@ -277,9 +298,7 @@
 		Process::Id processId;
 		wchar_t arguments[1024];
 	};
-	// END EPIC MOD
 
-	// BEGIN EPIC MOD - Support for lazy-loading modules
 	struct EnableLazyLoadedModule
 	{
 		static const uint32_t ID = SetBuildArguments::ID + 1u;
@@ -294,36 +313,29 @@
 	{
 		static const uint32_t ID = EnableLazyLoadedModule::ID + 1u;
 	};
-	// END EPIC MOD
 
-	// tell the EXE that a bool setting needs to be changed
-	struct ApplySettingBool
+	struct PreCompile
 	{
-		// BEGIN EPIC MOD - Support for lazy-loading modules
 		static const uint32_t ID = FinishedLazyLoadingModules::ID + 1u;
-		// END EPIC MOD
-
-		char settingName[256];
-		int settingValue;
 	};
 
-	// tell the EXE that a int setting needs to be changed
-	struct ApplySettingInt
+	struct PostCompile
 	{
-		static const uint32_t ID = ApplySettingBool::ID + 1u;
-
-		char settingName[256];
-		int settingValue;
+		static const uint32_t ID = PreCompile::ID + 1u;
 	};
 
-	// tell the EXE that a bool setting needs to be changed
-	struct ApplySettingString
+	struct TriggerReload
 	{
-		static const uint32_t ID = ApplySettingInt::ID + 1u;
-
-		char settingName[256];
-		wchar_t settingValue[256];
+		static const uint32_t ID = PostCompile::ID + 1u;
 	};
 
-	static const uint32_t COUNT = ApplySettingString::ID + 1u;
+	struct EnableReinstancingFlow
+	{
+		static const uint32_t ID = TriggerReload::ID + 1u;
+
+		Process::Id processId;
+	};
+
+	static const uint32_t COUNT = EnableReinstancingFlow::ID + 1u;
+	// END EPIC MOD
 }
diff --git a/Engine/Source/Developer/Windows/LiveCoding/Private/LiveCodingModule.cpp b/Engine/Source/Developer/Windows/LiveCoding/Private/LiveCodingModule.cpp
--- a/Engine/Source/Developer/Windows/LiveCoding/Private/LiveCodingModule.cpp
+++ b/Engine/Source/Developer/Windows/LiveCoding/Private/LiveCodingModule.cpp
@@ -30,6 +30,7 @@ IMPLEMENT_MODULE(FLiveCodingModule, LiveCoding)
 #define LOCTEXT_NAMESPACE "LiveCodingModule"
 
 bool GIsCompileActive = false;
+bool GTriggerReload = false;
 bool GHasLoadedPatch = false;
 FString GLiveCodingConsolePath;
 FString GLiveCodingConsoleArguments;
@@ -40,18 +41,6 @@ extern const TCHAR* GLiveCodingEngineDir;
 extern const TCHAR* GLiveCodingProject;
 #endif
 
-#ifdef __clang__
-#pragma clang diagnostic push
-#pragma clang diagnostic ignored "-Wextern-initializer"
-#endif
-
-LPP_PRECOMPILE_HOOK(FLiveCodingModule::PreCompileHook);
-LPP_POSTCOMPILE_HOOK(FLiveCodingModule::PostCompileHook);
-
-#ifdef __clang__
-#pragma clang diagnostic pop
-#endif
-
 #if !WITH_EDITOR
 class FNullReload : public IReload
 {
@@ -376,32 +365,75 @@ void FLiveCodingModule::AttemptSyncLivePatching()
 	extern void LppSyncPoint();
 	LppSyncPoint();
 
-	if (!GIsCompileActive && Reload.IsValid())
+	if ((!GIsCompileActive || GTriggerReload) && Reload.IsValid())
 	{
 		if (GHasLoadedPatch)
 		{
-			CheckForClassChanges();
-			OnPatchCompleteDelegate.Broadcast();
-			GHasLoadedPatch = false;
-		}
-		Reload.Reset();
-	}
-}
-
-void FLiveCodingModule::CheckForClassChanges()
-{
 #if WITH_COREUOBJECT && WITH_ENGINE
 
-	// During the module loading process, the list of changed classes will be recorded.  Invoking this method will 
-	// result in the RegisterForReinstancing method being invoked which in turn records the classes in the ClassesToReinstance
-	// member variable being populated.
-	ProcessNewlyLoadedUObjects();
+			// During the module loading process, the list of changed classes will be recorded.  Invoking this method will 
+			// result in the RegisterForReinstancing method being invoked which in turn records the classes in the ClassesToReinstance
+			// member variable being populated.
+			ProcessNewlyLoadedUObjects();
 
+			// Complete the process of re-instancing without doing a GC
 #if WITH_EDITOR
-	Reload->Finalize();
+			Reload->Finalize(false);
 #endif
 
+			// Loop through all of the classes looking for classes that have been re-instanced.  We have to clear
+			// the CDO in order for the CDO to be GC'ed.  Otherwise when the CDO is later GC'ed, it will invoke the
+			// wrong destructor.
+			TArray<UClass*> ReinstClasses;
+			if (Reload->GetEnableReinstancing(false))
+			{
+				for (TObjectIterator<UClass> It; It; ++It)
+				{
+					UClass* Class = *It;
+					if (Class->GetName().StartsWith(TEXT("LIVECODING_")))
+					{
+						ReinstClasses.Add(Class);
+						Class->ClassDefaultObject = nullptr;
+					}
+					else if (Class->GetName().StartsWith(TEXT("REINST_")))
+					{
+						Class->ClassDefaultObject = nullptr;
+					}
+				}
+			}
+
+			// Perform the GC to try and destruct all the objects which will be invoking the old destructors.
+			CollectGarbage(GARBAGE_COLLECTION_KEEPFLAGS, true);
+#endif
+
+			// Second sync point to finish off the patching
+			if (GTriggerReload)
+			{
+				LppSyncPoint();
+			}
+
+#if WITH_COREUOBJECT && WITH_ENGINE
+			// Re-create the default objects.  This is technically incorrect since the structure doesn't match.
+			// But we need a CDO just in case.
+			if (Reload->GetEnableReinstancing(false))
+			{
+				for (UClass* Class : ReinstClasses)
+				{
+					Class->GetDefaultObject(true);
+				}
+			}
 #endif
+
+			OnPatchCompleteDelegate.Broadcast();
+			GHasLoadedPatch = false;
+		}
+		else if (GTriggerReload)
+		{
+			LppSyncPoint();
+		}
+		Reload.Reset();
+	}
+	GTriggerReload = false;
 }
 
 ILiveCodingModule::FOnPatchCompleteDelegate& FLiveCodingModule::GetOnPatchCompleteDelegate()
@@ -461,6 +493,15 @@ bool FLiveCodingModule::StartLiveCoding()
 		}
 		LppSetBuildArguments(*Arguments);
 
+#if WITH_EDITOR
+		bool bEnableReinstancing = true;
+		GConfig->GetBool(TEXT("LiveCoding"), TEXT("bEnableReinstancing"), bEnableReinstancing, GEngineIni);
+		if (bEnableReinstancing)
+		{
+			LppEnableReinstancingFlow();
+		}
+#endif
+
 		// Create a mutex that allows UBT to detect that we shouldn't hot-reload into this executable. The handle to it will be released automatically when the process exits.
 		FString ExecutablePath = FPaths::ConvertRelativePathToFull(FPlatformProcess::ExecutablePath());
 
@@ -624,19 +665,6 @@ bool FLiveCodingModule::ShouldPreloadModule(const FName& Name, const FString& Fu
 	}
 }
 
-void FLiveCodingModule::PreCompileHook()
-{
-	UE_LOG(LogLiveCoding, Display, TEXT("Starting Live Coding compile."));
-	GIsCompileActive = true;
-	BeginReload();
-}
-
-void FLiveCodingModule::PostCompileHook()
-{
-	UE_LOG(LogLiveCoding, Display, TEXT("Live Coding compile done.  See Live Coding console for more information."));
-	GIsCompileActive = false;
-}
-
 void FLiveCodingModule::BeginReload()
 {
 	if (GLiveCodingModule != nullptr)
@@ -669,4 +697,28 @@ void LiveCodingEndCompile()
 	GIsCompileActive = false;
 }
 
+// Invoked from LC_ClientCommandActions
+void LiveCodingPreCompile()
+{
+	UE_LOG(LogLiveCoding, Display, TEXT("Starting Live Coding compile."));
+	GIsCompileActive = true;
+	if (GLiveCodingModule != nullptr)
+	{
+		GLiveCodingModule->BeginReload();
+	}
+}
+
+// Invoked from LC_ClientCommandActions
+void LiveCodingPostCompile()
+{
+	UE_LOG(LogLiveCoding, Display, TEXT("Live Coding compile done.  See Live Coding console for more information."));
+	GIsCompileActive = false;
+}
+
+// Invoked from LC_ClientCommandActions
+void LiveCodingTriggerReload()
+{
+	GTriggerReload = true;
+}
+
 #undef LOCTEXT_NAMESPACE
diff --git a/Engine/Source/Developer/Windows/LiveCoding/Private/LiveCodingModule.h b/Engine/Source/Developer/Windows/LiveCoding/Private/LiveCodingModule.h
--- a/Engine/Source/Developer/Windows/LiveCoding/Private/LiveCodingModule.h
+++ b/Engine/Source/Developer/Windows/LiveCoding/Private/LiveCodingModule.h
@@ -41,9 +41,6 @@ public:
 	virtual void Tick() override;
 	virtual FOnPatchCompleteDelegate& GetOnPatchCompleteDelegate() override;
 
-	// Methods invoked by the LiveCodingConsole
-	static void PreCompileHook();
-	static void PostCompileHook();
 	static void BeginReload();
 
 private:
@@ -77,8 +74,6 @@ private:
 	TUniquePtr<IReload> Reload;
 #endif
 
-	void CheckForClassChanges();
-
 	bool StartLiveCoding();
 
 	void OnModulesChanged(FName ModuleName, EModuleChangeReason Reason);
diff --git a/Engine/Source/Developer/Windows/LiveCodingServer/Private/External/LC_Hook.h b/Engine/Source/Developer/Windows/LiveCodingServer/Private/External/LC_Hook.h
--- a/Engine/Source/Developer/Windows/LiveCodingServer/Private/External/LC_Hook.h
+++ b/Engine/Source/Developer/Windows/LiveCodingServer/Private/External/LC_Hook.h
@@ -27,10 +27,6 @@ namespace hook
 			COMPILE_SUCCESS,
 			COMPILE_ERROR,
 			COMPILE_ERROR_MESSAGE,
-			// BEGIN EPIC MOD - Add the ability for pre and post compile notifications
-			PRECOMPILE,
-			POSTCOMPILE
-			// END EPIC MOD
 		};
 	};
 
@@ -42,11 +38,6 @@ namespace hook
 	typedef void (*CompileErrorFunction)(void);
 	typedef void (*CompileErrorMessageFunction)(const wchar_t*);
 
-	// BEGIN EPIC MOD - Add the ability for pre and post compile notifications
-	typedef void (*PrecompileFunction)(void);
-	typedef void (*PostcompileFunction)(void);
-	// END EPIC MOD
-
 	uint32_t FindFirstInSection(const symbols::ImageSectionDB* imageSectionDb, const ImmutableString& sectionName);
 	uint32_t FindLastInSection(const symbols::ImageSectionDB* imageSectionDb, const ImmutableString& sectionName);
 
diff --git a/Engine/Source/Developer/Windows/LiveCodingServer/Private/External/LC_LiveModule.cpp b/Engine/Source/Developer/Windows/LiveCodingServer/Private/External/LC_LiveModule.cpp
--- a/Engine/Source/Developer/Windows/LiveCodingServer/Private/External/LC_LiveModule.cpp
+++ b/Engine/Source/Developer/Windows/LiveCodingServer/Private/External/LC_LiveModule.cpp
@@ -1627,6 +1627,18 @@ LiveModule::ErrorType::Enum LiveModule::Update(FileAttributeCache* fileCache, Di
 	const PerProcessData* processData = m_perProcessData.data();
 	const size_t processCount = m_perProcessData.size();
 
+	// BEGIN EPIC MOD
+	bool enableReinstancingFlow = false;
+	if (updateType != LiveModule::UpdateType::NO_CLIENT_COMMUNICATION)
+	{
+		for (size_t p = 0u; p < processCount; ++p)
+		{
+			const PerProcessData& data = processData[p];
+			enableReinstancingFlow |= data.liveProcess->IsReinstancingFlowEnabled();
+		}
+	}
+	// END EPIC MOD
+
 	// recompile changed files
 	if (m_runMode == RunMode::DEFAULT)
 	{
@@ -3693,124 +3705,135 @@ LiveModule::ErrorType::Enum LiveModule::Update(FileAttributeCache* fileCache, Di
 	GLiveCodingServer->GetStatusChangeDelegate().ExecuteIfBound(L"Patching relocations...");
 	// END EPIC MOD
 
-	LC_LOG_DEV("Patching relocations before calling entry point");
-
-	// walk all relocations in the .OBJ files, find their current locations in the .exe,
-	// and patch the relocations to point to the original symbols in the original .exe.
-	// we need to patch relocations *before* calling the DLL entry point, because global
-	// initializer code might refer to symbols that have been stripped by us.
-	// note that we only patch relocations to data symbols at this time, because functions haven't been
-	// hooked yet, and we need to ensure that dynamic initializers end up using new code paths (if available), while
-	// still referring to existing data symbols.
+	// BEGIN EPIC MOD
+	auto patchRelocationsPreEntryPoint = [&](bool forceBackwards)
 	{
-		telemetry::Scope patchingRelocationsScope("Patching relocations");
-
-		uint32_t relocationsHandledCount = 0u;
-		size_t relocationsCount = 0u;
+	// END EPIC MOD
+		LC_LOG_DEV("Patching relocations before calling entry point");
 
-		for (auto it = patch_compilandDB->compilands.begin(); it != patch_compilandDB->compilands.end(); ++it)
+		// walk all relocations in the .OBJ files, find their current locations in the .exe,
+		// and patch the relocations to point to the original symbols in the original .exe.
+		// we need to patch relocations *before* calling the DLL entry point, because global
+		// initializer code might refer to symbols that have been stripped by us.
+		// note that we only patch relocations to data symbols at this time, because functions haven't been
+		// hooked yet, and we need to ensure that dynamic initializers end up using new code paths (if available), while
+		// still referring to existing data symbols.
 		{
-			const symbols::ObjPath& objPath = it->first;
+			telemetry::Scope patchingRelocationsScope("Patching relocations");
 
-			LC_LOG_DEV("Patching relocations for file %s", objPath.c_str());
-			LC_LOG_INDENT_DEV;
+			uint32_t relocationsHandledCount = 0u;
+			size_t relocationsCount = 0u;
 
-			const coff::CoffDB* coffDb = m_coffCache->Lookup(objPath);
-			if (!coffDb)
+			for (auto it = patch_compilandDB->compilands.begin(); it != patch_compilandDB->compilands.end(); ++it)
 			{
-				LC_ERROR_USER("Could not find COFF database for file %s", objPath.c_str());
-				continue;
-			}
-
-			const std::wstring wideObjPath = string::ToWideString(objPath);
-
-			// note that we need to take the original compiland here, to make sure that single split-files get assigned
-			// the same unique ID as the corresponding amalgamated file.
-			symbols::Compiland* compiland = symbols::FindCompiland(m_compilandDB, objPath);
-			const uint32_t compilandUniqueId = symbols::GetCompilandId(compiland, wideObjPath.c_str(), modifiedOrNewObjFiles);
+				const symbols::ObjPath& objPath = it->first;
 
-			const types::StringSet& strippedSymbols = strippedSymbolsPerCompiland[objPath];
-			const types::StringSet& forceStrippedSymbols = forceStrippedSymbolsPerCompiland[objPath];
+				LC_LOG_DEV("Patching relocations for file %s", objPath.c_str());
+				LC_LOG_INDENT_DEV;
 
-			const size_t symbolCount = coffDb->symbols.size();
-			for (size_t i = 0u; i < symbolCount; ++i)
-			{
-				const coff::Symbol* symbol = coffDb->symbols[i];
-				relocationsCount += symbol->relocations.size();
-
-				// check if the patch knows this symbol.
-				// if not, it has probably been stripped and there is no need to walk all its relocations.
-				const ImmutableString& symbolName = symbols::TransformAnonymousNamespacePattern(coff::GetSymbolName(coffDb, symbol), compilandUniqueId);
-				const symbols::Symbol* realSymbol = symbols::FindSymbolByName(patch_symbolDB, symbolName);
-				if (!realSymbol)
+				const coff::CoffDB* coffDb = m_coffCache->Lookup(objPath);
+				if (!coffDb)
 				{
-					// this symbol has been stripped from the executable.
-					// in optimized builds, the compiler will sometimes e.g. leave a static function in an OBJ file,
-					// which will be kicked out by the linker.
+					LC_ERROR_USER("Could not find COFF database for file %s", objPath.c_str());
 					continue;
 				}
 
-				// before patching relocations, check whether the symbol which relocations we want to patch originated from
-				// a compiland that is the same as the file we're working on.
-				// this might not be the case, especially when using static libraries, COMDATs, and compilands that use the
-				// same inline function but have slightly different compiler options (/hotpatch vs. no /hotpatch, e.g.
-				// __local_stdio_printf_options in the main module vs. in the dynamic runtime)
-				const symbols::Contribution* originalContribution = symbols::FindContributionByRVA(patch_contributionDB, realSymbol->rva);
-				if (originalContribution)
-				{
-					const ImmutableString& compilandName = symbols::GetContributionCompilandName(patch_contributionDB, originalContribution);
-					if (compilandName != objPath)
-					{
-						LC_LOG_DEV("Ignoring relocations for symbol %s in file %s (original compiland: %s)",
-							symbolName.c_str(), objPath.c_str(), compilandName.c_str());
-						continue;
-					}
-				}
+				const std::wstring wideObjPath = string::ToWideString(objPath);
+
+				// note that we need to take the original compiland here, to make sure that single split-files get assigned
+				// the same unique ID as the corresponding amalgamated file.
+				symbols::Compiland* compiland = symbols::FindCompiland(m_compilandDB, objPath);
+				const uint32_t compilandUniqueId = symbols::GetCompilandId(compiland, wideObjPath.c_str(), modifiedOrNewObjFiles);
+
+				const types::StringSet& strippedSymbols = strippedSymbolsPerCompiland[objPath];
+				const types::StringSet& forceStrippedSymbols = forceStrippedSymbolsPerCompiland[objPath];
 
-				const size_t relocationCount = symbol->relocations.size();
-				for (size_t j = 0u; j < relocationCount; ++j)
+				const size_t symbolCount = coffDb->symbols.size();
+				for (size_t i = 0u; i < symbolCount; ++i)
 				{
-					const coff::Relocation* relocation = symbol->relocations[j];
-					if (relocation->dstIsSection)
+					const coff::Symbol* symbol = coffDb->symbols[i];
+					relocationsCount += symbol->relocations.size();
+
+					// check if the patch knows this symbol.
+					// if not, it has probably been stripped and there is no need to walk all its relocations.
+					const ImmutableString& symbolName = symbols::TransformAnonymousNamespacePattern(coff::GetSymbolName(coffDb, symbol), compilandUniqueId);
+					const symbols::Symbol* realSymbol = symbols::FindSymbolByName(patch_symbolDB, symbolName);
+					if (!realSymbol)
 					{
-						// don't patch relocations to sections
+						// this symbol has been stripped from the executable.
+						// in optimized builds, the compiler will sometimes e.g. leave a static function in an OBJ file,
+						// which will be kicked out by the linker.
 						continue;
 					}
 
-					// ignore relocations to symbols in .msvcjmc (MSVC JustMyCode) sections
-					if (relocation->dstSectionIndex >= 0)
+					// before patching relocations, check whether the symbol which relocations we want to patch originated from
+					// a compiland that is the same as the file we're working on.
+					// this might not be the case, especially when using static libraries, COMDATs, and compilands that use the
+					// same inline function but have slightly different compiler options (/hotpatch vs. no /hotpatch, e.g.
+					// __local_stdio_printf_options in the main module vs. in the dynamic runtime)
+					const symbols::Contribution* originalContribution = symbols::FindContributionByRVA(patch_contributionDB, realSymbol->rva);
+					if (originalContribution)
 					{
-						const uint32_t index = static_cast<uint32_t>(relocation->dstSectionIndex);
-						const coff::Section& section = coffDb->sections[index];
-						if (coff::IsMSVCJustMyCodeSection(section.name.c_str()))
+						const ImmutableString& compilandName = symbols::GetContributionCompilandName(patch_contributionDB, originalContribution);
+						if (compilandName != objPath)
 						{
-							LC_LOG_DEV("Ignoring relocation to symbol in section %s", section.name.c_str());
+							LC_LOG_DEV("Ignoring relocations for symbol %s in file %s (original compiland: %s)",
+								symbolName.c_str(), objPath.c_str(), compilandName.c_str());
 							continue;
 						}
 					}
 
-					ImmutableString dstSymbolName = symbols::TransformAnonymousNamespacePattern(GetRelocationDstSymbolName(coffDb, relocation), compilandUniqueId);
-					const bool refersToDataSymbol = !coff::IsFunctionSymbol(coff::GetRelocationDstSymbolType(relocation));
-					const bool refersToStrippedSymbol = (strippedSymbols.find(dstSymbolName) != strippedSymbols.end());
-					if (refersToDataSymbol || refersToStrippedSymbol)
+					const size_t relocationCount = symbol->relocations.size();
+					for (size_t j = 0u; j < relocationCount; ++j)
 					{
-						const relocations::Record& relocationRecord = relocations::PatchRelocation(relocation, coffDb, forceStrippedSymbols, m_moduleCache, symbolName, dstSymbolName, realSymbol, token, &loadedPatches[0]);
-						if (relocations::IsValidRecord(relocationRecord))
+						const coff::Relocation* relocation = symbol->relocations[j];
+						if (relocation->dstIsSection)
+						{
+							// don't patch relocations to sections
+							continue;
+						}
+
+						// ignore relocations to symbols in .msvcjmc (MSVC JustMyCode) sections
+						if (relocation->dstSectionIndex >= 0)
 						{
-							compiledModulePatch->RegisterPreEntryPointRelocation(relocationRecord);
+							const uint32_t index = static_cast<uint32_t>(relocation->dstSectionIndex);
+							const coff::Section& section = coffDb->sections[index];
+							if (coff::IsMSVCJustMyCodeSection(section.name.c_str()))
+							{
+								LC_LOG_DEV("Ignoring relocation to symbol in section %s", section.name.c_str());
+								continue;
+							}
 						}
 
-						++relocationsHandledCount;
+						ImmutableString dstSymbolName = symbols::TransformAnonymousNamespacePattern(GetRelocationDstSymbolName(coffDb, relocation), compilandUniqueId);
+						const bool refersToDataSymbol = !coff::IsFunctionSymbol(coff::GetRelocationDstSymbolType(relocation));
+						const bool refersToStrippedSymbol = (strippedSymbols.find(dstSymbolName) != strippedSymbols.end());
+						if (refersToDataSymbol || refersToStrippedSymbol)
+						{
+							// BEGIN EPIC MOD
+							const relocations::Record& relocationRecord = relocations::PatchRelocation(relocation, coffDb, forceStrippedSymbols, m_moduleCache, symbolName, dstSymbolName, realSymbol, token, &loadedPatches[0], forceBackwards);
+							if (!forceBackwards && relocations::IsValidRecord(relocationRecord))
+							{
+								compiledModulePatch->RegisterPreEntryPointRelocation(relocationRecord);
+							}
+							// END EPIC MOD
+
+							++relocationsHandledCount;
+						}
 					}
 				}
 			}
-		}
-
-		LC_LOG_TELEMETRY("Handled %d of %d relocations in %.3fms (avg: %.3fus)", relocationsHandledCount, relocationsCount,
-			patchingRelocationsScope.ReadMilliSeconds(), (patchingRelocationsScope.ReadMicroSeconds() / relocationsHandledCount));
-	}
 
+			LC_LOG_TELEMETRY("Handled %d of %d relocations in %.3fms (avg: %.3fus)", relocationsHandledCount, relocationsCount,
+				patchingRelocationsScope.ReadMilliSeconds(), (patchingRelocationsScope.ReadMicroSeconds() / relocationsHandledCount));
+		}
+	// BEGIN EPIC MOD
+	};
+	// END EPIC MOD
 
+	// BEGIN EPIC MOD
+	patchRelocationsPreEntryPoint(enableReinstancingFlow);
+	// END EPIC MOD
 
 	// BEGIN EPIC MOD
 	GLiveCodingServer->GetStatusChangeDelegate().ExecuteIfBound(L"Patching public functions in lib compilands...");
@@ -4165,257 +4188,284 @@ LiveModule::ErrorType::Enum LiveModule::Update(FileAttributeCache* fileCache, Di
 	}
 
 
-
-	// dynamic initializers have run, patch the remaining relocations
 	// BEGIN EPIC MOD
-	GLiveCodingServer->GetStatusChangeDelegate().ExecuteIfBound(L"Patching remaining relocations...");
-	// END EPIC MOD
+	auto patchRelocations = [&](bool forceBackwards)
 	{
-		telemetry::Scope patchingRelocationsScope("Patching remaining relocations");
-
-		uint32_t relocationsHandledCount = 0u;
-		size_t relocationsCount = 0u;
-
-		LC_LOG_DEV("Patching relocations after calling entry point");
-
-		for (auto it = patch_compilandDB->compilands.begin(); it != patch_compilandDB->compilands.end(); ++it)
+	// END EPIC MOD
+		// dynamic initializers have run, patch the remaining relocations
+		// BEGIN EPIC MOD
+		GLiveCodingServer->GetStatusChangeDelegate().ExecuteIfBound(L"Patching remaining relocations...");
+		// END EPIC MOD
 		{
-			const symbols::ObjPath& objPath = it->first;
+			telemetry::Scope patchingRelocationsScope("Patching remaining relocations");
 
-			LC_LOG_DEV("Patching relocations for file %s", objPath.c_str());
-			LC_LOG_INDENT_DEV;
+			uint32_t relocationsHandledCount = 0u;
+			size_t relocationsCount = 0u;
 
-			const coff::CoffDB* coffDb = m_coffCache->Lookup(objPath);
-			if (!coffDb)
-			{
-				LC_ERROR_USER("Could not find COFF database for file %s", objPath.c_str());
-				continue;
-			}
-
-			const std::wstring wideObjPath = string::ToWideString(objPath);
+			LC_LOG_DEV("Patching relocations after calling entry point");
 
-			// note that we need to take the original compiland here, to make sure that single split-files get assigned
-			// the same unique ID as the corresponding amalgamated file.
-			symbols::Compiland* compiland = symbols::FindCompiland(m_compilandDB, objPath);
-			const uint32_t compilandUniqueId = symbols::GetCompilandId(compiland, wideObjPath.c_str(), modifiedOrNewObjFiles);
+			for (auto it = patch_compilandDB->compilands.begin(); it != patch_compilandDB->compilands.end(); ++it)
+			{
+				const symbols::ObjPath& objPath = it->first;
 
-			const types::StringSet& strippedSymbols = strippedSymbolsPerCompiland[objPath];
-			const types::StringSet& forceStrippedSymbols = forceStrippedSymbolsPerCompiland[objPath];
+				LC_LOG_DEV("Patching relocations for file %s", objPath.c_str());
+				LC_LOG_INDENT_DEV;
 
-			const size_t symbolCount = coffDb->symbols.size();
-			for (size_t i = 0u; i < symbolCount; ++i)
-			{
-				const coff::Symbol* symbol = coffDb->symbols[i];
-				relocationsCount += symbol->relocations.size();
-
-				// check if the patch knows this symbol.
-				// if not, it has probably been stripped and there is no need to walk all its relocations.
-				const ImmutableString& symbolName = coff::GetSymbolName(coffDb, symbol);
-				const symbols::Symbol* realSymbol = symbols::FindSymbolByName(patch_symbolDB, symbolName);
-				if (!realSymbol)
+				const coff::CoffDB* coffDb = m_coffCache->Lookup(objPath);
+				if (!coffDb)
 				{
-					// this symbol has been stripped from the executable.
-					// in optimized builds, the compiler will sometimes e.g. leave a static function in an OBJ file,
-					// which will be kicked out by the linker.
+					LC_ERROR_USER("Could not find COFF database for file %s", objPath.c_str());
 					continue;
 				}
 
-				// before patching relocations, check whether the symbol which relocations we want to patch originated from
-				// a compiland that is the same as the file we're working on.
-				// this might not be the case, especially when using static libraries, COMDATs, and compilands that use the
-				// same inline function but have slightly different compiler options (/hotpatch vs. no /hotpatch, e.g.
-				// __local_stdio_printf_options in the main module vs. in the dynamic runtime)
-				const symbols::Contribution* originalContribution = symbols::FindContributionByRVA(patch_contributionDB, realSymbol->rva);
-				if (originalContribution)
-				{
-					const ImmutableString& compilandName = symbols::GetContributionCompilandName(patch_contributionDB, originalContribution);
-					if (compilandName != objPath)
-					{
-						LC_LOG_DEV("Ignoring relocations for symbol %s in file %s (original compiland: %s)",
-							symbolName.c_str(), objPath.c_str(), compilandName.c_str());
-						continue;
-					}
-				}
+				const std::wstring wideObjPath = string::ToWideString(objPath);
+
+				// note that we need to take the original compiland here, to make sure that single split-files get assigned
+				// the same unique ID as the corresponding amalgamated file.
+				symbols::Compiland* compiland = symbols::FindCompiland(m_compilandDB, objPath);
+				const uint32_t compilandUniqueId = symbols::GetCompilandId(compiland, wideObjPath.c_str(), modifiedOrNewObjFiles);
+
+				const types::StringSet& strippedSymbols = strippedSymbolsPerCompiland[objPath];
+				const types::StringSet& forceStrippedSymbols = forceStrippedSymbolsPerCompiland[objPath];
 
-				const size_t relocationCount = symbol->relocations.size();
-				for (size_t j = 0u; j < relocationCount; ++j)
+				const size_t symbolCount = coffDb->symbols.size();
+				for (size_t i = 0u; i < symbolCount; ++i)
 				{
-					const coff::Relocation* relocation = symbol->relocations[j];
-					if (relocation->dstIsSection)
+					const coff::Symbol* symbol = coffDb->symbols[i];
+					relocationsCount += symbol->relocations.size();
+
+					// check if the patch knows this symbol.
+					// if not, it has probably been stripped and there is no need to walk all its relocations.
+					const ImmutableString& symbolName = coff::GetSymbolName(coffDb, symbol);
+					const symbols::Symbol* realSymbol = symbols::FindSymbolByName(patch_symbolDB, symbolName);
+					if (!realSymbol)
 					{
-						// don't patch relocations to sections
+						// this symbol has been stripped from the executable.
+						// in optimized builds, the compiler will sometimes e.g. leave a static function in an OBJ file,
+						// which will be kicked out by the linker.
 						continue;
 					}
 
-					// ignore relocations to symbols in .msvcjmc (MSVC JustMyCode) sections
-					if (relocation->dstSectionIndex >= 0)
+					// before patching relocations, check whether the symbol which relocations we want to patch originated from
+					// a compiland that is the same as the file we're working on.
+					// this might not be the case, especially when using static libraries, COMDATs, and compilands that use the
+					// same inline function but have slightly different compiler options (/hotpatch vs. no /hotpatch, e.g.
+					// __local_stdio_printf_options in the main module vs. in the dynamic runtime)
+					const symbols::Contribution* originalContribution = symbols::FindContributionByRVA(patch_contributionDB, realSymbol->rva);
+					if (originalContribution)
 					{
-						const uint32_t index = static_cast<uint32_t>(relocation->dstSectionIndex);
-						const coff::Section& section = coffDb->sections[index];
-						if (coff::IsMSVCJustMyCodeSection(section.name.c_str()))
+						const ImmutableString& compilandName = symbols::GetContributionCompilandName(patch_contributionDB, originalContribution);
+						if (compilandName != objPath)
 						{
-							LC_LOG_DEV("Ignoring relocation to symbol in section %s", section.name.c_str());
+							LC_LOG_DEV("Ignoring relocations for symbol %s in file %s (original compiland: %s)",
+								symbolName.c_str(), objPath.c_str(), compilandName.c_str());
 							continue;
 						}
 					}
 
-					ImmutableString dstSymbolName = symbols::TransformAnonymousNamespacePattern(GetRelocationDstSymbolName(coffDb, relocation), compilandUniqueId);
-
-					// relocations to data symbols and stripped symbols have already been done
-					const bool refersToFunctionSymbol = coff::IsFunctionSymbol(coff::GetRelocationDstSymbolType(relocation));
-					const bool refersToStrippedSymbol = (strippedSymbols.find(dstSymbolName) != strippedSymbols.end());
-					if (refersToFunctionSymbol && !refersToStrippedSymbol)
+					const size_t relocationCount = symbol->relocations.size();
+					for (size_t j = 0u; j < relocationCount; ++j)
 					{
-						const relocations::Record& relocationRecord = relocations::PatchRelocation(relocation, coffDb, forceStrippedSymbols, m_moduleCache, symbolName, dstSymbolName, realSymbol, token, &loadedPatches[0]);
-						if (relocations::IsValidRecord(relocationRecord))
+						const coff::Relocation* relocation = symbol->relocations[j];
+						if (relocation->dstIsSection)
 						{
-							compiledModulePatch->RegisterPostEntryPointRelocation(relocationRecord);
+							// don't patch relocations to sections
+							continue;
 						}
 
-						++relocationsHandledCount;
+						// ignore relocations to symbols in .msvcjmc (MSVC JustMyCode) sections
+						if (relocation->dstSectionIndex >= 0)
+						{
+							const uint32_t index = static_cast<uint32_t>(relocation->dstSectionIndex);
+							const coff::Section& section = coffDb->sections[index];
+							if (coff::IsMSVCJustMyCodeSection(section.name.c_str()))
+							{
+								LC_LOG_DEV("Ignoring relocation to symbol in section %s", section.name.c_str());
+								continue;
+							}
+						}
+
+						ImmutableString dstSymbolName = symbols::TransformAnonymousNamespacePattern(GetRelocationDstSymbolName(coffDb, relocation), compilandUniqueId);
+
+						// relocations to data symbols and stripped symbols have already been done
+						const bool refersToFunctionSymbol = coff::IsFunctionSymbol(coff::GetRelocationDstSymbolType(relocation));
+						const bool refersToStrippedSymbol = (strippedSymbols.find(dstSymbolName) != strippedSymbols.end());
+						if (refersToFunctionSymbol && !refersToStrippedSymbol)
+						{
+							// BEGIN EPIC MOD
+							const relocations::Record& relocationRecord = relocations::PatchRelocation(relocation, coffDb, forceStrippedSymbols, m_moduleCache, symbolName, dstSymbolName, realSymbol, token, &loadedPatches[0], forceBackwards);
+							if (!forceBackwards && relocations::IsValidRecord(relocationRecord))
+							{
+								compiledModulePatch->RegisterPostEntryPointRelocation(relocationRecord);
+							}
+							// END EPIC MOD
+
+							++relocationsHandledCount;
+						}
 					}
 				}
 			}
-		}
-
-		LC_LOG_TELEMETRY("Handled %d of %d remaining relocations in %.3fms (avg: %.3fus)", relocationsHandledCount, relocationsCount,
-			patchingRelocationsScope.ReadMilliSeconds(), (patchingRelocationsScope.ReadMicroSeconds() / relocationsHandledCount));
-	}
-
 
+			LC_LOG_TELEMETRY("Handled %d of %d remaining relocations in %.3fms (avg: %.3fus)", relocationsHandledCount, relocationsCount,
+				patchingRelocationsScope.ReadMilliSeconds(), (patchingRelocationsScope.ReadMicroSeconds() / relocationsHandledCount));
+		}
+	// BEGIN EPIC MOD
+	};
+	// END EPIC MOD
 
 	// BEGIN EPIC MOD
-	GLiveCodingServer->GetStatusChangeDelegate().ExecuteIfBound(L"Patching functions...");
+	patchRelocations(enableReinstancingFlow);
 	// END EPIC MOD
 
-	// suspend the main processes before patching functions, because they might not use synchronization points.
-	for (size_t p = 0u; p < processCount; ++p)
+	// BEGIN EPIC MOD
+	auto patchFunctions = [&](bool forceBackwards)
+	// END EPIC MOD
 	{
-		const PerProcessData& data = processData[p];
-		Process::Suspend(data.liveProcess->GetProcessHandle());
-	}
-
 
-	// determining which functions have changed (or lead to a different execution path) would be very hard
-	// to do, therefore we hook all functions.
-	// even though internal functions can only be referenced from external ones, it is not enough to hook
-	// only those. the reason for that is that global/static instances might refer to internal functions
-	// by function-pointer, address, etc., so internal functions must also be hooked.
-	{
-		telemetry::Scope patchingFunctionsScope("Patching functions");
+		// BEGIN EPIC MOD
+		GLiveCodingServer->GetStatusChangeDelegate().ExecuteIfBound(L"Patching functions...");
+		// END EPIC MOD
 
-		// the processes are all halted. fetch instruction pointers from all their threads.
-		typedef types::vector<const void*> PerProcessThreadIPs;
-		types::vector<PerProcessThreadIPs> processThreadIPs;
-		processThreadIPs.reserve(processCount);
+		// suspend the main processes before patching functions, because they might not use synchronization points.
 		for (size_t p = 0u; p < processCount; ++p)
 		{
 			const PerProcessData& data = processData[p];
-			processThreadIPs.emplace_back(EnumerateInstructionPointers(data.liveProcess->GetProcessId()));
+			Process::Suspend(data.liveProcess->GetProcessHandle());
 		}
 
-		uint32_t functionsPatchedCount = 0u;
-		size_t functionsCount = 0u;
 
-		// we deliberately do not hook functions in lib compilands because they cannot have changed, per definition.
-		// they are part of a static library that won't be recompiled during a Live++ session.
-		for (auto it = patch_compilandDB->compilands.begin(); it != patch_compilandDB->compilands.end(); ++it)
+		// determining which functions have changed (or lead to a different execution path) would be very hard
+		// to do, therefore we hook all functions.
+		// even though internal functions can only be referenced from external ones, it is not enough to hook
+		// only those. the reason for that is that global/static instances might refer to internal functions
+		// by function-pointer, address, etc., so internal functions must also be hooked.
 		{
-			const symbols::ObjPath& objPath = it->first;
+			telemetry::Scope patchingFunctionsScope("Patching functions");
 
-			LC_LOG_DEV("Patching functions for file %s", objPath.c_str());
-			LC_LOG_INDENT_DEV;
-
-			const coff::CoffDB* coffDb = m_coffCache->Lookup(objPath);
-			if (!coffDb)
+			// the processes are all halted. fetch instruction pointers from all their threads.
+			typedef types::vector<const void*> PerProcessThreadIPs;
+			types::vector<PerProcessThreadIPs> processThreadIPs;
+			processThreadIPs.reserve(processCount);
+			for (size_t p = 0u; p < processCount; ++p)
 			{
-				LC_ERROR_USER("Could not find COFF database for file %s", objPath.c_str());
-				continue;
+				const PerProcessData& data = processData[p];
+				processThreadIPs.emplace_back(EnumerateInstructionPointers(data.liveProcess->GetProcessId()));
 			}
 
-			const std::wstring& wideObjPath = string::ToWideString(objPath);
-			const symbols::Compiland* compiland = it->second;
-			const uint32_t compilandUniqueId = symbols::GetCompilandId(compiland, wideObjPath.c_str(), modifiedOrNewObjFiles);
+			uint32_t functionsPatchedCount = 0u;
+			size_t functionsCount = 0u;
 
-			for (size_t i = 0u; i < coffDb->symbols.size(); ++i)
+			// we deliberately do not hook functions in lib compilands because they cannot have changed, per definition.
+			// they are part of a static library that won't be recompiled during a Live++ session.
+			for (auto it = patch_compilandDB->compilands.begin(); it != patch_compilandDB->compilands.end(); ++it)
 			{
-				const coff::Symbol* symbol = coffDb->symbols[i];
-				if (!coff::IsFunctionSymbol(symbol->type))
-				{
-					continue;
-				}
-
-				++functionsCount;
+				const symbols::ObjPath& objPath = it->first;
 
-				const ImmutableString& functionName = symbols::TransformAnonymousNamespacePattern(coff::GetSymbolName(coffDb, symbol), compilandUniqueId);
-				if (symbols::IsExceptionRelatedSymbol(functionName))
-				{
-					LC_LOG_DEV("Ignoring exception-related function %s", functionName.c_str());
-					continue;
-				}
+				LC_LOG_DEV("Patching functions for file %s", objPath.c_str());
+				LC_LOG_INDENT_DEV;
 
-				const symbols::Symbol* patchSymbol = symbols::FindSymbolByName(patch_symbolDB, functionName);
-				if (!patchSymbol)
+				const coff::CoffDB* coffDb = m_coffCache->Lookup(objPath);
+				if (!coffDb)
 				{
-					LC_WARNING_DEV("Cannot find function %s in patch, possibly stripped by linker", functionName.c_str());
+					LC_ERROR_USER("Could not find COFF database for file %s", objPath.c_str());
 					continue;
 				}
 
-				const ModuleCache::FindSymbolData& originalData = m_moduleCache->FindSymbolByName(token, functionName);
-				if (!originalData.symbol)
-				{
-					LC_LOG_DEV("Ignoring new function %s", functionName.c_str());
-					continue;
-				}
+				const std::wstring& wideObjPath = string::ToWideString(objPath);
+				const symbols::Compiland* compiland = it->second;
+				const uint32_t compilandUniqueId = symbols::GetCompilandId(compiland, wideObjPath.c_str(), modifiedOrNewObjFiles);
 
-				// if the original function to be patched did not come from a compiland, it cannot possibly have changed and
-				// therefore can be ignored.
-				const symbols::Contribution* originalContribution = symbols::FindContributionByRVA(originalData.data->contributionDb, originalData.symbol->rva);
-				if (originalContribution)
+				for (size_t i = 0u; i < coffDb->symbols.size(); ++i)
 				{
-					const ImmutableString& compilandName = symbols::GetContributionCompilandName(originalData.data->contributionDb, originalContribution);
-					const symbols::Compiland* originalCompiland = symbols::FindCompiland(originalData.data->compilandDb, compilandName);
-					if (!originalCompiland)
+					const coff::Symbol* symbol = coffDb->symbols[i];
+					if (!coff::IsFunctionSymbol(symbol->type))
 					{
-						LC_LOG_DEV("Ignoring function %s originally contributed from lib compiland %s", functionName.c_str(), compilandName.c_str());
 						continue;
 					}
-				}
 
-				const size_t moduleProcessCount = originalData.data->processes.size();
-				for (size_t p = 0u; p < moduleProcessCount; ++p)
-				{
-					++functionsPatchedCount;
+					++functionsCount;
 
-					const ModuleCache::ProcessData& hookProcessData = originalData.data->processes[p];
+					const ImmutableString& functionName = symbols::TransformAnonymousNamespacePattern(coff::GetSymbolName(coffDb, symbol), compilandUniqueId);
+					if (symbols::IsExceptionRelatedSymbol(functionName))
+					{
+						LC_LOG_DEV("Ignoring exception-related function %s", functionName.c_str());
+						continue;
+					}
 
-					const Process::Id pid = hookProcessData.processId;
-					void* moduleBase = hookProcessData.moduleBase;
-					Process::Handle processHandle = hookProcessData.processHandle;
+					const symbols::Symbol* patchSymbol = symbols::FindSymbolByName(patch_symbolDB, functionName);
+					if (!patchSymbol)
+					{
+						LC_WARNING_DEV("Cannot find function %s in patch, possibly stripped by linker", functionName.c_str());
+						continue;
+					}
 
-					char* originalAddress = pointer::Offset<char*>(moduleBase, originalData.symbol->rva);
-					char* patchAddress = pointer::Offset<char*>(loadedPatches[p], patchSymbol->rva);
-					types::unordered_set<const void*>& patchedAddresses = m_patchedAddressesPerProcess[pid];
+					const ModuleCache::FindSymbolData& originalData = m_moduleCache->FindSymbolByName(token, functionName);
+					if (!originalData.symbol)
+					{
+						LC_LOG_DEV("Ignoring new function %s", functionName.c_str());
+						continue;
+					}
 
-					const functions::Record& record = functions::PatchFunction(originalAddress, patchAddress, originalData.symbol->rva, patchSymbol->rva,
-						originalData.data->thunkDb, originalContribution, processHandle, moduleBase, originalData.data->index,
-						patchedAddresses, processThreadIPs[p], pid, functionName.c_str());
+					// if the original function to be patched did not come from a compiland, it cannot possibly have changed and
+					// therefore can be ignored.
+					const symbols::Contribution* originalContribution = symbols::FindContributionByRVA(originalData.data->contributionDb, originalData.symbol->rva);
+					if (originalContribution)
+					{
+						const ImmutableString& compilandName = symbols::GetContributionCompilandName(originalData.data->contributionDb, originalContribution);
+						const symbols::Compiland* originalCompiland = symbols::FindCompiland(originalData.data->compilandDb, compilandName);
+						if (!originalCompiland)
+						{
+							LC_LOG_DEV("Ignoring function %s originally contributed from lib compiland %s", functionName.c_str(), compilandName.c_str());
+							continue;
+						}
+					}
 
-					if (functions::IsValidRecord(record))
+					const size_t moduleProcessCount = originalData.data->processes.size();
+					for (size_t p = 0u; p < moduleProcessCount; ++p)
 					{
-						compiledModulePatch->RegisterFunctionPatch(record);
+						++functionsPatchedCount;
+
+						const ModuleCache::ProcessData& hookProcessData = originalData.data->processes[p];
+
+						const Process::Id pid = hookProcessData.processId;
+						void* moduleBase = hookProcessData.moduleBase;
+						Process::Handle processHandle = hookProcessData.processHandle;
+
+						char* originalAddress = pointer::Offset<char*>(moduleBase, originalData.symbol->rva);
+						char* patchAddress = pointer::Offset<char*>(loadedPatches[p], patchSymbol->rva);
+						types::unordered_set<const void*>& patchedAddresses = m_patchedAddressesPerProcess[pid];
+
+						const functions::Record& record = functions::PatchFunction(originalAddress, patchAddress, originalData.symbol->rva, patchSymbol->rva,
+							originalData.data->thunkDb, originalContribution, processHandle, moduleBase, originalData.data->index,
+							patchedAddresses, processThreadIPs[p], pid, functionName.c_str());
+
+						// BEGIN EPIC MOD
+						if (!forceBackwards && functions::IsValidRecord(record))
+						{
+							compiledModulePatch->RegisterFunctionPatch(record);
+						}
+						// END EPIC MOD
 					}
 				}
 			}
 		}
-	}
 
-	// resume the main processes again
-	for (size_t p = 0u; p < processCount; ++p)
+		// resume the main processes again
+		for (size_t p = 0u; p < processCount; ++p)
+		{
+			const PerProcessData& data = processData[p];
+			Process::Resume(data.liveProcess->GetProcessHandle());
+		}
+	// BEGIN EPIC MOD
+	};
+	// END EPIC MOD
+
+	// BEGIN EPIC MOD
+	if (!enableReinstancingFlow)
 	{
-		const PerProcessData& data = processData[p];
-		Process::Resume(data.liveProcess->GetProcessHandle());
+		patchFunctions(enableReinstancingFlow);
 	}
+	// END EPIC MOD
 
 	{
 		// post-patch hooks must be called on the current executable because the hooks want to use the newest memory layout of
@@ -4441,6 +4491,26 @@ LiveModule::ErrorType::Enum LiveModule::Update(FileAttributeCache* fileCache, Di
 		}
 	}
 
+	// BEGIN EPIC MOD
+	// pulse the sync point in all processes
+	if (enableReinstancingFlow)
+	{
+		if (updateType != LiveModule::UpdateType::NO_CLIENT_COMMUNICATION)
+		{
+			for (size_t p = 0u; p < processCount; ++p)
+			{
+				const PerProcessData& data = processData[p];
+				data.liveProcess->GetPipe()->SendCommandAndWaitForAck(commands::TriggerReload{}, nullptr, 0u);
+			}
+		}
+
+		// Re-patch everything so that things are patched as normally expected
+		patchRelocationsPreEntryPoint(false);
+		patchRelocations(false);
+		patchFunctions(false);
+	}
+	// END EPIC MOD
+
 	// leave sync point in all processes
 	if (updateType != LiveModule::UpdateType::NO_CLIENT_COMMUNICATION)
 	{
@@ -4793,45 +4863,3 @@ void LiveModule::OnCompiledFile(const symbols::ObjPath& objPath, symbols::Compil
 		LC_ERROR_USER("Failed to compile %s (%.3fs) (Exit code: 0x%X)", objPath.c_str(), compileTime, compileResult.exitCode);
 	}
 }
-
-// BEGIN EPIC MOD - Add the ability for pre and post compile notifications
-void LiveModule::CallPrecompileHooks()
-{
-	const ModuleCache::FindHookData& hookData = m_moduleCache->FindHooksInSectionBackwards(ModuleCache::SEARCH_ALL_MODULES, ImmutableString(LPP_PRECOMPILE_SECTION));
-	if ((hookData.firstRva != 0u) && (hookData.lastRva != 0u))
-	{
-		const size_t count = hookData.data->processes.size();
-		for (size_t p = 0u; p < count; ++p)
-		{
-			const ModuleCache::ProcessData& processData = hookData.data->processes[p];
-
-			const Process::Id pid = processData.processId;
-			void* moduleBase = processData.moduleBase;
-			const DuplexPipe* pipe = processData.pipe;
-
-			LC_LOG_USER("Calling precompile hooks (PID: %d)", pid);
-			pipe->SendCommandAndWaitForAck(MakeCallHooksCommand(hook::Type::PRECOMPILE, moduleBase, hookData), nullptr, 0u);
-		}
-	}
-}
-
-void LiveModule::CallPostcompileHooks()
-{
-	const ModuleCache::FindHookData& hookData = m_moduleCache->FindHooksInSectionBackwards(ModuleCache::SEARCH_ALL_MODULES, ImmutableString(LPP_POSTCOMPILE_SECTION));
-	if ((hookData.firstRva != 0u) && (hookData.lastRva != 0u))
-	{
-		const size_t count = hookData.data->processes.size();
-		for (size_t p = 0u; p < count; ++p)
-		{
-			const ModuleCache::ProcessData& processData = hookData.data->processes[p];
-
-			const Process::Id pid = processData.processId;
-			void* moduleBase = processData.moduleBase;
-			const DuplexPipe* pipe = processData.pipe;
-
-			LC_LOG_USER("Calling postcompile hooks (PID: %d)", pid);
-			pipe->SendCommandAndWaitForAck(MakeCallHooksCommand(hook::Type::POSTCOMPILE, moduleBase, hookData), nullptr, 0u);
-		}
-	}
-}
-// END EPIC MOD
\ No newline at end of file
diff --git a/Engine/Source/Developer/Windows/LiveCodingServer/Private/External/LC_LiveModule.h b/Engine/Source/Developer/Windows/LiveCodingServer/Private/External/LC_LiveModule.h
--- a/Engine/Source/Developer/Windows/LiveCodingServer/Private/External/LC_LiveModule.h
+++ b/Engine/Source/Developer/Windows/LiveCodingServer/Private/External/LC_LiveModule.h
@@ -90,11 +90,6 @@ public:
 
 	bool HasInstalledPatches(void) const;
 
-	// BEGIN EPIC MOD - Add the ability for pre and post compile notifications
-	void CallPrecompileHooks();
-	void CallPostcompileHooks();
-	// END EPIC MOD
-
 private:
 	void UpdateDirectoryCache(const ImmutableString& path, symbols::Dependency* dependency, DirectoryCache* cache);
 
diff --git a/Engine/Source/Developer/Windows/LiveCodingServer/Private/External/LC_LiveProcess.h b/Engine/Source/Developer/Windows/LiveCodingServer/Private/External/LC_LiveProcess.h
--- a/Engine/Source/Developer/Windows/LiveCodingServer/Private/External/LC_LiveProcess.h
+++ b/Engine/Source/Developer/Windows/LiveCodingServer/Private/External/LC_LiveProcess.h
@@ -97,6 +97,17 @@ public:
 	Windows::HMODULE GetLazyLoadedModuleBase(const std::wstring& moduleName) const;
 	// END EPIC MOD
 
+	// BEGIN EPIC MOD
+	void EnableReinstancingFlow()
+	{
+		m_reinstancingFlow = true;
+	}
+	bool IsReinstancingFlowEnabled() const
+	{
+		return m_reinstancingFlow;
+	}
+	// END EPIC MOD
+
 private:
 	Process::Handle m_processHandle;
 	Process::Id m_processId;
@@ -123,6 +134,10 @@ private:
 	types::unordered_map<std::wstring, LazyLoadedModule> m_lazyLoadedModules;
 	// END EPIC MOD
 
+	// BEGIN EPIC MOD
+	bool m_reinstancingFlow = false;
+	// END EPIC MOD
+
 	// loaded modules are not identified by their full path, but by their executable image header.
 	// we do this to ensure that the same executable loaded from a different path is not treated as
 	// a different executable.
diff --git a/Engine/Source/Developer/Windows/LiveCodingServer/Private/External/LC_RelocationPatcher.cpp b/Engine/Source/Developer/Windows/LiveCodingServer/Private/External/LC_RelocationPatcher.cpp
--- a/Engine/Source/Developer/Windows/LiveCodingServer/Private/External/LC_RelocationPatcher.cpp
+++ b/Engine/Source/Developer/Windows/LiveCodingServer/Private/External/LC_RelocationPatcher.cpp
@@ -118,6 +118,9 @@ relocations::Record relocations::PatchRelocation
 	const symbols::Symbol* srcSymbol,
 	size_t newModuleIndex,
 	void* newModuleBases[]
+	// BEGIN EPIC MOD
+	, bool forceBackwards
+	// END EPIC MOD
 )
 {
 	Record record = { coff::Relocation::Type::UNKNOWN, 0u, 0u, {} };
@@ -193,7 +196,7 @@ relocations::Record relocations::PatchRelocation
 	}
 
 	// BEGIN EPIC MOD
-	bool backwards = symbols::IsUEReversePatchSymbol(dstSymbolName) || symbols::IsUEReversePatchSymbol(srcSymbolName);
+	bool backwards = forceBackwards || symbols::IsUEReversePatchSymbol(dstSymbolName) || symbols::IsUEReversePatchSymbol(srcSymbolName);
 	// END EPIC MOD
 
 	// find the relocation's destination symbol in the original .exe, and patch the relocation
@@ -284,7 +287,9 @@ relocations::Record relocations::PatchRelocation
 
 				Process::WriteProcessMemory(processHandle, relocationAddress, displacement);
 
-				LC_LOG_DEV("Patched relocation from symbol %s to %s at 0x%p (0x%x + 0x%x)", srcSymbolName.c_str(), dstSymbolName.c_str(), newModuleBase, srcSymbol->rva, relocation->srcRva);
+				// BEGIN EPIC MOD
+				LC_LOG_DEV("Patched relocation from symbol %d:%s to %d:%s at 0x%p (0x%x + 0x%x) (relative offset)", newModuleIndex, srcSymbolName.c_str(), record.patchIndex, dstSymbolName.c_str(), record.patchIndex, newModuleBase, srcSymbol->rva, relocation->srcRva);
+				// END EPIC MOD
 			}
 		}
 		break;
@@ -319,7 +324,9 @@ relocations::Record relocations::PatchRelocation
 
 					Process::WriteProcessMemory(processHandle, relocationAddress, sectionRelativeRva);
 
-					LC_LOG_DEV("Patched relocation from symbol %s to %s at 0x%p (0x%x + 0x%x)", srcSymbolName.c_str(), dstSymbolName.c_str(), newModuleBase, srcSymbol->rva, relocation->srcRva);
+					// BEGIN EPIC MOD
+					LC_LOG_DEV("Patched relocation from symbol %d:%s to %d:%s at 0x%p (0x%x + 0x%x) (section relative)", newModuleIndex, srcSymbolName.c_str(), record.patchIndex, dstSymbolName.c_str(), newModuleBase, srcSymbol->rva, relocation->srcRva);
+					// END EPIC MOD
 				}
 			}
 			else
@@ -395,7 +402,9 @@ relocations::Record relocations::PatchRelocation
 
 				Process::WriteProcessMemory(processHandle, relocationAddress, displacement);
 
-				LC_LOG_DEV("Patched relocation from symbol %s to %s at 0x%p (0x%x + 0x%x)", srcSymbolName.c_str(), dstSymbolName.c_str(), newModuleBase, srcSymbol->rva, relocation->srcRva);
+				// BEGIN EPIC MOD
+				LC_LOG_DEV("Patched relocation from symbol %d:%s to %d:%s at 0x%p (0x%x + 0x%x) (RVA_32)", newModuleIndex, srcSymbolName.c_str(), record.patchIndex, dstSymbolName.c_str(), newModuleBase, srcSymbol->rva, relocation->srcRva);
+				// END EPIC MOD
 			}
 		}
 		break;
@@ -426,7 +435,9 @@ relocations::Record relocations::PatchRelocation
 
 				Process::WriteProcessMemory(processHandle, relocationAddress, va);
 
-				LC_LOG_DEV("Patched relocation from symbol %s to %s at 0x%p (0x%x + 0x%x)", srcSymbolName.c_str(), dstSymbolName.c_str(), newModuleBase, srcSymbol->rva, relocation->srcRva);
+				// BEGIN EPIC MOD
+				LC_LOG_DEV("Patched relocation from symbol %d:%s to %d:%s at 0x%p (0x%x + 0x%x) (VA_64)", newModuleIndex, srcSymbolName.c_str(), record.patchIndex, dstSymbolName.c_str(), newModuleBase, srcSymbol->rva, relocation->srcRva);
+				// END EPIC MOD
 			}
 		}
 		break;
diff --git a/Engine/Source/Developer/Windows/LiveCodingServer/Private/External/LC_RelocationPatcher.h b/Engine/Source/Developer/Windows/LiveCodingServer/Private/External/LC_RelocationPatcher.h
--- a/Engine/Source/Developer/Windows/LiveCodingServer/Private/External/LC_RelocationPatcher.h
+++ b/Engine/Source/Developer/Windows/LiveCodingServer/Private/External/LC_RelocationPatcher.h
@@ -69,6 +69,9 @@ namespace relocations
 		const symbols::Symbol* srcSymbol,
 		size_t newModuleIndex,
 		void* newModuleBases[]
+// BEGIN EPIC MOD
+		, bool forceBackwards
+// END EPIC MOD
 	);
 
 	void PatchRelocation
diff --git a/Engine/Source/Developer/Windows/LiveCodingServer/Private/External/LC_ServerCommandThread.cpp b/Engine/Source/Developer/Windows/LiveCodingServer/Private/External/LC_ServerCommandThread.cpp
--- a/Engine/Source/Developer/Windows/LiveCodingServer/Private/External/LC_ServerCommandThread.cpp
+++ b/Engine/Source/Developer/Windows/LiveCodingServer/Private/External/LC_ServerCommandThread.cpp
@@ -1063,11 +1063,11 @@ void ServerCommandThread::CallPrecompileHooks(bool didAllProcessesMakeProgress)
 		return;
 	}
 
-	const size_t count = m_liveModules.size();
+	const size_t count = m_liveProcesses.size();
 	for (size_t i = 0u; i < count; ++i)
 	{
-		LiveModule* liveModule = m_liveModules[i];
-		liveModule->CallPrecompileHooks();
+		LiveProcess* liveProcess = m_liveProcesses[i];
+		liveProcess->GetPipe()->SendCommandAndWaitForAck(commands::PreCompile{}, nullptr, 0u);
 	}
 }
 
@@ -1088,11 +1088,11 @@ void ServerCommandThread::CallPostcompileHooks(bool didAllProcessesMakeProgress)
 		return;
 	}
 
-	const size_t count = m_liveModules.size();
+	const size_t count = m_liveProcesses.size();
 	for (size_t i = 0u; i < count; ++i)
 	{
-		LiveModule* liveModule = m_liveModules[i];
-		liveModule->CallPostcompileHooks();
+		LiveProcess* liveProcess = m_liveProcesses[i];
+		liveProcess->GetPipe()->SendCommandAndWaitForAck(commands::PostCompile{}, nullptr, 0u);
 	}
 }
 // END EPIC MOD
@@ -1370,6 +1370,9 @@ Thread::ReturnValue ServerCommandThread::CommandThread(DuplexPipeServer* pipe, E
 	// BEGIN EPIC MOD - Support for lazy-loading modules
 	commandMap.RegisterAction<actions::EnableLazyLoadedModule>();
 	// END EPIC MOD
+	// BEGIN EPIC MOD
+	commandMap.RegisterAction<actions::EnableReinstancingFlow>();
+	// END EPIC MOD
 
 	for (;;)
 	{
@@ -1725,6 +1728,28 @@ bool ServerCommandThread::actions::EnableLazyLoadedModule::Execute(const Command
 }
 // END EPIC MOD
 
+// BEGIN EPIC MOD
+bool ServerCommandThread::actions::EnableReinstancingFlow::Execute(const CommandType* command, const DuplexPipe* pipe, void* context, const void*, size_t)
+{
+	ServerCommandThread* commandThread = static_cast<ServerCommandThread*>(context);
+
+	// protect against accepting this command while compilation is already in progress
+	CriticalSection::ScopedLock lock(&commandThread->m_actionCS);
+
+	for (LiveProcess* process : commandThread->m_liveProcesses)
+	{
+		if (process->GetProcessId() == command->processId)
+		{
+			process->EnableReinstancingFlow();
+		}
+	}
+
+	pipe->SendAck();
+
+	return true;
+}
+// END EPIC MOD
+
 bool ServerCommandThread::actions::RegisterProcess::Execute(const CommandType* command, const DuplexPipe* pipe, void* context, const void* payload, size_t)
 {
 	pipe->SendAck();
diff --git a/Engine/Source/Developer/Windows/LiveCodingServer/Private/External/LC_ServerCommandThread.h b/Engine/Source/Developer/Windows/LiveCodingServer/Private/External/LC_ServerCommandThread.h
--- a/Engine/Source/Developer/Windows/LiveCodingServer/Private/External/LC_ServerCommandThread.h
+++ b/Engine/Source/Developer/Windows/LiveCodingServer/Private/External/LC_ServerCommandThread.h
@@ -116,6 +116,10 @@ private:
 		DECLARE_ACTION(FinishedLazyLoadingModules);
 		// END EPIC MOD
 
+		// BEGIN EPIC MOD
+		DECLARE_ACTION(EnableReinstancingFlow);
+		// END EPIC MOD
+
 		#undef DECLARE_ACTION
 	};
 
diff --git a/Engine/Source/Editor/UnrealEd/Private/Kismet2/ReloadUtilities.cpp b/Engine/Source/Editor/UnrealEd/Private/Kismet2/ReloadUtilities.cpp
--- a/Engine/Source/Editor/UnrealEd/Private/Kismet2/ReloadUtilities.cpp
+++ b/Engine/Source/Editor/UnrealEd/Private/Kismet2/ReloadUtilities.cpp
@@ -941,7 +941,7 @@ void FReload::ReinstanceClass(UClass* NewClass, UClass* OldClass, const TSet<UOb
 	}
 }
 
-void FReload::Finalize()
+void FReload::Finalize(bool bRunGC)
 {
 	// If we have to collect the packages
 	if (bCollectPackages)
@@ -996,7 +996,10 @@ void FReload::Finalize()
 	ReplaceReferencesToReconstructedCDOs();
 
 	// Force GC to collect reinstanced objects
-	CollectGarbage(GARBAGE_COLLECTION_KEEPFLAGS, true);
+	if (bRunGC)
+	{
+		CollectGarbage(GARBAGE_COLLECTION_KEEPFLAGS, true);
+	}
 }
 
 void FReload::ReplaceReferencesToReconstructedCDOs()
diff --git a/Engine/Source/Editor/UnrealEd/Public/Kismet2/ReloadUtilities.h b/Engine/Source/Editor/UnrealEd/Public/Kismet2/ReloadUtilities.h
--- a/Engine/Source/Editor/UnrealEd/Public/Kismet2/ReloadUtilities.h
+++ b/Engine/Source/Editor/UnrealEd/Public/Kismet2/ReloadUtilities.h
@@ -50,7 +50,7 @@ public:
 	/**
 	 * Perform any finalize processing for reloading.
 	 */
-	void Finalize();
+	void Finalize(bool bRunGC = true);
 
 	/**
 	 * Set the sending of the complete notification
diff --git a/Engine/Source/Runtime/CoreUObject/Private/UObject/UObjectBase.cpp b/Engine/Source/Runtime/CoreUObject/Private/UObject/UObjectBase.cpp
--- a/Engine/Source/Runtime/CoreUObject/Private/UObject/UObjectBase.cpp
+++ b/Engine/Source/Runtime/CoreUObject/Private/UObject/UObjectBase.cpp
@@ -339,7 +339,14 @@ bool UObjectBase::IsValidLowLevelFast(bool bRecursive /*= true*/) const
 	// These should all be non-NULL (except CDO-alignment check which should be 0)
 	if (ClassPrivate == nullptr || ClassPrivate->ClassDefaultObject == nullptr || ((UPTRINT)ClassPrivate->ClassDefaultObject & AlignmentCheck) != 0)
 	{
-		UE_LOG(LogUObjectBase, Error, TEXT("Class pointer is invalid or CDO is invalid."));
+#if WITH_LIVE_CODING
+		// When live coding is re-instancing blueprint generated classes, we have to clear out the default object so it can get 
+		// GC'ed and deleted prior to live coding completing the patching process (of the destructor specifically)
+		if (!ClassPrivate->HasAnyClassFlags(CLASS_NewerVersionExists))
+#endif
+		{
+			UE_LOG(LogUObjectBase, Error, TEXT("Class pointer is invalid or CDO is invalid."));
+		}
 		return false;
 	}
 	// Avoid infinite recursion so call IsValidLowLevelFast on the class object with bRecirsive = false.
