Fix NonUnity build

#rnx

diff --git a/Engine/Source/Developer/Windows/LiveCoding/Private/LiveCodingModule.cpp b/Engine/Source/Developer/Windows/LiveCoding/Private/LiveCodingModule.cpp
--- a/Engine/Source/Developer/Windows/LiveCoding/Private/LiveCodingModule.cpp
+++ b/Engine/Source/Developer/Windows/LiveCoding/Private/LiveCodingModule.cpp
@@ -23,6 +23,7 @@
 #endif
 #if WITH_ENGINE
 	#include "Engine/Engine.h"
+	#include "UObject/UObjectIterator.h"
 #endif
 
 IMPLEMENT_MODULE(FLiveCodingModule, LiveCoding)
