Keep rewriting unused inputs until they stop changing (UE-115083)

- In a complex material, compiler may not be able to see through the code to discover all unused inputs.

E.g. first pass compilation says that input A is unused, but B and C are used. After we rewrite the code to remove A, the code is simpler and compiler is now able to see that B is also unused.

Since we pass the "used attributes" to the next stage, this creates a mismatch where vertex shader outputs only C, but pixel shader (rewritten based on the first pass results) expects B and C.

- The fixed code will check if the unused inputs changed, and if so, will go back and rewrite the code based on a more complete info.

- This fixes a bug that JonL ran into for which I haven't yet created a JIRA (will do before checkin).

diff --git a/Engine/Source/Developer/Windows/ShaderFormatD3D/Private/D3DShaderCompiler.cpp b/Engine/Source/Developer/Windows/ShaderFormatD3D/Private/D3DShaderCompiler.cpp
--- a/Engine/Source/Developer/Windows/ShaderFormatD3D/Private/D3DShaderCompiler.cpp
+++ b/Engine/Source/Developer/Windows/ShaderFormatD3D/Private/D3DShaderCompiler.cpp
@@ -522,7 +522,7 @@ bool DumpDebugShaderUSF(FString& PreprocessedShaderSource, const FShaderCompiler
 // Generate the dumped usf file; call the D3D compiler, gather reflection information and generate the output data
 bool CompileAndProcessD3DShaderFXC(FString& PreprocessedShaderSource, const FString& CompilerPath,
 	uint32 CompileFlags, const FShaderCompilerInput& Input, FString& EntryPointName,
-	const TCHAR* ShaderProfile, bool bProcessingSecondTime,
+	const TCHAR* ShaderProfile, bool bSecondPassAferUnusedInputRemoval,
 	TArray<FString>& FilteredErrors, FShaderCompilerOutput& Output)
 {
 	auto AnsiSourceFile = StringCast<ANSICHAR>(*PreprocessedShaderSource);
@@ -760,24 +760,86 @@ bool CompileAndProcessD3DShaderFXC(FString& PreprocessedShaderSource, const FStr
 					}
 				}
 
-				if (GD3DAllowRemoveUnused && Input.bCompilingForShaderPipeline && bFoundUnused && !bProcessingSecondTime)
+				if (GD3DAllowRemoveUnused && Input.bCompilingForShaderPipeline && bFoundUnused && !bSecondPassAferUnusedInputRemoval)
 				{
-					// Rewrite the source removing the unused inputs so the bindings will match
-					TArray<FString> RemoveErrors;
-					if (RemoveUnusedInputs(PreprocessedShaderSource, ShaderInputs, EntryPointName, RemoveErrors))
+					// Rewrite the source removing the unused inputs so the bindings will match.
+					// We may need to do this more than once if unused inputs change after the removal. Ie. for complex shaders, what can happen is:
+					// pass1 detects that input A is not used, but input B and C are. Input A is removed, and we recompile (pass2). After the recompilation, we see that Input B is now also unused in pass2 
+					// (it became simpler and the compiler could see through that).
+					// Since unused inputs are passed to the next stage, that will cause us to generate a vertex shader that does not output B, but our pixel shader will still be expecting B on input,
+					// as it was rewritten based on the pass1 results.
+
+					FString OriginalPreprocSource = PreprocessedShaderSource;
+					FString OriginalEntryPointName = EntryPointName;
+					FShaderCompilerOutput OriginalOutput = Output;
+					const int kMaxReasonableAttempts = 64;
+					for (int32 Attempt = 0; Attempt < kMaxReasonableAttempts; ++Attempt)
 					{
-						return CompileAndProcessD3DShaderFXC(PreprocessedShaderSource, CompilerPath, CompileFlags, Input, EntryPointName, ShaderProfile, true, FilteredErrors, Output);
-					}
-					else
-					{
-						UE_LOG(LogD3D11ShaderCompiler, Warning, TEXT("Failed to Remove unused inputs [%s]!"), *Input.DumpDebugInfoPath);
-						for (int32 Index = 0; Index < RemoveErrors.Num(); ++Index)
+						TArray<FString> RemoveErrors;
+						PreprocessedShaderSource = OriginalPreprocSource;
+						EntryPointName = OriginalEntryPointName;
+						if (RemoveUnusedInputs(PreprocessedShaderSource, ShaderInputs, EntryPointName, RemoveErrors))
+						{
+							Output = OriginalOutput;
+							if (!CompileAndProcessD3DShaderFXC(PreprocessedShaderSource, CompilerPath, CompileFlags, Input, EntryPointName, ShaderProfile, true, FilteredErrors, Output))
+							{
+								// if we failed to compile the shader, propagate the error up
+								return false;
+							}
+
+							// check if the ShaderInputs changed - if not, we're done here
+							if (Output.UsedAttributes.Num() == ShaderInputs.Num())
+							{
+								return true;
+							}
+
+							// second pass cannot use more attributes than previously
+							if (Output.UsedAttributes.Num() > ShaderInputs.Num())
+							{
+								UE_LOG(LogD3D11ShaderCompiler, Warning, TEXT("Second pass had more used attributes (%d) than first pass (%d)"), Output.UsedAttributes.Num(), ShaderInputs.Num());
+								FShaderCompilerError NewError;
+								NewError.StrippedErrorMessage = FString::Printf(TEXT("Second pass had more used attributes (%d) than first pass (%d)"), Output.UsedAttributes.Num(), ShaderInputs.Num());
+								Output = OriginalOutput;
+								Output.Errors.Add(NewError);
+								Output.bFailedRemovingUnused = true;
+								break;
+							}
+
+							// if we're about to run out of attempts, report
+							if (Attempt >= kMaxReasonableAttempts - 1)
+							{
+								UE_LOG(LogD3D11ShaderCompiler, Warning, TEXT("Unable to determine unused inputs after %d attempts (last number of used attributes: %d, previous step:%d)!"), 
+									Attempt + 1,
+									Output.UsedAttributes.Num(),
+									ShaderInputs.Num()
+									);
+								FShaderCompilerError NewError;
+								NewError.StrippedErrorMessage = FString::Printf(TEXT("Unable to determine unused inputs after %d attempts (last number of used attributes: %d, previous step:%d)!"),
+									Attempt + 1,
+									Output.UsedAttributes.Num(),
+									ShaderInputs.Num()
+									);
+								Output = OriginalOutput;
+								Output.Errors.Add(NewError);
+								Output.bFailedRemovingUnused = true;
+								break;
+							}
+
+							ShaderInputs = Output.UsedAttributes;
+							// go around to remove newly identified unused inputs
+						}
+						else
 						{
-							FShaderCompilerError NewError;
-							NewError.StrippedErrorMessage = RemoveErrors[Index];
-							Output.Errors.Add(NewError);
+							UE_LOG(LogD3D11ShaderCompiler, Warning, TEXT("Failed to Remove unused inputs [%s]!"), *Input.DumpDebugInfoPath);
+							for (int32 Index = 0; Index < RemoveErrors.Num(); ++Index)
+							{
+								FShaderCompilerError NewError;
+								NewError.StrippedErrorMessage = RemoveErrors[Index];
+								Output.Errors.Add(NewError);
+							}
+							Output.bFailedRemovingUnused = true;
+							break;
 						}
-						Output.bFailedRemovingUnused = true;
 					}
 				}
 			}
@@ -869,7 +931,7 @@ bool CompileAndProcessD3DShaderFXC(FString& PreprocessedShaderSource, const FStr
 			GenerateFinalOutput(CompressedData,
 				Input, VendorExtensions,
 				UsedUniformBufferSlots, UniformBufferNames,
-				bProcessingSecondTime, ShaderInputs,
+				bSecondPassAferUnusedInputRemoval, ShaderInputs,
 				PackedResourceCounts, NumInstructions,
 				Output,
 				[](FMemoryWriter&){},
diff --git a/Engine/Source/Developer/Windows/ShaderFormatD3D/Private/ShaderFormatD3D.h b/Engine/Source/Developer/Windows/ShaderFormatD3D/Private/ShaderFormatD3D.h
--- a/Engine/Source/Developer/Windows/ShaderFormatD3D/Private/ShaderFormatD3D.h
+++ b/Engine/Source/Developer/Windows/ShaderFormatD3D/Private/ShaderFormatD3D.h
@@ -18,9 +18,12 @@ enum class ELanguage
 
 void CompileShader_Windows(const struct FShaderCompilerInput& Input,struct FShaderCompilerOutput& Output,const class FString& WorkingDirectory, ELanguage Language);
 
+/**
+ * @param bSecondPassAferUnusedInputRemoval whether we're compiling the shader second time, after having removed the unused inputs discovered in the first pass
+ */
 bool CompileAndProcessD3DShaderFXC(FString& PreprocessedShaderSource, const FString& CompilerPath,
 	uint32 CompileFlags, const FShaderCompilerInput& Input, FString& EntryPointName,
-	const TCHAR* ShaderProfile, bool bProcessingSecondTime,
+	const TCHAR* ShaderProfile, bool bSecondPassAferUnusedInputRemoval,
 	TArray<FString>& FilteredErrors, FShaderCompilerOutput& Output);
 
 bool CompileAndProcessD3DShaderDXC(FString& PreprocessedShaderSource,
