Fixing LiveCoding issue where multiple changes to properties could still result in crashes during editing and at exit time.

#rnx

diff --git a/Engine/Source/Developer/Windows/LiveCoding/Private/LiveCodingModule.cpp b/Engine/Source/Developer/Windows/LiveCoding/Private/LiveCodingModule.cpp
--- a/Engine/Source/Developer/Windows/LiveCoding/Private/LiveCodingModule.cpp
+++ b/Engine/Source/Developer/Windows/LiveCoding/Private/LiveCodingModule.cpp
@@ -16,6 +16,8 @@
 #include "ISettingsModule.h"
 #include "ISettingsSection.h"
 #include "Windows/WindowsHWrapper.h"
+#include "Algo/Sort.h"
+#include "Algo/BinarySearch.h"
 #if WITH_EDITOR
 	#include "Kismet2/ReloadUtilities.h"
 #else
@@ -24,6 +26,7 @@
 #if WITH_ENGINE
 	#include "Engine/Engine.h"
 	#include "UObject/UObjectIterator.h"
+	#include "UObject/StrongObjectPtr.h"
 #endif
 
 IMPLEMENT_MODULE(FLiveCodingModule, LiveCoding)
@@ -372,6 +375,18 @@ void FLiveCodingModule::AttemptSyncLivePatching()
 		{
 #if WITH_COREUOBJECT && WITH_ENGINE
 
+			// Collect the existing objects
+			TArray<UObject*> StartingObjects;
+			if (Reload->GetEnableReinstancing(false))
+			{
+				StartingObjects.Reserve(1024); // Arbitrary
+				for (TObjectIterator<UObject> It(EObjectFlags::RF_NoFlags); It; ++It)
+				{
+					StartingObjects.Add(*It);
+				}
+				Algo::Sort(StartingObjects);
+			}
+
 			// During the module loading process, the list of changed classes will be recorded.  Invoking this method will 
 			// result in the RegisterForReinstancing method being invoked which in turn records the classes in the ClassesToReinstance
 			// member variable being populated.
@@ -382,27 +397,44 @@ void FLiveCodingModule::AttemptSyncLivePatching()
 			Reload->Finalize(false);
 #endif
 
-			// Loop through all of the classes looking for classes that have been re-instanced.  We have to clear
-			// the CDO in order for the CDO to be GC'ed.  Otherwise when the CDO is later GC'ed, it will invoke the
-			// wrong destructor.
-			TArray<UClass*> ReinstClasses;
+			TArray<TStrongObjectPtr<UObject>> NewObjects;
 			if (Reload->GetEnableReinstancing(false))
 			{
-				for (TObjectIterator<UClass> It; It; ++It)
+
+				// Loop through the objects again looking for anything new that isn't associated with a
+				// reinstanced class.
+				for (TObjectIterator<UObject> It(EObjectFlags::RF_NoFlags); It; ++It)
 				{
-					UClass* Class = *It;
-					if (Class->GetName().StartsWith(TEXT("LIVECODING_")))
+					if (Algo::BinarySearch(StartingObjects, *It) == INDEX_NONE)
 					{
-						ReinstClasses.Add(Class);
-						Class->ClassDefaultObject = nullptr;
+						if (!It->GetClass()->HasAnyClassFlags(CLASS_NewerVersionExists))
+						{
+							NewObjects.Add(TStrongObjectPtr<UObject>(*It));
+						}
 					}
-					else if (Class->GetName().StartsWith(TEXT("REINST_")))
+				}
+
+				// Loop through all of the classes looking for classes that have been re-instanced.  Reset the CDO
+				// to something that will never change.  Since these classes have been replaced, they should NEVER
+				// have their CDos accessed again.  In the future we should try to figure out a better solution the issue
+				// where the reinstanced crashes recreating the default object probably due to a mismatch between then
+				// new constructor being invoked and the blueprint data associated with the old class.  With LC, the
+				// old constructor has been replaced.
+				static UObject* DummyDefaultObject = UObject::StaticClass()->ClassDefaultObject;
+				for (TObjectIterator<UClass> It; It; ++It)
+				{
+					UClass* Class = *It;
+					if (Class->GetName().StartsWith(TEXT("LIVECODING_")) ||
+						Class->GetName().StartsWith(TEXT("REINST_")))
 					{
-						Class->ClassDefaultObject = nullptr;
+						Class->ClassDefaultObject = DummyDefaultObject;
 					}
 				}
 			}
 
+			// Broadcast event prior to GC.  Otherwise some things are holding onto references
+			FCoreUObjectDelegates::ReloadCompleteDelegate.Broadcast(EReloadCompleteReason::None);
+
 			// Perform the GC to try and destruct all the objects which will be invoking the old destructors.
 			CollectGarbage(GARBAGE_COLLECTION_KEEPFLAGS, true);
 #endif
@@ -414,15 +446,8 @@ void FLiveCodingModule::AttemptSyncLivePatching()
 			}
 
 #if WITH_COREUOBJECT && WITH_ENGINE
-			// Re-create the default objects.  This is technically incorrect since the structure doesn't match.
-			// But we need a CDO just in case.
-			if (Reload->GetEnableReinstancing(false))
-			{
-				for (UClass* Class : ReinstClasses)
-				{
-					Class->GetDefaultObject(true);
-				}
-			}
+			// Remove the reference to any new objects
+			NewObjects.Empty();
 #endif
 
 			OnPatchCompleteDelegate.Broadcast();
@@ -677,6 +702,7 @@ void FLiveCodingModule::BeginReload()
 			GConfig->GetBool(TEXT("LiveCoding"), TEXT("bEnableReinstancing"), bEnableReinstancing, GEngineIni);
 			GLiveCodingModule->Reload.Reset(new FReload(EActiveReloadType::LiveCoding, TEXT("LIVECODING"), *GLog));
 			GLiveCodingModule->Reload->SetEnableReinstancing(bEnableReinstancing);
+			GLiveCodingModule->Reload->SetSendReloadCompleteNotification(false);
 #else
 			GLiveCodingModule->Reload.Reset(new FNullReload(*GLiveCodingModule));
 #endif
