Windows-based Platform extentions can specialize GenericWindowsTargetPlatform, allowing access to all PC rendering features etc.

#rnx

diff --git a/Engine/Source/Developer/Windows/WindowsTargetPlatform/Private/GenericWindowsTargetPlatform.h b/Engine/Source/Developer/Windows/WindowsTargetPlatform/Private/GenericWindowsTargetPlatform.h
--- a/Engine/Source/Developer/Windows/WindowsTargetPlatform/Private/GenericWindowsTargetPlatform.h
+++ b/Engine/Source/Developer/Windows/WindowsTargetPlatform/Private/GenericWindowsTargetPlatform.h
@@ -21,26 +21,21 @@
 #define LOCTEXT_NAMESPACE "TGenericWindowsTargetPlatform"
 
 /**
- * Template for Windows target platforms
+ * Template for common Windows target platforms
  */
-template<bool HAS_EDITOR_DATA, bool IS_DEDICATED_SERVER, bool IS_CLIENT_ONLY>
-class TGenericWindowsTargetPlatform
-	: public TTargetPlatformBase<FWindowsPlatformProperties<HAS_EDITOR_DATA, IS_DEDICATED_SERVER, IS_CLIENT_ONLY> >
+template<class TPlatformProperties>
+class TWindowsTargetPlatformBase
+	: public TTargetPlatformBase<TPlatformProperties>
 {
 public:
-	typedef FWindowsPlatformProperties<HAS_EDITOR_DATA, IS_DEDICATED_SERVER, IS_CLIENT_ONLY> TProperties;
+	typedef TPlatformProperties TProperties;
 	typedef TTargetPlatformBase<TProperties> TSuper;
 
 	/**
 	 * Default constructor.
 	 */
-	TGenericWindowsTargetPlatform( )
+	TWindowsTargetPlatformBase( )
 	{
-#if PLATFORM_WINDOWS
-		// only add local device if actually running on Windows
-		LocalDevice = MakeShareable(new TLocalPcTargetDevice<PLATFORM_64BITS>(*this));
-#endif
-
 	#if WITH_ENGINE
 		TextureLODSettings = nullptr; // These are registered by the device profile system.
 		StaticMeshLODSettings.Initialize(this);
@@ -97,46 +92,15 @@
 
 	//~ Begin ITargetPlatform Interface
 
-	virtual void EnableDeviceCheck(bool OnOff) override {}
-
-	virtual void GetAllDevices( TArray<ITargetDevicePtr>& OutDevices ) const override
-	{
-		OutDevices.Reset();
-		if (LocalDevice.IsValid())
-		{
-			OutDevices.Add(LocalDevice);
-		}
-	}
-
 	virtual bool GenerateStreamingInstallManifest(const TMultiMap<FString, int32>& PakchunkMap, const TSet<int32>& PakchunkIndicesInUse) const override
 	{
 		return true;
 	}
 
-	virtual ITargetDevicePtr GetDefaultDevice( ) const override
-	{
-		if (LocalDevice.IsValid())
-		{
-			return LocalDevice;
-		}
-
-		return nullptr;
-	}
-
-	virtual ITargetDevicePtr GetDevice( const FTargetDeviceId& DeviceId )
-	{
-		if (LocalDevice.IsValid() && (DeviceId == LocalDevice->GetId()))
-		{
-			return LocalDevice;
-		}
-
-		return nullptr;
-	}
-
 	virtual bool IsRunningPlatform( ) const override
 	{
 		// Must be Windows platform as editor for this to be considered a running platform
-		return PLATFORM_WINDOWS && !UE_SERVER && !UE_GAME && WITH_EDITOR && HAS_EDITOR_DATA;
+		return PLATFORM_WINDOWS && !UE_SERVER && !UE_GAME && WITH_EDITOR && TProperties::HasEditorOnlyData();
 	}
 
 	virtual void GetShaderCompilerDependencies(TArray<FString>& OutDependencies) const override
@@ -152,12 +116,12 @@
 		// we currently do not have a build target for WindowsServer
 		if (Feature == ETargetPlatformFeatures::Packaging)
 		{
-			return (HAS_EDITOR_DATA || !IS_DEDICATED_SERVER);
+			return (TProperties::HasEditorOnlyData() || !TProperties::IsServerOnly());
 		}
 
 		if ( Feature == ETargetPlatformFeatures::ShouldSplitPaksIntoSmallerSizes )
 		{
-			return IS_CLIENT_ONLY;
+			return TProperties::IsClientOnly();
 		}
 
 		if (Feature == ETargetPlatformFeatures::MobileRendering)
@@ -190,7 +154,7 @@
 
 	virtual void GetBuildProjectSettingKeys(FString& OutSection, TArray<FString>& InBoolKeys, TArray<FString>& InIntKeys, TArray<FString>& InStringKeys) const override
 	{
-		OutSection = TEXT("/Script/WindowsTargetPlatform.WindowsTargetSettings");
+		OutSection = TProperties::GetRuntimeSettingsClassName();
 		InStringKeys.Add(TEXT("MinimumOSVersion"));
 	}
 
@@ -208,7 +172,7 @@
 	virtual void GetAllPossibleShaderFormats( TArray<FName>& OutFormats ) const override
 	{
 		// no shaders needed for dedicated server target
-		if (!IS_DEDICATED_SERVER)
+		if (!TProperties::IsServerOnly())
 		{
 			static FName NAME_PCD3D_SM6(TEXT("PCD3D_SM6"));
 			static FName NAME_PCD3D_SM5(TEXT("PCD3D_SM5"));
@@ -230,7 +194,7 @@
 	{
 		// Get the Target RHIs for this platform, we do not always want all those that are supported. (reload in case user changed in the editor)
 		TArray<FString>TargetedShaderFormats;
-		GConfig->GetArray(TEXT("/Script/WindowsTargetPlatform.WindowsTargetSettings"), TEXT("TargetedRHIs"), TargetedShaderFormats, GEngineIni);
+		GConfig->GetArray(TProperties::GetRuntimeSettingsClassName(), TEXT("TargetedRHIs"), TargetedShaderFormats, GEngineIni);
 
 		// Gather the list of Target RHIs and filter out any that may be invalid.
 		TArray<FName> PossibleShaderFormats;
@@ -257,7 +221,7 @@
 
 	virtual void GetTextureFormats( const UTexture* InTexture, TArray< TArray<FName> >& OutFormats) const override
 	{
-		if (!IS_DEDICATED_SERVER)
+		if (!TProperties::IsServerOnly())
 		{
 			GetDefaultTextureFormatNamePerLayer(OutFormats.AddDefaulted_GetRef(), this, InTexture, bSupportDX11TextureFormats, bSupportCompressedVolumeTexture);
 		}
@@ -265,7 +229,7 @@
 
 	virtual void GetAllTextureFormats(TArray<FName>& OutFormats) const override
 	{
-		if (!IS_DEDICATED_SERVER)
+		if (!TProperties::IsServerOnly())
 		{
 			GetAllDefaultTextureFormats(this, OutFormats, bSupportDX11TextureFormats);
 		}
@@ -459,17 +423,14 @@
 	virtual bool UsesRayTracing() const override
 	{
 		bool bEnableRayTracing = false;
-		GConfig->GetBool(TEXT("/Script/WindowsTargetPlatform.WindowsTargetSettings"), TEXT("bEnableRayTracing"), bEnableRayTracing, GEngineIni);
+		GConfig->GetBool(TProperties::GetRuntimeSettingsClassName(), TEXT("bEnableRayTracing"), bEnableRayTracing, GEngineIni);
 
 		return bEnableRayTracing && TSuper::UsesRayTracing();
 	}
 
 	//~ End ITargetPlatform Interface
 
-private:
+protected:
 
-	// Holds the local device.
-	ITargetDevicePtr LocalDevice;
-
 #if WITH_ENGINE
 	// Holds the texture LOD settings.
@@ -490,4 +451,72 @@
 
 };
 
+
+
+/**
+ * Template for Windows target platforms
+ */
+template<bool HAS_EDITOR_DATA, bool IS_DEDICATED_SERVER, bool IS_CLIENT_ONLY>
+class TGenericWindowsTargetPlatform
+	: public TWindowsTargetPlatformBase<FWindowsPlatformProperties<HAS_EDITOR_DATA, IS_DEDICATED_SERVER, IS_CLIENT_ONLY> >
+{
+public:
+	typedef TWindowsTargetPlatformBase<FWindowsPlatformProperties<HAS_EDITOR_DATA, IS_DEDICATED_SERVER, IS_CLIENT_ONLY> > TSuper;
+
+	/**
+	 * Default constructor.
+	 */
+	TGenericWindowsTargetPlatform( )
+		: TSuper()
+	{
+#if PLATFORM_WINDOWS
+		// only add local device if actually running on Windows
+		LocalDevice = MakeShareable(new TLocalPcTargetDevice<PLATFORM_64BITS>(*this));
+#endif
+	}
+
+public:
+
+	//~ Begin ITargetPlatform Interface
+
+	virtual void EnableDeviceCheck(bool OnOff) override {}
+
+	virtual void GetAllDevices( TArray<ITargetDevicePtr>& OutDevices ) const override
+	{
+		OutDevices.Reset();
+		if (LocalDevice.IsValid())
+		{
+			OutDevices.Add(LocalDevice);
+		}
+	}
+
+	virtual ITargetDevicePtr GetDefaultDevice( ) const override
+	{
+		if (LocalDevice.IsValid())
+		{
+			return LocalDevice;
+		}
+
+		return nullptr;
+	}
+
+	virtual ITargetDevicePtr GetDevice( const FTargetDeviceId& DeviceId )
+	{
+		if (LocalDevice.IsValid() && (DeviceId == LocalDevice->GetId()))
+		{
+			return LocalDevice;
+		}
+
+		return nullptr;
+	}
+
+private:
+
+	// Holds the local device.
+	ITargetDevicePtr LocalDevice;
+
+};
+
+
+
 #undef LOCTEXT_NAMESPACE
