Report error code and addresses on D3DCompile exception.

#rnx

diff --git a/Engine/Source/Developer/Windows/ShaderFormatD3D/Private/D3DShaderCompiler.cpp b/Engine/Source/Developer/Windows/ShaderFormatD3D/Private/D3DShaderCompiler.cpp
--- a/Engine/Source/Developer/Windows/ShaderFormatD3D/Private/D3DShaderCompiler.cpp
+++ b/Engine/Source/Developer/Windows/ShaderFormatD3D/Private/D3DShaderCompiler.cpp
@@ -433,9 +433,58 @@ static bool GetD3DCompilerFuncs(const FString& NewCompilerPath, pD3DCompile* Out
 	return false;
 }
 
+static const char* Win32SehExceptionToString(DWORD Code)
+{
+#define CASE_TO_STRING(IDENT) case IDENT: return #IDENT
+
+	switch (Code)
+	{
+		CASE_TO_STRING(EXCEPTION_ACCESS_VIOLATION);
+		CASE_TO_STRING(EXCEPTION_ARRAY_BOUNDS_EXCEEDED);
+		CASE_TO_STRING(EXCEPTION_BREAKPOINT);
+		CASE_TO_STRING(EXCEPTION_DATATYPE_MISALIGNMENT);
+		CASE_TO_STRING(EXCEPTION_FLT_DENORMAL_OPERAND);
+		CASE_TO_STRING(EXCEPTION_FLT_DIVIDE_BY_ZERO);
+		CASE_TO_STRING(EXCEPTION_FLT_INEXACT_RESULT);
+		CASE_TO_STRING(EXCEPTION_FLT_INVALID_OPERATION);
+		CASE_TO_STRING(EXCEPTION_FLT_OVERFLOW);
+		CASE_TO_STRING(EXCEPTION_FLT_STACK_CHECK);
+		CASE_TO_STRING(EXCEPTION_FLT_UNDERFLOW);
+		CASE_TO_STRING(EXCEPTION_GUARD_PAGE);
+		CASE_TO_STRING(EXCEPTION_ILLEGAL_INSTRUCTION);
+		CASE_TO_STRING(EXCEPTION_IN_PAGE_ERROR);
+		CASE_TO_STRING(EXCEPTION_INT_DIVIDE_BY_ZERO);
+		CASE_TO_STRING(EXCEPTION_INT_OVERFLOW);
+		CASE_TO_STRING(EXCEPTION_INVALID_DISPOSITION);
+		CASE_TO_STRING(EXCEPTION_INVALID_HANDLE);
+		CASE_TO_STRING(EXCEPTION_NONCONTINUABLE_EXCEPTION);
+		CASE_TO_STRING(EXCEPTION_PRIV_INSTRUCTION);
+		CASE_TO_STRING(EXCEPTION_SINGLE_STEP);
+		CASE_TO_STRING(EXCEPTION_STACK_OVERFLOW);
+		CASE_TO_STRING(STATUS_UNWIND_CONSOLIDATE);
+		default: return nullptr;
+	}
+
+#undef CASE_TO_STRING
+}
+
+struct FD3DExceptionInfo
+{
+	uint32 Code;
+	uint64 Base;
+	uint64 Address;
+};
+
+static int D3DExceptionFilter(DWORD Code, LPEXCEPTION_POINTERS InInfo, FD3DExceptionInfo& OutInfo)
+{
+	OutInfo.Code = InInfo->ExceptionRecord->ExceptionCode;
+	OutInfo.Base = (uint64)GetModuleHandle(NULL);
+	OutInfo.Address = (uint64)InInfo->ExceptionRecord->ExceptionAddress;
+	return EXCEPTION_EXECUTE_HANDLER;
+}
+
 static HRESULT D3DCompileWrapper(
 	pD3DCompile				D3DCompileFunc,
-	bool&					bException,
 	LPCVOID					pSrcData,
 	SIZE_T					SrcDataSize,
 	LPCSTR					pFileName,
@@ -446,7 +495,9 @@ static HRESULT D3DCompileWrapper(
 	uint32					Flags1,
 	uint32					Flags2,
 	ID3DBlob**				ppCode,
-	ID3DBlob**				ppErrorMsgs
+	ID3DBlob**				ppErrorMsgs,
+	bool&					bOutException,
+	FD3DExceptionInfo&		OutExceptionInfo
 	)
 {
 #if !PLATFORM_SEH_EXCEPTIONS_DISABLED
@@ -468,10 +519,10 @@ static HRESULT D3DCompileWrapper(
 		);
 	}
 #if !PLATFORM_SEH_EXCEPTIONS_DISABLED
-	__except( EXCEPTION_EXECUTE_HANDLER )
+	__except(D3DExceptionFilter(GetExceptionCode(), GetExceptionInformation(), OutExceptionInfo))
 	{
 		GSCWErrorCode = ESCWErrorCode::CrashInsidePlatformCompiler;
-		bException = true;
+		bOutException = true;
 		return E_FAIL;
 	}
 #endif
@@ -566,10 +617,10 @@ bool CompileAndProcessD3DShaderFXC(FString& PreprocessedShaderSource, const FStr
 	if (D3DCompileFunc)
 	{
 		bool bException = false;
+		FD3DExceptionInfo ExceptionInfo;
 
 		Result = D3DCompileWrapper(
 			D3DCompileFunc,
-			bException,
 			AnsiSourceFile.Get(),
 			AnsiSourceFile.Length(),
 			TCHAR_TO_ANSI(*Input.VirtualSourceFilePath),
@@ -580,12 +631,25 @@ bool CompileAndProcessD3DShaderFXC(FString& PreprocessedShaderSource, const FStr
 			CompileFlags,
 			0,
 			Shader.GetInitReference(),
-			Errors.GetInitReference()
+			Errors.GetInitReference(),
+			bException,
+			ExceptionInfo
 		);
 
 		if (bException)
 		{
-			FilteredErrors.Add(TEXT("D3DCompile exception"));
+			const char* CodeName = Win32SehExceptionToString(ExceptionInfo.Code);
+			const FString CodeNameStr = (CodeName != nullptr ? ANSI_TO_TCHAR(CodeName) : TEXT("Unknown"));
+			FilteredErrors.Add(
+				FString::Printf(
+					TEXT("D3DCompile exception: Code = 0x%08X (%s), Address = 0x%016llX, Offset = 0x%016llX, Codebase = 0x%016llX"),
+					ExceptionInfo.Code,
+					*CodeNameStr,
+					ExceptionInfo.Address,
+					(ExceptionInfo.Address - ExceptionInfo.Base),
+					ExceptionInfo.Base
+				)
+			);
 		}
 	}
 	else
