If re-instancing is enabled in a process using LiveCoding, disable the "Quick Restart" button until we can resolve the issues.

#rnx

diff --git a/Engine/Source/Developer/Windows/LiveCodingServer/Private/External/LC_ServerCommandThread.cpp b/Engine/Source/Developer/Windows/LiveCodingServer/Private/External/LC_ServerCommandThread.cpp
--- a/Engine/Source/Developer/Windows/LiveCodingServer/Private/External/LC_ServerCommandThread.cpp
+++ b/Engine/Source/Developer/Windows/LiveCodingServer/Private/External/LC_ServerCommandThread.cpp
@@ -612,6 +612,25 @@ Thread::ReturnValue ServerCommandThread::ServerThread(void)
 	return Thread::ReturnValue(0u);
 }
 
+// BEGIN EPIC MOD
+bool ServerCommandThread::HasReinstancingProcess()
+{
+
+	// protect against m_liveProcesses being accessed when processes restart and register themselves with this Live++ instance
+	CriticalSection::ScopedLock lock(&m_actionCS);
+
+	// remove processes that were successfully restarted last time
+	for (auto processIt = m_liveProcesses.begin(); processIt != m_liveProcesses.end(); ++processIt)
+	{
+		LiveProcess* liveProcess = *processIt;
+		if (liveProcess->IsReinstancingFlowEnabled())
+		{
+			return true;
+		}
+	}
+	return false;
+}
+// END EPIC MOD
 
 // BEGIN EPIC MOD - Focus application windows on patch complete
 BOOL CALLBACK FocusApplicationWindows(HWND WindowHandle, LPARAM Lparam)
@@ -1733,6 +1752,11 @@ bool ServerCommandThread::actions::EnableReinstancingFlow::Execute(const Command
 {
 	ServerCommandThread* commandThread = static_cast<ServerCommandThread*>(context);
 
+	if (!commandThread->HasReinstancingProcess())
+	{
+		LC_WARNING_USER("Quick restart disabled when re-instancing is enabled.");
+	}
+
 	// protect against accepting this command while compilation is already in progress
 	CriticalSection::ScopedLock lock(&commandThread->m_actionCS);
 
diff --git a/Engine/Source/Developer/Windows/LiveCodingServer/Private/External/LC_ServerCommandThread.h b/Engine/Source/Developer/Windows/LiveCodingServer/Private/External/LC_ServerCommandThread.h
--- a/Engine/Source/Developer/Windows/LiveCodingServer/Private/External/LC_ServerCommandThread.h
+++ b/Engine/Source/Developer/Windows/LiveCodingServer/Private/External/LC_ServerCommandThread.h
@@ -35,6 +35,10 @@ public:
 
 	std::wstring GetProcessImagePath(void) const;
 
+	// BEGIN EPIC MOD
+	bool HasReinstancingProcess();
+	// END EPIC MOD
+
 private:
 	scheduler::Task<LiveModule*>* LoadModule(Process::Id processId, void* moduleBase, const wchar_t* modulePath, scheduler::TaskBase* taskRoot);
 	bool UnloadModule(Process::Id processId, const wchar_t* modulePath);
diff --git a/Engine/Source/Developer/Windows/LiveCodingServer/Private/LiveCodingServer.cpp b/Engine/Source/Developer/Windows/LiveCodingServer/Private/LiveCodingServer.cpp
--- a/Engine/Source/Developer/Windows/LiveCodingServer/Private/LiveCodingServer.cpp
+++ b/Engine/Source/Developer/Windows/LiveCodingServer/Private/LiveCodingServer.cpp
@@ -114,3 +114,8 @@ ILiveCodingServer::FSetVisibleDelegate& FLiveCodingServer::GetSetVisibleDelegate
 {
 	return SetVisibleDelegate;
 }
+
+bool FLiveCodingServer::HasReinstancingProcess()
+{
+	return CommandThread->HasReinstancingProcess();
+}
\ No newline at end of file
diff --git a/Engine/Source/Developer/Windows/LiveCodingServer/Private/LiveCodingServer.h b/Engine/Source/Developer/Windows/LiveCodingServer/Private/LiveCodingServer.h
--- a/Engine/Source/Developer/Windows/LiveCodingServer/Private/LiveCodingServer.h
+++ b/Engine/Source/Developer/Windows/LiveCodingServer/Private/LiveCodingServer.h
@@ -32,6 +32,7 @@ public:
 	virtual FCompileFinishedDelegate& GetCompileFinishedDelegate() override final;
 	virtual FShowConsoleDelegate& GetShowConsoleDelegate() override final;
 	virtual FSetVisibleDelegate& GetSetVisibleDelegate() override final;
+	virtual bool HasReinstancingProcess() override final;
 
 private:
 	std::wstring ProcessGroupName;
diff --git a/Engine/Source/Developer/Windows/LiveCodingServer/Public/ILiveCodingServer.h b/Engine/Source/Developer/Windows/LiveCodingServer/Public/ILiveCodingServer.h
--- a/Engine/Source/Developer/Windows/LiveCodingServer/Public/ILiveCodingServer.h
+++ b/Engine/Source/Developer/Windows/LiveCodingServer/Public/ILiveCodingServer.h
@@ -83,5 +83,7 @@ public:
 
 	DECLARE_DELEGATE_OneParam(FSetVisibleDelegate, bool);
 	virtual FSetVisibleDelegate& GetSetVisibleDelegate() = 0;
+
+	virtual bool HasReinstancingProcess() = 0;
 };
 
diff --git a/Engine/Source/Programs/LiveCodingConsole/Private/LiveCodingConsole.cpp b/Engine/Source/Programs/LiveCodingConsole/Private/LiveCodingConsole.cpp
--- a/Engine/Source/Programs/LiveCodingConsole/Private/LiveCodingConsole.cpp
+++ b/Engine/Source/Programs/LiveCodingConsole/Private/LiveCodingConsole.cpp
@@ -95,6 +95,10 @@ public:
 							SNew(SButton)
 							.Text(LOCTEXT("QuickRestart", "Quick Restart"))
 							.OnClicked(FOnClicked::CreateRaw(this, &FLiveCodingConsoleApp::RestartTargets))
+							.ToolTipText_Lambda([this]() { return Server.HasReinstancingProcess() ? 
+								LOCTEXT("DisableQuickRestart", "Quick restarting isn't supported when re-instancing is enabled") :
+								LOCTEXT("EnableQuickRestart", "Restart all live coding applications"); })
+							.IsEnabled_Lambda([this]() { return !Server.HasReinstancingProcess(); })
 						]
 						+ SHorizontalBox::Slot()
 						.HAlign(HAlign_Center)
