Starting with VS2019 16.10, we are getting multiple _valmd symbols in the object files.  Changing to not generate any warnings about symbols in COMDATs.

#rnx

diff --git a/Engine/Source/Developer/Windows/LiveCodingServer/Private/External/LC_Coff.cpp b/Engine/Source/Developer/Windows/LiveCodingServer/Private/External/LC_Coff.cpp
--- a/Engine/Source/Developer/Windows/LiveCodingServer/Private/External/LC_Coff.cpp
+++ b/Engine/Source/Developer/Windows/LiveCodingServer/Private/External/LC_Coff.cpp
@@ -1221,16 +1221,20 @@ namespace coff
 							{
 								// BEGIN EPIC MOD
 								// VS2019 seems to generate duplicate unwind and pdata block for the "dynamic atexit destructor" method (__F) under some conditions that I can't determine.
-								if (!string::StartsWith(symbolName.c_str(), "$unwind$??__F") &&
-									!string::StartsWith(symbolName.c_str(), "$pdata$??__F"))
+								// Also, do not generate warnings for symbols in COMDATs
+								if (!coffDetail::IsComdatSection(section))
 								{
-								    LC_WARNING_USER("Non-unique symbol %s found in COFF file %s. Do not change the order of these variables while live coding, or consider upgrading to a newer compiler (VS 2015 or later)",
-									    symbolName.c_str(), objFile->filename.c_str());
-								}
-								else
-								{
-									LC_LOG_USER("Non-unique at-exit symbol %s found in COFF file %s. These sometimes appear in debug builds",
-										symbolName.c_str(), objFile->filename.c_str());
+									if (!string::StartsWith(symbolName.c_str(), "$unwind$??__F") &&
+										!string::StartsWith(symbolName.c_str(), "$pdata$??__F"))
+									{
+										LC_WARNING_USER("Non-unique symbol %s found in COFF file %s. Do not change the order of these variables while live coding, or consider upgrading to a newer compiler (VS 2015 or later)",
+											symbolName.c_str(), objFile->filename.c_str());
+									}
+									else
+									{
+										LC_LOG_USER("Non-unique at-exit symbol %s found in COFF file %s. These sometimes appear in debug builds",
+											symbolName.c_str(), objFile->filename.c_str());
+									}
 								}
 								// END EPIC MOD
 							}
