Fix new bug with LiveCodingConsole where it doesn't update UI during a compile.

#rnx

diff --git a/Engine/Source/Developer/Windows/LiveCodingServer/Private/External/LC_ServerCommandThread.cpp b/Engine/Source/Developer/Windows/LiveCodingServer/Private/External/LC_ServerCommandThread.cpp
--- a/Engine/Source/Developer/Windows/LiveCodingServer/Private/External/LC_ServerCommandThread.cpp
+++ b/Engine/Source/Developer/Windows/LiveCodingServer/Private/External/LC_ServerCommandThread.cpp
@@ -615,20 +615,7 @@ Thread::ReturnValue ServerCommandThread::ServerThread(void)
 // BEGIN EPIC MOD
 bool ServerCommandThread::HasReinstancingProcess()
 {
-
-	// protect against m_liveProcesses being accessed when processes restart and register themselves with this Live++ instance
-	CriticalSection::ScopedLock lock(&m_actionCS);
-
-	// remove processes that were successfully restarted last time
-	for (auto processIt = m_liveProcesses.begin(); processIt != m_liveProcesses.end(); ++processIt)
-	{
-		LiveProcess* liveProcess = *processIt;
-		if (liveProcess->IsReinstancingFlowEnabled())
-		{
-			return true;
-		}
-	}
-	return false;
+	return m_reinstancingProcessCount.load() != 0;
 }
 // END EPIC MOD
 
@@ -1764,7 +1751,11 @@ bool ServerCommandThread::actions::EnableReinstancingFlow::Execute(const Command
 	{
 		if (process->GetProcessId() == command->processId)
 		{
-			process->EnableReinstancingFlow();
+			if (!process->IsReinstancingFlowEnabled())
+			{
+				++commandThread->m_reinstancingProcessCount;
+				process->EnableReinstancingFlow();
+			}
 		}
 	}
 
diff --git a/Engine/Source/Developer/Windows/LiveCodingServer/Private/External/LC_ServerCommandThread.h b/Engine/Source/Developer/Windows/LiveCodingServer/Private/External/LC_ServerCommandThread.h
--- a/Engine/Source/Developer/Windows/LiveCodingServer/Private/External/LC_ServerCommandThread.h
+++ b/Engine/Source/Developer/Windows/LiveCodingServer/Private/External/LC_ServerCommandThread.h
@@ -16,6 +16,7 @@
 // BEGIN EPIC MOD
 #include "LC_Types.h"
 #include "VisualStudioDTE.h"
+#include <atomic>
 // END EPIC MOD
 
 class MainFrame;
@@ -151,6 +152,10 @@ private:
 	CriticalSection m_connectionCS;
 	types::vector<CommandThreadContext*> m_commandThreads;
 
+	// BEGIN EPIC MOD
+	std::atomic<unsigned int> m_reinstancingProcessCount = 0;
+	// END EPIC MOD
+
 	// BEGIN EPIC MOD - Non-destructive compile
 	std::vector<std::pair<std::wstring, std::wstring>> m_restoreFiles;
 	// END EPIC MOD
