Improved notifications in the editor/game for live coding.
Added message that packaging can fail if assets reference new changes.

#rnx

diff --git a/Engine/Source/Developer/Windows/LiveCoding/LiveCoding.Build.cs b/Engine/Source/Developer/Windows/LiveCoding/LiveCoding.Build.cs
--- a/Engine/Source/Developer/Windows/LiveCoding/LiveCoding.Build.cs
+++ b/Engine/Source/Developer/Windows/LiveCoding/LiveCoding.Build.cs
@@ -19,7 +19,8 @@ public class LiveCoding : ModuleRules
 		if (Target.bBuildEditor == true)
         {
             PrivateDependencyModuleNames.Add("UnrealEd");
-        }        
+			PrivateDependencyModuleNames.Add("Slate");
+		}
 
 		if(Target.bUseDebugLiveCodingConsole)
         {
diff --git a/Engine/Source/Developer/Windows/LiveCoding/Private/External/LC_API.cpp b/Engine/Source/Developer/Windows/LiveCoding/Private/External/LC_API.cpp
--- a/Engine/Source/Developer/Windows/LiveCoding/Private/External/LC_API.cpp
+++ b/Engine/Source/Developer/Windows/LiveCoding/Private/External/LC_API.cpp
@@ -346,9 +346,16 @@ LPP_DLL_API(void*) LppEnableLazyLoadedModule(const wchar_t* nameOfExeOrDll)
 // END EPIC MOD
 
 // BEGIN EPIC MOD
-LPP_DLL_API(void) LppEnableReinstancingFlow()
+LPP_DLL_API(void) LppSetReinstancingFlow(bool enable)
 {
-	g_startupThread->EnableReinstancingFlow();
+	g_startupThread->SetReinstancingFlow(enable);
+}
+// END EPIC MOD
+
+// BEGIN EPIC MOD
+LPP_DLL_API(void) LppDisableCompileFinishNotification()
+{
+	g_startupThread->DisableCompileFinishNotification();
 }
 // END EPIC MOD
 
diff --git a/Engine/Source/Developer/Windows/LiveCoding/Private/External/LC_API.h b/Engine/Source/Developer/Windows/LiveCoding/Private/External/LC_API.h
--- a/Engine/Source/Developer/Windows/LiveCoding/Private/External/LC_API.h
+++ b/Engine/Source/Developer/Windows/LiveCoding/Private/External/LC_API.h
@@ -41,7 +41,8 @@ LPP_DLL_API(void) LppSetVisible(bool visible);
 LPP_DLL_API(void) LppSetActive(bool active);
 LPP_DLL_API(void) LppSetBuildArguments(const wchar_t* arguments);
 LPP_DLL_API(void*) LppEnableLazyLoadedModule(const wchar_t* nameOfExeOrDll);
-LPP_DLL_API(void) LppEnableReinstancingFlow();
+LPP_DLL_API(void) LppSetReinstancingFlow(bool enable);
+LPP_DLL_API(void) LppDisableCompileFinishNotification();
 // END EPIC MOD
 LPP_DLL_API(void) LppApplySettingBool(const char* settingName, int value);
 LPP_DLL_API(void) LppApplySettingInt(const char* settingName, int value);
diff --git a/Engine/Source/Developer/Windows/LiveCoding/Private/External/LC_ClientCommandActions.cpp b/Engine/Source/Developer/Windows/LiveCoding/Private/External/LC_ClientCommandActions.cpp
--- a/Engine/Source/Developer/Windows/LiveCoding/Private/External/LC_ClientCommandActions.cpp
+++ b/Engine/Source/Developer/Windows/LiveCoding/Private/External/LC_ClientCommandActions.cpp
@@ -116,7 +116,7 @@ extern void LiveCodingEndCompile();
 
 // BEGIN EPIC MOD
 extern void LiveCodingPreCompile();
-extern void LiveCodingPostCompile();
+extern void LiveCodingPostCompile(commands::PostCompileResult postCompileResult);
 extern void LiveCodingTriggerReload();
 // END EPIC MOD
 
@@ -217,9 +217,9 @@ bool actions::PreCompile::Execute(const CommandType*, const DuplexPipe* pipe, vo
 }
 
 
-bool actions::PostCompile::Execute(const CommandType*, const DuplexPipe* pipe, void*, const void*, size_t)
+bool actions::PostCompile::Execute(const CommandType* command, const DuplexPipe* pipe, void*, const void*, size_t)
 {
-	LiveCodingPostCompile();
+	LiveCodingPostCompile(command->postCompileResult);
 	pipe->SendAck();
 
 	return true;
diff --git a/Engine/Source/Developer/Windows/LiveCoding/Private/External/LC_ClientStartupThread.cpp b/Engine/Source/Developer/Windows/LiveCoding/Private/External/LC_ClientStartupThread.cpp
--- a/Engine/Source/Developer/Windows/LiveCoding/Private/External/LC_ClientStartupThread.cpp
+++ b/Engine/Source/Developer/Windows/LiveCoding/Private/External/LC_ClientStartupThread.cpp
@@ -385,14 +385,27 @@ void* ClientStartupThread::EnableLazyLoadedModule(const wchar_t* fileName, Windo
 // END EPIC MOD
 
 // BEGIN EPIC MOD
-void ClientStartupThread::EnableReinstancingFlow()
+void ClientStartupThread::SetReinstancingFlow(bool enable)
 {
 	// we cannot wait for commands in the user command thread as long as startup hasn't finished
 	Join();
 
 	if (m_userCommandThread)
 	{
-		m_userCommandThread->EnableReinstancingFlow();
+		m_userCommandThread->SetReinstancingFlow(enable);
+	}
+}
+// END EPIC MOD
+
+// BEGIN EPIC MOD
+void ClientStartupThread::DisableCompileFinishNotification()
+{
+	// we cannot wait for commands in the user command thread as long as startup hasn't finished
+	Join();
+
+	if (m_userCommandThread)
+	{
+		m_userCommandThread->DisableCompileFinishNotification();
 	}
 }
 // END EPIC MOD
diff --git a/Engine/Source/Developer/Windows/LiveCoding/Private/External/LC_ClientStartupThread.h b/Engine/Source/Developer/Windows/LiveCoding/Private/External/LC_ClientStartupThread.h
--- a/Engine/Source/Developer/Windows/LiveCoding/Private/External/LC_ClientStartupThread.h
+++ b/Engine/Source/Developer/Windows/LiveCoding/Private/External/LC_ClientStartupThread.h
@@ -78,7 +78,11 @@ public:
 	// END EPIC MOD
 
 	// BEGIN EPIC MOD
-	void EnableReinstancingFlow();
+	void SetReinstancingFlow(bool Enable);
+	// END EPIC MOD
+
+	// BEGIN EPIC MOD
+	void DisableCompileFinishNotification();
 	// END EPIC MOD
 
 private:
diff --git a/Engine/Source/Developer/Windows/LiveCoding/Private/External/LC_ClientUserCommandThread.cpp b/Engine/Source/Developer/Windows/LiveCoding/Private/External/LC_ClientUserCommandThread.cpp
--- a/Engine/Source/Developer/Windows/LiveCoding/Private/External/LC_ClientUserCommandThread.cpp
+++ b/Engine/Source/Developer/Windows/LiveCoding/Private/External/LC_ClientUserCommandThread.cpp
@@ -555,9 +555,20 @@ void* ClientUserCommandThread::EnableLazyLoadedModule(const wchar_t* fileName, W
 // END EPIC MOD
 
 // BEGIN EPIC MOD
-void ClientUserCommandThread::EnableReinstancingFlow()
+void ClientUserCommandThread::SetReinstancingFlow(bool enable)
 {
-	ProxyCommand<commands::EnableReinstancingFlow>* proxy = new ProxyCommand<commands::EnableReinstancingFlow>(false, 0u);
+	ProxyCommand<commands::SetReinstancingFlow>* proxy = new ProxyCommand<commands::SetReinstancingFlow>(false, 0u);
+	proxy->m_command.processId = Process::Current::GetId();
+	proxy->m_command.enable = enable;
+
+	PushUserCommand(proxy);
+}
+// END EPIC MOD
+
+// BEGIN EPIC MOD
+void ClientUserCommandThread::DisableCompileFinishNotification()
+{
+	ProxyCommand<commands::DisableCompileFinishNotification>* proxy = new ProxyCommand<commands::DisableCompileFinishNotification>(false, 0u);
 	proxy->m_command.processId = Process::Current::GetId();
 
 	PushUserCommand(proxy);
diff --git a/Engine/Source/Developer/Windows/LiveCoding/Private/External/LC_ClientUserCommandThread.h b/Engine/Source/Developer/Windows/LiveCoding/Private/External/LC_ClientUserCommandThread.h
--- a/Engine/Source/Developer/Windows/LiveCoding/Private/External/LC_ClientUserCommandThread.h
+++ b/Engine/Source/Developer/Windows/LiveCoding/Private/External/LC_ClientUserCommandThread.h
@@ -93,7 +93,11 @@ public:
 	// END EPIC MOD
 
 	// BEGIN EPIC MOD
-	void EnableReinstancingFlow();
+	void SetReinstancingFlow(bool enable);
+	// END EPIC MOD
+
+	// BEGIN EPIC MOD
+	void DisableCompileFinishNotification();
 	// END EPIC MOD
 
 	void ApplySettingBool(const char* settingName, int value);
diff --git a/Engine/Source/Developer/Windows/LiveCoding/Private/External/LC_CommandMap.cpp b/Engine/Source/Developer/Windows/LiveCoding/Private/External/LC_CommandMap.cpp
--- a/Engine/Source/Developer/Windows/LiveCoding/Private/External/LC_CommandMap.cpp
+++ b/Engine/Source/Developer/Windows/LiveCoding/Private/External/LC_CommandMap.cpp
@@ -88,7 +88,8 @@ CommandMap::CommandMap(void)
 	RegisterDefaultAction<commands::PreCompile>(m_actions);
 	RegisterDefaultAction<commands::PostCompile>(m_actions);
 	RegisterDefaultAction<commands::TriggerReload>(m_actions);
-	RegisterDefaultAction<commands::EnableReinstancingFlow>(m_actions);
+	RegisterDefaultAction<commands::SetReinstancingFlow>(m_actions);
+	RegisterDefaultAction<commands::DisableCompileFinishNotification>(m_actions);
 	// END EPIC MOD
 }
 
diff --git a/Engine/Source/Developer/Windows/LiveCoding/Private/External/LC_Commands.h b/Engine/Source/Developer/Windows/LiveCoding/Private/External/LC_Commands.h
--- a/Engine/Source/Developer/Windows/LiveCoding/Private/External/LC_Commands.h
+++ b/Engine/Source/Developer/Windows/LiveCoding/Private/External/LC_Commands.h
@@ -319,9 +319,18 @@ namespace commands
 		static const uint32_t ID = FinishedLazyLoadingModules::ID + 1u;
 	};
 
+	enum class PostCompileResult : unsigned char
+	{
+		Success,
+		NoChanges,
+		Failure,
+		Cancelled,
+	};
+
 	struct PostCompile
 	{
 		static const uint32_t ID = PreCompile::ID + 1u;
+		PostCompileResult postCompileResult = PostCompileResult::Success;
 	};
 
 	struct TriggerReload
@@ -329,13 +338,21 @@ namespace commands
 		static const uint32_t ID = PostCompile::ID + 1u;
 	};
 
-	struct EnableReinstancingFlow
+	struct SetReinstancingFlow
 	{
 		static const uint32_t ID = TriggerReload::ID + 1u;
 
+		Process::Id processId;
+		bool enable;
+	};
+
+	struct DisableCompileFinishNotification
+	{
+		static const uint32_t ID = SetReinstancingFlow::ID + 1u;
+
 		Process::Id processId;
 	};
 
-	static const uint32_t COUNT = EnableReinstancingFlow::ID + 1u;
+	static const uint32_t COUNT = DisableCompileFinishNotification::ID + 1u;
 	// END EPIC MOD
 }
diff --git a/Engine/Source/Developer/Windows/LiveCoding/Private/LiveCodingModule.cpp b/Engine/Source/Developer/Windows/LiveCoding/Private/LiveCodingModule.cpp
--- a/Engine/Source/Developer/Windows/LiveCoding/Private/LiveCodingModule.cpp
+++ b/Engine/Source/Developer/Windows/LiveCoding/Private/LiveCodingModule.cpp
@@ -6,6 +6,7 @@
 #include "HAL/IConsoleManager.h"
 #include "Misc/CoreDelegates.h"
 #include "LiveCodingLog.h"
+#include "External/LC_Commands.h"
 #include "External/LC_EntryPoint.h"
 #include "External/LC_API.h"
 #include "Misc/App.h"
@@ -20,6 +21,8 @@
 #include "Algo/BinarySearch.h"
 #if WITH_EDITOR
 	#include "Kismet2/ReloadUtilities.h"
+	#include "Widgets/Notifications/SNotificationList.h"
+	#include "Framework/Notifications/NotificationManager.h"
 #else
 	#include "UObject/Reload.h"
 #endif
@@ -36,6 +39,7 @@ IMPLEMENT_MODULE(FLiveCodingModule, LiveCoding)
 bool GIsCompileActive = false;
 bool GTriggerReload = false;
 bool GHasLoadedPatch = false;
+commands::PostCompileResult GPostCompileResult = commands::PostCompileResult::Success;
 FString GLiveCodingConsolePath;
 FString GLiveCodingConsoleArguments;
 FLiveCodingModule* GLiveCodingModule = nullptr;
@@ -96,6 +100,7 @@ public:
 		if (bHasChanged && !bEnabledMessage)
 		{
 			bEnabledMessage = true;
+			bHasReinstancingOccurred = true;
 			static const TCHAR* Message = TEXT("Object structure changes detected.  LiveCoding re-instancing isn't supported in builds without the editor");
 			UE_LOG(LogLiveCoding, Error, TEXT("%s"), Message);
 #if WITH_ENGINE
@@ -112,18 +117,25 @@ public:
 	{
 	}
 
+	bool HasReinstancingOccurred() const
+	{
+		return bHasReinstancingOccurred;
+	}
+
+	void Reset()
+	{
+		bHasReinstancingOccurred = false;
+	}
+
 private:
 	FLiveCodingModule& LiveCodingModule;
 	mutable bool bEnabledMessage = false;
+	mutable bool bHasReinstancingOccurred = false;
 };
 #endif
 
 FLiveCodingModule::FLiveCodingModule()
-	: bEnabledLastTick(false)
-	, bEnabledForSession(false)
-	, bStarted(false)
-	, bUpdateModulesInTick(false)
-	, FullEnginePluginsDir(FPaths::ConvertRelativePathToFull(FPaths::EnginePluginsDir()))
+	: FullEnginePluginsDir(FPaths::ConvertRelativePathToFull(FPaths::EnginePluginsDir()))
 	, FullProjectDir(FPaths::ConvertRelativePathToFull(FPaths::ProjectDir()))
 	, FullProjectPluginsDir(FPaths::ConvertRelativePathToFull(FPaths::ProjectPluginsDir()))
 {
@@ -456,16 +468,106 @@ void FLiveCodingModule::AttemptSyncLivePatching()
 
 			OnPatchCompleteDelegate.Broadcast();
 			GHasLoadedPatch = false;
+
+			bHasReinstancingOccurred |= Reload->HasReinstancingOccurred();
 		}
 		else if (GTriggerReload)
 		{
 			LppSyncPoint();
 		}
-		Reload.Reset();
+		if (!GIsCompileActive)
+		{
+			static const FString Success("Live coding succeeded");
+
+			// Reset this first so it does its logging first
+			Reload.Reset();
+
+			switch (GPostCompileResult)
+			{
+			case commands::PostCompileResult::Success:
+				if (bHasReinstancingOccurred)
+				{
+#if WITH_EDITOR
+					UE_LOG(LogLiveCoding, Warning, TEXT("%s, %s"), *Success, TEXT("data type changes may cause packaging to fail if assets reference the new or updated data types"));
+#else
+					UE_LOG(LogLiveCoding, Warning, TEXT("%s, %s"), *Success, TEXT("data type changes may cause unexpected failures"));
+#endif
+				}
+				else
+				{
+					UE_LOG(LogLiveCoding, Display, TEXT("%s"), *Success);
+				}
+				break;
+			case commands::PostCompileResult::NoChanges:
+				UE_LOG(LogLiveCoding, Display, TEXT("%s, %s"), *Success, TEXT("no code changes detected"));
+				break;
+			case commands::PostCompileResult::Cancelled:
+				UE_LOG(LogLiveCoding, Error, TEXT("Live coding cancelled"));
+				break;
+			case commands::PostCompileResult::Failure:
+				UE_LOG(LogLiveCoding, Error, TEXT("Live coding failed, please see Live console for more information"));
+				break;
+			default:
+				check(false);
+			}
+
+#if WITH_EDITOR
+			static const FText SuccessText = LOCTEXT("Success", "Live coding succeeded");
+			static const FText NoChangesText = LOCTEXT("NoChanges", "No code changes were detected.");
+			static const FText FailureText = LOCTEXT("Failed", "Live coding failed");
+			static const FText FailureDetailText = LOCTEXT("FailureDetail", "Please see Live Coding console for more information.");
+			static const FText CancelledText = LOCTEXT("Cancelled", "Live coding cancelled");
+			static const FText ReinstancingText = LOCTEXT("Reinstancing", "Data type changes may cause packaging to fail if assets reference the new or updated data types.");
+
+			switch (GPostCompileResult)
+			{
+			case commands::PostCompileResult::Success:
+				if (bHasReinstancingOccurred)
+				{
+					ShowNotification(true, SuccessText, &ReinstancingText);
+				}
+				else
+				{
+					ShowNotification(true, SuccessText, nullptr);
+				}
+				break;
+			case commands::PostCompileResult::NoChanges:
+				ShowNotification(true, SuccessText, &NoChangesText);
+				break;
+			case commands::PostCompileResult::Cancelled:
+				ShowNotification(false, CancelledText, nullptr);
+				break;
+			case commands::PostCompileResult::Failure:
+				ShowNotification(false, FailureText, &FailureDetailText);
+				break;
+			default:
+				check(false);
+			}
+#endif
+		}
+		else
+		{
+			Reload->Reset();
+		}
 	}
 	GTriggerReload = false;
 }
 
+#if WITH_EDITOR
+void FLiveCodingModule::ShowNotification(bool Success, const FText& Title, const FText* SubText)
+{
+	FNotificationInfo Info(Title);
+	Info.ExpireDuration = 5.0f;
+	Info.bUseSuccessFailIcons = true;
+	if (SubText)
+	{
+		Info.SubText = *SubText;
+	}
+	TSharedPtr<SNotificationItem> CompileNotification = FSlateNotificationManager::Get().AddNotification(Info);
+	CompileNotification->SetCompletionState(Success ? SNotificationItem::CS_Success : SNotificationItem::CS_Fail);
+}
+#endif
+
 ILiveCodingModule::FOnPatchCompleteDelegate& FLiveCodingModule::GetOnPatchCompleteDelegate()
 {
 	return OnPatchCompleteDelegate;
@@ -528,7 +630,12 @@ bool FLiveCodingModule::StartLiveCoding()
 		GConfig->GetBool(TEXT("LiveCoding"), TEXT("bEnableReinstancing"), bEnableReinstancing, GEngineIni);
 		if (bEnableReinstancing)
 		{
-			LppEnableReinstancingFlow();
+			LppSetReinstancingFlow(true);
+		}
+
+		if (GEditor != nullptr)
+		{
+			LppDisableCompileFinishNotification();
 		}
 #endif
 
@@ -701,6 +808,9 @@ void FLiveCodingModule::BeginReload()
 	{
 		if (!GLiveCodingModule->Reload.IsValid())
 		{
+			GLiveCodingModule->bHasReinstancingOccurred = false;
+			GLiveCodingModule->bHasPatchBeenLoaded = false;
+			GPostCompileResult = commands::PostCompileResult::Success;
 #if WITH_EDITOR
 			bool bEnableReinstancing = true;
 			GConfig->GetBool(TEXT("LiveCoding"), TEXT("bEnableReinstancing"), bEnableReinstancing, GEngineIni);
@@ -740,9 +850,9 @@ void LiveCodingPreCompile()
 }
 
 // Invoked from LC_ClientCommandActions
-void LiveCodingPostCompile()
+void LiveCodingPostCompile(commands::PostCompileResult PostCompileResult)
 {
-	UE_LOG(LogLiveCoding, Display, TEXT("Live Coding compile done.  See Live Coding console for more information."));
+	GPostCompileResult = PostCompileResult;
 	GIsCompileActive = false;
 }
 
diff --git a/Engine/Source/Developer/Windows/LiveCoding/Private/LiveCodingModule.h b/Engine/Source/Developer/Windows/LiveCoding/Private/LiveCodingModule.h
--- a/Engine/Source/Developer/Windows/LiveCoding/Private/LiveCodingModule.h
+++ b/Engine/Source/Developer/Windows/LiveCoding/Private/LiveCodingModule.h
@@ -15,7 +15,7 @@ class ULiveCodingSettings;
 #if WITH_EDITOR
 class FReload;
 #else
-class IReload;
+class FNullReload;
 #endif
 
 class FLiveCodingModule final : public ILiveCodingModule
@@ -49,10 +49,12 @@ private:
 private:
 	ULiveCodingSettings* Settings;
 	TSharedPtr<ISettingsSection> SettingsSection;
-	bool bEnabledLastTick;
-	bool bEnabledForSession;
-	bool bStarted;
-	bool bUpdateModulesInTick;
+	bool bEnabledLastTick = false;
+	bool bEnabledForSession = false;
+	bool bStarted = false;
+	bool bUpdateModulesInTick = false;
+	bool bHasReinstancingOccurred = false;
+	bool bHasPatchBeenLoaded = false;
 	TSet<FName> ConfiguredModules;
 	TArray<void*> LppPendingTokens;
 
@@ -71,7 +73,7 @@ private:
 #if WITH_EDITOR
 	TUniquePtr<FReload> Reload;
 #else
-	TUniquePtr<IReload> Reload;
+	TUniquePtr<FNullReload> Reload;
 #endif
 
 	bool StartLiveCoding();
@@ -81,5 +83,9 @@ private:
 	void UpdateModules();
 
 	bool ShouldPreloadModule(const FName& Name, const FString& FullFilePath) const;
+
+#if WITH_EDITOR
+	void ShowNotification(bool Success, const FText& Title, const FText* SubText);
+#endif
 };
 
diff --git a/Engine/Source/Developer/Windows/LiveCodingServer/Private/External/LC_LiveProcess.h b/Engine/Source/Developer/Windows/LiveCodingServer/Private/External/LC_LiveProcess.h
--- a/Engine/Source/Developer/Windows/LiveCodingServer/Private/External/LC_LiveProcess.h
+++ b/Engine/Source/Developer/Windows/LiveCodingServer/Private/External/LC_LiveProcess.h
@@ -98,9 +98,9 @@ public:
 	// END EPIC MOD
 
 	// BEGIN EPIC MOD
-	void EnableReinstancingFlow()
+	void SetReinstancingFlow(bool enable)
 	{
-		m_reinstancingFlow = true;
+		m_reinstancingFlow = enable;
 	}
 	bool IsReinstancingFlowEnabled() const
 	{
@@ -108,6 +108,17 @@ public:
 	}
 	// END EPIC MOD
 
+	// BEGIN EPIC MOD
+	void DisableCompileFinishNotification()
+	{
+		m_disableCompileFinishNotification = true;
+	}
+	bool IsDisableCompileFinishNotification()
+	{
+		return m_disableCompileFinishNotification;
+	}
+	// END EPIC MOD
+
 private:
 	Process::Handle m_processHandle;
 	Process::Id m_processId;
@@ -138,6 +149,10 @@ private:
 	bool m_reinstancingFlow = false;
 	// END EPIC MOD
 
+	// BEGIN EPIC MOD
+	bool m_disableCompileFinishNotification = false;
+	// END EPIC MOD
+
 	// loaded modules are not identified by their full path, but by their executable image header.
 	// we do this to ensure that the same executable loaded from a different path is not treated as
 	// a different executable.
diff --git a/Engine/Source/Developer/Windows/LiveCodingServer/Private/External/LC_ServerCommandThread.cpp b/Engine/Source/Developer/Windows/LiveCodingServer/Private/External/LC_ServerCommandThread.cpp
--- a/Engine/Source/Developer/Windows/LiveCodingServer/Private/External/LC_ServerCommandThread.cpp
+++ b/Engine/Source/Developer/Windows/LiveCodingServer/Private/External/LC_ServerCommandThread.cpp
@@ -201,6 +201,18 @@ void ServerCommandThread::RestartTargets(void)
 				liveModule->UnregisterProcess(liveProcess);
 			}
 
+			// BEGIN EPIC MOD
+			if (liveProcess->IsReinstancingFlowEnabled())
+			{
+				--m_reinstancingProcessCount;
+			}
+
+			if (liveProcess->IsDisableCompileFinishNotification())
+			{
+				--m_disableCompileFinishNotificationProcessCount;
+			}
+			// END EPIC MOD
+
 			delete liveProcess;
 
 			processIt = m_liveProcesses.erase(processIt);
@@ -619,6 +631,13 @@ bool ServerCommandThread::HasReinstancingProcess()
 }
 // END EPIC MOD
 
+// BEGIN EPIC MOD
+bool ServerCommandThread::ShowCompileFinishNotification()
+{
+	return m_disableCompileFinishNotificationProcessCount.load() == 0;
+}
+// END EPIC MOD
+
 // BEGIN EPIC MOD - Focus application windows on patch complete
 BOOL CALLBACK FocusApplicationWindows(HWND WindowHandle, LPARAM Lparam)
 {
@@ -752,7 +771,9 @@ bool ServerCommandThread::EnableRequiredModules(const TArray<FString>& RequiredM
 }
 // END EPIC MOD
 
-void ServerCommandThread::CompileChanges(bool didAllProcessesMakeProgress)
+// BEGIN EPIC MOD
+void ServerCommandThread::CompileChanges(bool didAllProcessesMakeProgress, commands::PostCompileResult& postCompileResult)
+// END EPIC MOD
 {
 	// recompile files, if any
 	// EPIC REMOVED: g_theApp.GetMainFrame()->SetBusy(true);
@@ -800,11 +821,17 @@ void ServerCommandThread::CompileChanges(bool didAllProcessesMakeProgress)
 			}
 			else if (CompileResult == ELiveCodingCompileResult::Canceled)
 			{
+				// BEGIN EPIC MOD
+				postCompileResult = commands::PostCompileResult::Cancelled;
+				// END EPIC MOD
 				GLiveCodingServer->GetCompileFinishedDelegate().ExecuteIfBound(ELiveCodingResult::Error, L"Compilation canceled.");
 				return;
 			}
 			else if (CompileResult == ELiveCodingCompileResult::Failure)
 			{
+				// BEGIN EPIC MOD
+				postCompileResult = commands::PostCompileResult::Failure;
+				// END EPIC MOD
 				GLiveCodingServer->GetCompileFinishedDelegate().ExecuteIfBound(ELiveCodingResult::Error, L"Compilation error.");
 				return;
 			}
@@ -812,6 +839,9 @@ void ServerCommandThread::CompileChanges(bool didAllProcessesMakeProgress)
 			// Enable any lazy-loaded modules that we need
 			if (!RequiredModules.IsEmpty() && !EnableRequiredModules(RequiredModules))
 			{
+				// BEGIN EPIC MOD
+				postCompileResult = commands::PostCompileResult::Failure;
+				// END EPIC MOD
 				GLiveCodingServer->GetCompileFinishedDelegate().ExecuteIfBound(ELiveCodingResult::Error, L"Compilation error.");
 				return;
 			}
@@ -852,6 +882,9 @@ void ServerCommandThread::CompileChanges(bool didAllProcessesMakeProgress)
 				{
 					LC_ERROR_USER("Live coding is not enabled for %S.", ModuleFileName.c_str());
 					LC_ERROR_USER("Configure the list of enabled modules from the Live Coding section of the editor preferences window.");
+					// BEGIN EPIC MOD
+					postCompileResult = commands::PostCompileResult::Failure;
+					// END EPIC MOD
 					GLiveCodingServer->GetCompileFinishedDelegate().ExecuteIfBound(ELiveCodingResult::Error, *FString::Printf(TEXT("Live coding not enabled for %s"), ModuleFileName.c_str()));
 					return;
 				}
@@ -1012,37 +1045,51 @@ void ServerCommandThread::CompileChanges(bool didAllProcessesMakeProgress)
 	}
 
 	// BEGIN EPIC MOD - Custom hooks for finishing compile
+	bool setFocus = false;
 	switch (updateError)
 	{
 		case LiveModule::ErrorType::NO_CHANGE:
+			postCompileResult = commands::PostCompileResult::NoChanges;
 			GLiveCodingServer->GetCompileFinishedDelegate().ExecuteIfBound(ELiveCodingResult::Success, L"No changes detected.");
+			setFocus = !ShowCompileFinishNotification();
 			break;
 
 		case LiveModule::ErrorType::COMPILE_ERROR:
+			postCompileResult = commands::PostCompileResult::Failure;
 			GLiveCodingServer->GetCompileFinishedDelegate().ExecuteIfBound(ELiveCodingResult::Error, L"Compilation error.");
 			break;
 
 		case LiveModule::ErrorType::LINK_ERROR:
+			postCompileResult = commands::PostCompileResult::Failure;
 			GLiveCodingServer->GetCompileFinishedDelegate().ExecuteIfBound(ELiveCodingResult::Error, L"Linker error.");
 			break;
 
 		case LiveModule::ErrorType::LOAD_PATCH_ERROR:
+			postCompileResult = commands::PostCompileResult::Failure;
 			GLiveCodingServer->GetCompileFinishedDelegate().ExecuteIfBound(ELiveCodingResult::Error, L"Could not load patch image.");
 			break;
 
 		case LiveModule::ErrorType::ACTIVATE_PATCH_ERROR:
+			postCompileResult = commands::PostCompileResult::Failure;
 			GLiveCodingServer->GetCompileFinishedDelegate().ExecuteIfBound(ELiveCodingResult::Error, L"Could not activate patch.");
 			break;
 
 		case LiveModule::ErrorType::SUCCESS:
+			postCompileResult = commands::PostCompileResult::Success;
 			GLiveCodingServer->GetCompileFinishedDelegate().ExecuteIfBound(ELiveCodingResult::Success, L"Patch creation successful.");
-			EnumWindows(FocusApplicationWindows, (LPARAM)&m_liveProcesses);
+			setFocus = true;
 			break;
 
 		default:
+			postCompileResult = commands::PostCompileResult::Success;
 			GLiveCodingServer->GetCompileFinishedDelegate().ExecuteIfBound(ELiveCodingResult::Success, L"Finished.");
 			break;
 	}
+
+	if (setFocus)
+	{
+		EnumWindows(FocusApplicationWindows, (LPARAM)&m_liveProcesses);
+	}
 	// END EPIC MOD
 	
 	LC_LOG_USER("---------- Finished (%.3fs) ----------", scope.ReadSeconds());
@@ -1077,7 +1124,7 @@ void ServerCommandThread::CallPrecompileHooks(bool didAllProcessesMakeProgress)
 	}
 }
 
-void ServerCommandThread::CallPostcompileHooks(bool didAllProcessesMakeProgress)
+void ServerCommandThread::CallPostcompileHooks(bool didAllProcessesMakeProgress, commands::PostCompileResult postCompileResult)
 {
 	// when all processes made progress, none of them is being held in the debugger which means it is safe to
 	// communicate with the client, call hooks, use synchronization points, etc.
@@ -1098,7 +1145,7 @@ void ServerCommandThread::CallPostcompileHooks(bool didAllProcessesMakeProgress)
 	for (size_t i = 0u; i < count; ++i)
 	{
 		LiveProcess* liveProcess = m_liveProcesses[i];
-		liveProcess->GetPipe()->SendCommandAndWaitForAck(commands::PostCompile{}, nullptr, 0u);
+		liveProcess->GetPipe()->SendCommandAndWaitForAck(commands::PostCompile{ postCompileResult }, nullptr, 0u);
 	}
 }
 // END EPIC MOD
@@ -1210,6 +1257,18 @@ Thread::ReturnValue ServerCommandThread::CompileThread(void)
 							liveModule->UnregisterProcess(liveProcess);
 						}
 
+						// BEGIN EPIC MOD
+						if (liveProcess->IsReinstancingFlowEnabled())
+						{
+							--m_reinstancingProcessCount;
+						}
+
+						if (liveProcess->IsDisableCompileFinishNotification())
+						{
+							--m_disableCompileFinishNotificationProcessCount;
+						}
+						// END EPIC MOD
+
 						delete liveProcess;
 
 						processIt = m_liveProcesses.erase(processIt);
@@ -1286,10 +1345,13 @@ Thread::ReturnValue ServerCommandThread::CompileThread(void)
 			CallPrecompileHooks(didAllProcessesMakeProgress);
 			// END EPIC MOD
 
-			CompileChanges(didAllProcessesMakeProgress);
+			// BEGIN EPIC MOD 
+			commands::PostCompileResult postCompileResult = commands::PostCompileResult::Success;
+			CompileChanges(didAllProcessesMakeProgress, postCompileResult);
+			// END EPIC MOD
 
 			// BEGIN EPIC MOD - Add the ability for pre and post compile notifications
-			CallPostcompileHooks(didAllProcessesMakeProgress);
+			CallPostcompileHooks(didAllProcessesMakeProgress, postCompileResult);
 			// END EPIC MOD
 
 			// BEGIN EPIC MOD - Non-destructive compile
@@ -1377,7 +1439,8 @@ Thread::ReturnValue ServerCommandThread::CommandThread(DuplexPipeServer* pipe, E
 	commandMap.RegisterAction<actions::EnableLazyLoadedModule>();
 	// END EPIC MOD
 	// BEGIN EPIC MOD
-	commandMap.RegisterAction<actions::EnableReinstancingFlow>();
+	commandMap.RegisterAction<actions::SetReinstancingFlow>();
+	commandMap.RegisterAction<actions::DisableCompileFinishNotification>();
 	// END EPIC MOD
 
 	for (;;)
@@ -1735,7 +1798,7 @@ bool ServerCommandThread::actions::EnableLazyLoadedModule::Execute(const Command
 // END EPIC MOD
 
 // BEGIN EPIC MOD
-bool ServerCommandThread::actions::EnableReinstancingFlow::Execute(const CommandType* command, const DuplexPipe* pipe, void* context, const void*, size_t)
+bool ServerCommandThread::actions::SetReinstancingFlow::Execute(const CommandType* command, const DuplexPipe* pipe, void* context, const void*, size_t)
 {
 	ServerCommandThread* commandThread = static_cast<ServerCommandThread*>(context);
 
@@ -1751,11 +1814,49 @@ bool ServerCommandThread::actions::EnableReinstancingFlow::Execute(const Command
 	{
 		if (process->GetProcessId() == command->processId)
 		{
-			if (!process->IsReinstancingFlowEnabled())
+			if (command->enable)
 			{
-				++commandThread->m_reinstancingProcessCount;
-				process->EnableReinstancingFlow();
+				if (!process->IsReinstancingFlowEnabled())
+				{
+					++commandThread->m_reinstancingProcessCount;
+				}
+			}
+			else
+			{
+				if (process->IsReinstancingFlowEnabled())
+				{
+					--commandThread->m_reinstancingProcessCount;
+				}
 			}
+			process->SetReinstancingFlow(command->enable);
+		}
+	}
+
+	pipe->SendAck();
+
+	return true;
+}
+// END EPIC MOD
+
+// BEGIN EPIC MOD
+bool ServerCommandThread::actions::DisableCompileFinishNotification::Execute(const CommandType* command, const DuplexPipe* pipe, void* context, const void*, size_t)
+{
+	ServerCommandThread* commandThread = static_cast<ServerCommandThread*>(context);
+
+	if (!commandThread->HasReinstancingProcess())
+	{
+		LC_WARNING_USER("Quick restart disabled when re-instancing is enabled.");
+	}
+
+	// protect against accepting this command while compilation is already in progress
+	CriticalSection::ScopedLock lock(&commandThread->m_actionCS);
+
+	for (LiveProcess* process : commandThread->m_liveProcesses)
+	{
+		if (process->GetProcessId() == command->processId)
+		{
+			++commandThread->m_disableCompileFinishNotificationProcessCount;
+			process->DisableCompileFinishNotification();
 		}
 	}
 
diff --git a/Engine/Source/Developer/Windows/LiveCodingServer/Private/External/LC_ServerCommandThread.h b/Engine/Source/Developer/Windows/LiveCodingServer/Private/External/LC_ServerCommandThread.h
--- a/Engine/Source/Developer/Windows/LiveCodingServer/Private/External/LC_ServerCommandThread.h
+++ b/Engine/Source/Developer/Windows/LiveCodingServer/Private/External/LC_ServerCommandThread.h
@@ -40,6 +40,10 @@ public:
 	bool HasReinstancingProcess();
 	// END EPIC MOD
 
+	// BEGIN EPIC MOD
+	bool ShowCompileFinishNotification();
+	// END EPIC MOD
+
 private:
 	scheduler::Task<LiveModule*>* LoadModule(Process::Id processId, void* moduleBase, const wchar_t* modulePath, scheduler::TaskBase* taskRoot);
 	bool UnloadModule(Process::Id processId, const wchar_t* modulePath);
@@ -51,7 +55,7 @@ private:
 
 	// BEGIN EPIC MOD - Add the ability for pre and post compile notifications
 	void CallPrecompileHooks(bool didAllProcessesMakeProgress);
-	void CallPostcompileHooks(bool didAllProcessesMakeProgress);
+	void CallPostcompileHooks(bool didAllProcessesMakeProgress, commands::PostCompileResult postCompileResult);
 	// END EPIC MOD
 
 	struct CommandThreadContext
@@ -71,8 +75,9 @@ private:
 
 	LiveProcess* FindProcessById(Process::Id processId);
 
-	void CompileChanges(bool didAllProcessesMakeProgress);
-
+	// BEGIN EPIC MOD
+	void CompileChanges(bool didAllProcessesMakeProgress, commands::PostCompileResult& postCompileResult);
+	// END EPIC MOD
 
 	// actions
 	struct actions
@@ -122,7 +127,11 @@ private:
 		// END EPIC MOD
 
 		// BEGIN EPIC MOD
-		DECLARE_ACTION(EnableReinstancingFlow);
+		DECLARE_ACTION(SetReinstancingFlow);
+		// END EPIC MOD
+
+		// BEGIN EPIC MOD
+		DECLARE_ACTION(DisableCompileFinishNotification);
 		// END EPIC MOD
 
 		#undef DECLARE_ACTION
@@ -156,6 +165,10 @@ private:
 	std::atomic<unsigned int> m_reinstancingProcessCount = 0;
 	// END EPIC MOD
 
+	// BEGIN EPIC MOD
+	std::atomic<unsigned int> m_disableCompileFinishNotificationProcessCount = 0;
+	// END EPIC MOD
+
 	// BEGIN EPIC MOD - Non-destructive compile
 	std::vector<std::pair<std::wstring, std::wstring>> m_restoreFiles;
 	// END EPIC MOD
diff --git a/Engine/Source/Developer/Windows/LiveCodingServer/Private/LiveCodingServer.cpp b/Engine/Source/Developer/Windows/LiveCodingServer/Private/LiveCodingServer.cpp
--- a/Engine/Source/Developer/Windows/LiveCodingServer/Private/LiveCodingServer.cpp
+++ b/Engine/Source/Developer/Windows/LiveCodingServer/Private/LiveCodingServer.cpp
@@ -118,4 +118,9 @@ ILiveCodingServer::FSetVisibleDelegate& FLiveCodingServer::GetSetVisibleDelegate
 bool FLiveCodingServer::HasReinstancingProcess()
 {
 	return CommandThread->HasReinstancingProcess();
-}
\ No newline at end of file
+}
+
+bool FLiveCodingServer::ShowCompileFinishNotification()
+{
+	return CommandThread->ShowCompileFinishNotification();
+}
diff --git a/Engine/Source/Developer/Windows/LiveCodingServer/Private/LiveCodingServer.h b/Engine/Source/Developer/Windows/LiveCodingServer/Private/LiveCodingServer.h
--- a/Engine/Source/Developer/Windows/LiveCodingServer/Private/LiveCodingServer.h
+++ b/Engine/Source/Developer/Windows/LiveCodingServer/Private/LiveCodingServer.h
@@ -33,6 +33,7 @@ public:
 	virtual FShowConsoleDelegate& GetShowConsoleDelegate() override final;
 	virtual FSetVisibleDelegate& GetSetVisibleDelegate() override final;
 	virtual bool HasReinstancingProcess() override final;
+	virtual bool ShowCompileFinishNotification() override final;
 
 private:
 	std::wstring ProcessGroupName;
diff --git a/Engine/Source/Developer/Windows/LiveCodingServer/Public/ILiveCodingServer.h b/Engine/Source/Developer/Windows/LiveCodingServer/Public/ILiveCodingServer.h
--- a/Engine/Source/Developer/Windows/LiveCodingServer/Public/ILiveCodingServer.h
+++ b/Engine/Source/Developer/Windows/LiveCodingServer/Public/ILiveCodingServer.h
@@ -85,5 +85,7 @@ public:
 	virtual FSetVisibleDelegate& GetSetVisibleDelegate() = 0;
 
 	virtual bool HasReinstancingProcess() = 0;
+
+	virtual bool ShowCompileFinishNotification() = 0;
 };
 
diff --git a/Engine/Source/Editor/UnrealEd/Private/Kismet2/ReloadUtilities.cpp b/Engine/Source/Editor/UnrealEd/Private/Kismet2/ReloadUtilities.cpp
--- a/Engine/Source/Editor/UnrealEd/Private/Kismet2/ReloadUtilities.cpp
+++ b/Engine/Source/Editor/UnrealEd/Private/Kismet2/ReloadUtilities.cpp
@@ -739,6 +739,7 @@ bool FReload::GetEnableReinstancing(bool bHasChanged) const
 	if (bHasChanged && !bEnableReinstancing  && !bEnabledMessage)
 	{
 		bEnabledMessage = true;
+		bHasReinstancingOccurred = true;
 		Ar.Logf(ELogVerbosity::Display, TEXT("Re-instancing has been disabled.  Some changes will be ignored."));
 	}
 	return bEnableReinstancing;
@@ -750,6 +751,10 @@ void FReload::Reset()
 	FunctionRemap.Empty();
 	ReconstructedCDOsMap.Empty();
 	ReinstancedClasses.Empty();
+	ReinstancedEnums.Empty();
+	ReinstancedStructs.Empty();
+	Packages.Empty();
+	bHasReinstancingOccurred = false;
 }
 
 void FReload::UpdateStats(FReinstanceStats& Stats, void* New, void* Old)
@@ -802,6 +807,11 @@ void FReload::NotifyChange(UClass* New, UClass* Old)
 {
 	UpdateStats(ClassStats, New, Old);
 
+	if (New != Old)
+	{
+		bHasReinstancingOccurred = true;
+	}
+
 	// Ignore new classes
 	if (Old != nullptr)
 	{
@@ -824,6 +834,12 @@ void FReload::NotifyChange(UEnum* New, UEnum* Old)
 {
 	UpdateStats(EnumStats, New, Old);
 
+	if (New != Old)
+	{
+		bHasReinstancingOccurred = true;
+	}
+
+
 	if (Old != nullptr)
 	{
 		UEnum* NewIfChanged = Old != New ? New : nullptr; // supporting code detects unchanged based on null new pointer
@@ -836,6 +852,11 @@ void FReload::NotifyChange(UScriptStruct* New, UScriptStruct* Old)
 {
 	UpdateStats(StructStats, New, Old);
 
+	if (New != Old)
+	{
+		bHasReinstancingOccurred = true;
+	}
+
 	if (Old != nullptr)
 	{
 		UScriptStruct* NewIfChanged = Old != New ? New : nullptr; // supporting code detects unchanged based on null new pointer
diff --git a/Engine/Source/Editor/UnrealEd/Public/Kismet2/ReloadUtilities.h b/Engine/Source/Editor/UnrealEd/Public/Kismet2/ReloadUtilities.h
--- a/Engine/Source/Editor/UnrealEd/Public/Kismet2/ReloadUtilities.h
+++ b/Engine/Source/Editor/UnrealEd/Public/Kismet2/ReloadUtilities.h
@@ -69,6 +69,14 @@ public:
 		bEnableReinstancing = bInEnableReinstancing;
 	}
 
+	/**
+	 * Return true if anything was re-instanced
+	 */
+	bool HasReinstancingOccurred() const
+	{
+		return bHasReinstancingOccurred;
+	}
+
 private:
 
 	struct FReinstanceStats
@@ -81,6 +89,11 @@ private:
 		{
 			return New + Changed + Unchanged != 0;
 		}
+
+		bool HasReinstancingOccurred() const
+		{
+			return New + Changed != 0;
+		}
 	};
 
 	struct FBlueprintUpdateInfo
@@ -158,4 +171,5 @@ private:
 	int32 NumFunctionsRemapped = 0;
 	int32 NumScriptStructsRemapped = 0;
 	mutable bool bEnabledMessage = false;
+	mutable bool bHasReinstancingOccurred = false;
 };
diff --git a/Engine/Source/Programs/LiveCodingConsole/Private/LiveCodingConsole.cpp b/Engine/Source/Programs/LiveCodingConsole/Private/LiveCodingConsole.cpp
--- a/Engine/Source/Programs/LiveCodingConsole/Private/LiveCodingConsole.cpp
+++ b/Engine/Source/Programs/LiveCodingConsole/Private/LiveCodingConsole.cpp
@@ -451,28 +451,46 @@ private:
 	{
 		if(CompileNotification.IsValid())
 		{
-			if(Result == ELiveCodingResult::Success)
+			if (Result == ELiveCodingResult::Success)
 			{
-				CompileNotification->SetText(FText::FromString(Status));
-				CompileNotification->SetCompletionState(SNotificationItem::CS_Success);
-				CompileNotification->SetExpireDuration(1.5f);
-				CompileNotification->SetFadeOutDuration(0.4f);
+				if (Server.ShowCompileFinishNotification())
+				{
+					CompileNotification->SetText(FText::FromString(Status));
+					CompileNotification->SetCompletionState(SNotificationItem::CS_Success);
+					CompileNotification->SetExpireDuration(1.5f);
+					CompileNotification->SetFadeOutDuration(0.4f);
+				}
+				else
+				{
+					CompileNotification->SetExpireDuration(0.0f);
+					CompileNotification->SetFadeOutDuration(0.1f);
+				}
 			}
-			else if(HasCancelledBuild())
+			else if (HasCancelledBuild())
 			{
 				CompileNotification->SetExpireDuration(0.0f);
 				CompileNotification->SetFadeOutDuration(0.1f);
 			}
 			else
 			{
-				CompileNotification->SetText(FText::FromString(Status));
-				CompileNotification->SetCompletionState(SNotificationItem::CS_Fail);
-				CompileNotification->SetExpireDuration(5.0f);
-				CompileNotification->SetFadeOutDuration(2.0f);
+				if (Server.ShowCompileFinishNotification())
+				{
+					CompileNotification->SetText(FText::FromString(Status));
+					CompileNotification->SetCompletionState(SNotificationItem::CS_Fail);
+					CompileNotification->SetExpireDuration(5.0f);
+					CompileNotification->SetFadeOutDuration(2.0f);
+				}
+				else
+				{
+					CompileNotification->SetExpireDuration(0.0f);
+					CompileNotification->SetFadeOutDuration(0.1f);
+					ShowConsole();
+				}
 			}
-			CompileNotification->ExpireAndFadeout();
-			CompileNotification.Reset();
+
 		}
+		CompileNotification->ExpireAndFadeout();
+		CompileNotification.Reset();
 	}
 
 	void OnStatusChangedAsync(const FString& Status)
