Fixes for no unity errors
#rnx

diff --git a/Engine/Plugins/Runtime/Nvidia/Reflex/Source/Private/ReflexLatencyMarkers.cpp b/Engine/Plugins/Runtime/Nvidia/Reflex/Source/Private/ReflexLatencyMarkers.cpp
--- a/Engine/Plugins/Runtime/Nvidia/Reflex/Source/Private/ReflexLatencyMarkers.cpp
+++ b/Engine/Plugins/Runtime/Nvidia/Reflex/Source/Private/ReflexLatencyMarkers.cpp
@@ -6,6 +6,8 @@
 
 #include "RHI.h"
 #include "Engine/Console.h"
+#include "Engine/LocalPlayer.h"
+#include "Engine/World.h"
 #include "GameFramework/PlayerController.h"
 #include "Kismet/GameplayStatics.h"
 
diff --git a/Engine/Plugins/Runtime/Nvidia/Reflex/Source/Public/ReflexLatencyMarkers.h b/Engine/Plugins/Runtime/Nvidia/Reflex/Source/Public/ReflexLatencyMarkers.h
--- a/Engine/Plugins/Runtime/Nvidia/Reflex/Source/Public/ReflexLatencyMarkers.h
+++ b/Engine/Plugins/Runtime/Nvidia/Reflex/Source/Public/ReflexLatencyMarkers.h
@@ -3,6 +3,7 @@
 #pragma once
 
 #include "CoreMinimal.h"
+#include "Misc/CoreMisc.h"
 #include "Tickable.h"
 #include "Performance/LatencyMarkerModule.h"
 
diff --git a/Engine/Source/Developer/Windows/LiveCoding/Private/LiveCodingModule.cpp b/Engine/Source/Developer/Windows/LiveCoding/Private/LiveCodingModule.cpp
--- a/Engine/Source/Developer/Windows/LiveCoding/Private/LiveCodingModule.cpp
+++ b/Engine/Source/Developer/Windows/LiveCoding/Private/LiveCodingModule.cpp
@@ -20,6 +20,7 @@
 #include "Algo/Sort.h"
 #include "Algo/BinarySearch.h"
 #if WITH_EDITOR
+	#include "Editor.h"
 	#include "Kismet2/ReloadUtilities.h"
 	#include "Widgets/Notifications/SNotificationList.h"
 	#include "Framework/Notifications/NotificationManager.h"
diff --git a/Engine/Source/Editor/Persona/Private/AnimationEditorViewportClient.h b/Engine/Source/Editor/Persona/Private/AnimationEditorViewportClient.h
--- a/Engine/Source/Editor/Persona/Private/AnimationEditorViewportClient.h
+++ b/Engine/Source/Editor/Persona/Private/AnimationEditorViewportClient.h
@@ -12,6 +12,7 @@
 #include "Animation/DebugSkelMeshComponent.h"
 #include "IPersonaPreviewScene.h"
 #include "Preferences/PersonaOptions.h"
+#include "ScopedTransaction.h"
 
 class FCanvas;
 class UPersonaOptions;
