Fixing minor feedback issue where a warning was being generated incorrectly.

#rnx

diff --git a/Engine/Source/Developer/Windows/LiveCodingServer/Private/External/LC_ServerCommandThread.cpp b/Engine/Source/Developer/Windows/LiveCodingServer/Private/External/LC_ServerCommandThread.cpp
--- a/Engine/Source/Developer/Windows/LiveCodingServer/Private/External/LC_ServerCommandThread.cpp
+++ b/Engine/Source/Developer/Windows/LiveCodingServer/Private/External/LC_ServerCommandThread.cpp
@@ -1802,7 +1802,7 @@ bool ServerCommandThread::actions::SetReinstancingFlow::Execute(const CommandTyp
 {
 	ServerCommandThread* commandThread = static_cast<ServerCommandThread*>(context);
 
-	if (!commandThread->HasReinstancingProcess())
+	if (!commandThread->HasReinstancingProcess() && command->enable)
 	{
 		LC_WARNING_USER("Quick restart disabled when re-instancing is enabled.");
 	}
@@ -1843,11 +1843,6 @@ bool ServerCommandThread::actions::DisableCompileFinishNotification::Execute(con
 {
 	ServerCommandThread* commandThread = static_cast<ServerCommandThread*>(context);
 
-	if (!commandThread->HasReinstancingProcess())
-	{
-		LC_WARNING_USER("Quick restart disabled when re-instancing is enabled.");
-	}
-
 	// protect against accepting this command while compilation is already in progress
 	CriticalSection::ScopedLock lock(&commandThread->m_actionCS);
 
