Fix to LiveCoding symbol resolution that could result in a function being patched to the wrong implementation.  This would result in unpredictable behavior and crashes.

#rnx

diff --git a/Engine/Source/Developer/Windows/LiveCodingServer/Private/External/LC_SymbolReconstruction.cpp b/Engine/Source/Developer/Windows/LiveCodingServer/Private/External/LC_SymbolReconstruction.cpp
--- a/Engine/Source/Developer/Windows/LiveCodingServer/Private/External/LC_SymbolReconstruction.cpp
+++ b/Engine/Source/Developer/Windows/LiveCodingServer/Private/External/LC_SymbolReconstruction.cpp
@@ -649,7 +649,14 @@ walkOpenSymbols:
 	}
 
 	// populate a cache of all DIA names for all potential contributions once
-	types::StringMap<uint32_t> diaNameToRva;
+	// BEGIN EPIC MOD
+	struct DiaRvaData
+	{
+		uint32_t rva;
+		bool valid;
+	};
+	types::StringMap<DiaRvaData> diaNameToRva;
+	// END EPIC MOD
 	diaNameToRva.reserve(potentialContributionRVAsAcrossAllMissingSymbols.size());
 
 	types::unordered_map<uint32_t, IDiaSymbol*> rvaToDiaSymbol;
@@ -672,7 +679,14 @@ walkOpenSymbols:
 			const std::wstring& diaName = dia::GetSymbolName(diaSymbol).GetString();
 			const ImmutableString& name = string::ToUtf8String(diaName);
 
-			diaNameToRva.insert(std::make_pair(name, rva));
+			// BEGIN EPIC MOD
+			auto ixb = diaNameToRva.insert(std::make_pair(name, DiaRvaData{ rva, true }));
+			if (!ixb.second && ixb.first->second.rva != rva && ixb.first->second.valid)
+			{
+				LC_LOG_DEV("Ignoring Dia symbol %s for fast path because multiple RVAs claim the symbol", name.c_str());
+				ixb.first->second.valid = false;
+			}
+			// END EPIC MOD
 			rvaToDiaSymbol.insert(std::make_pair(rva, diaSymbol));
 		}
 	}
@@ -693,11 +707,15 @@ walkOpenSymbols:
 		const std::string& coffUndecoratedName = symbols::UndecorateSymbolName(missingSymbolName);
 
 		auto diaNameIt = diaNameToRva.find(ImmutableString(coffUndecoratedName.c_str()));
-		if (diaNameIt != diaNameToRva.end())
+		// BEGIN EPIC MOD
+		if (diaNameIt != diaNameToRva.end() && diaNameIt->second.valid)
+		// END EPIC MOD
 		{
 			// fast path.
 			// there is a symbol that matches the exact name of the symbol in the .obj file
-			const uint32_t rva = diaNameIt->second;
+			// BEGIN EPIC MOD
+			const uint32_t rva = diaNameIt->second.rva;
+			// END EPIC MOD
 
 			LC_LOG_DEV("Fast path, found symbol %s at 0x%X", missingSymbolName.c_str(), rva);
 
