Removing validation skipping from the SM6.6 code path as it's a fully supported feature in our version of DXC.

diff --git a/Engine/Source/Developer/Windows/ShaderFormatD3D/Private/D3DShaderCompilerDXC.cpp b/Engine/Source/Developer/Windows/ShaderFormatD3D/Private/D3DShaderCompilerDXC.cpp
--- a/Engine/Source/Developer/Windows/ShaderFormatD3D/Private/D3DShaderCompilerDXC.cpp
+++ b/Engine/Source/Developer/Windows/ShaderFormatD3D/Private/D3DShaderCompilerDXC.cpp
@@ -668,19 +668,6 @@ bool CompileAndProcessD3DShaderDXC(FString& PreprocessedShaderSource,
 
 	const TCHAR* ValidatorVersion = nullptr; // Whatever default validator chosen by DXC
 
-#if USE_SHADER_MODEL_6_6
-	if (Input.Target.GetPlatform() == SP_PCD3D_SM6)
-	{
-		// Experimental shader model 6.6 requires turning off DXIL validation and signing as of 2021-04-04.
-		// Unsigned shader will only run in Windows 10 developer mode, and only when experimental shader models are enabled.
-		// TODO: This should be removed when dxil.dll supports SM 6.6 validation.
-		DXCFlags |= D3DCOMPILE_SKIP_VALIDATION;
-
-		// Shader model 6.6 requires validator 1.6
-		ValidatorVersion = TEXT("1.6");
-	}
-#endif
-
 	const bool bKeepDebugInfo = Input.Environment.CompilerFlags.Contains(CFLAG_KeepDebugInfo);
 
 	FDxcArguments Args
