LC_StringUtil.cpp workaround is also required for MSVC Clang 12.0.0

diff --git a/Engine/Source/Developer/Windows/LiveCoding/Private/External/LC_StringUtil.cpp b/Engine/Source/Developer/Windows/LiveCoding/Private/External/LC_StringUtil.cpp
--- a/Engine/Source/Developer/Windows/LiveCoding/Private/External/LC_StringUtil.cpp
+++ b/Engine/Source/Developer/Windows/LiveCoding/Private/External/LC_StringUtil.cpp
@@ -7,9 +7,9 @@
 #include "LC_StringUtil.h"
 // BEGIN EPIC MOD
 #if defined(__clang__) && defined(_MSC_VER) && _MSVC_LANG > 201402L
-#if __clang_major__ > 11
+#if __clang_major__ > 12
 // For reference: https://bugs.llvm.org/show_bug.cgi?id=41226#c16
-#error Re-check whether this workaround is still needed for clang 12.0 and higher.
+#error Re-check whether this workaround is still needed for clang 13.0 and higher.
 #endif
 #include <wchar.h>
 #endif
@@ -168,9 +168,9 @@ namespace detail
 	static size_t WideFindOffset(const std::wstring& haystack, const std::wstring& needle)
 	{
 #if defined(__clang__) && defined(_MSC_VER) && _MSVC_LANG > 201402L
-#if __clang_major__ > 11
+#if __clang_major__ > 12
 		// For reference: https://bugs.llvm.org/show_bug.cgi?id=41226#c16
-		#error Re-check whether this workaround is still needed for clang 12.0 and higher.
+		#error Re-check whether this workaround is still needed for clang 13.0 and higher.
 #endif
 		const wchar_t* found = ::wcsstr(haystack.c_str(), needle.c_str());
 		if (found)
