Fixed issue where starting PIE after starting a LiveCoding compile would lockup the editor.

diff --git a/Engine/Source/Developer/Windows/LiveCoding/Private/LiveCodingModule.cpp b/Engine/Source/Developer/Windows/LiveCoding/Private/LiveCodingModule.cpp
--- a/Engine/Source/Developer/Windows/LiveCoding/Private/LiveCodingModule.cpp
+++ b/Engine/Source/Developer/Windows/LiveCoding/Private/LiveCodingModule.cpp
@@ -376,15 +376,15 @@ void FLiveCodingModule::Tick()
 
 void FLiveCodingModule::AttemptSyncLivePatching()
 {
-	while (LppPendingTokens.Num() > 0)
+
+	// We use to wait for all commands to finish, but that causes a lock up if starting PIE after a compilation 
+	// request caused another command to be sent to the live coding console.  For example, the registering of 
+	// another lazy load module at PIE start would cause this problem.
+	for (int Index = LppPendingTokens.Num(); Index-- > 0;)
 	{
-		if (LppTryWaitForToken(LppPendingTokens[0]))
-		{
-			LppPendingTokens.RemoveAt(0);
-		}
-		else
+		if (LppTryWaitForToken(LppPendingTokens[Index]))
 		{
-			return;
+			LppPendingTokens.RemoveAt(Index);
 		}
 	}
 
