Remove the bEnableReinstancing global setting to reduce confusion now that a project setting is available.

#rnx

diff --git a/Engine/Config/BaseEngine.ini b/Engine/Config/BaseEngine.ini
--- a/Engine/Config/BaseEngine.ini
+++ b/Engine/Config/BaseEngine.ini
@@ -2819,7 +2819,6 @@ DefaultSizeInMegabyte=64
 UnrealBuildTool=Engine/Binaries/DotNET/UnrealBuildTool/UnrealBuildTool.exe
 
 [LiveCoding]
-bEnableReinstancing=false
 bDisableActionLimit=false
 ActionLimit=100
 
diff --git a/Engine/Source/Developer/Windows/LiveCoding/Private/LiveCodingModule.cpp b/Engine/Source/Developer/Windows/LiveCoding/Private/LiveCodingModule.cpp
--- a/Engine/Source/Developer/Windows/LiveCoding/Private/LiveCodingModule.cpp
+++ b/Engine/Source/Developer/Windows/LiveCoding/Private/LiveCodingModule.cpp
@@ -830,13 +830,7 @@ void FLiveCodingModule::BeginReload()
 bool FLiveCodingModule::IsReinstancingEnabled() const
 {
 #if WITH_EDITOR
-	if (Settings->bEnableReinstancing)
-	{
-		return true;
-	}
-	bool bEnableReinstancing = true;
-	GConfig->GetBool(TEXT("LiveCoding"), TEXT("bEnableReinstancing"), bEnableReinstancing, GEngineIni);
-	return bEnableReinstancing;
+	return Settings->bEnableReinstancing;
 #else
 	return false;
 #endif
