Implements CFLAG_AllowRealTypes on SP_PCD3D_SM6

diff --git a/Engine/Source/Developer/Windows/ShaderFormatD3D/Private/D3DShaderCompiler.cpp b/Engine/Source/Developer/Windows/ShaderFormatD3D/Private/D3DShaderCompiler.cpp
--- a/Engine/Source/Developer/Windows/ShaderFormatD3D/Private/D3DShaderCompiler.cpp
+++ b/Engine/Source/Developer/Windows/ShaderFormatD3D/Private/D3DShaderCompiler.cpp
@@ -1045,6 +1045,11 @@ void CompileD3DShader(const FShaderCompilerInput& Input, FShaderCompilerOutput&
 		AdditionalDefines.SetDefine(TEXT("PLATFORM_SUPPORTS_STATIC_SAMPLERS"), 1);
 		AdditionalDefines.SetDefine(TEXT("PLATFORM_SUPPORTS_DIAGNOSTIC_BUFFER"), 1);
 		AdditionalDefines.SetDefine(TEXT("COMPILER_SUPPORTS_NOINLINE"), 1);
+
+		if (Language == ELanguage::SM6 && Input.Environment.CompilerFlags.Contains(CFLAG_AllowRealTypes))
+		{
+			AdditionalDefines.SetDefine(TEXT("PLATFORM_SUPPORTS_REAL_TYPES"), 1);
+		}
 	}
 	else
 	{
diff --git a/Engine/Source/Developer/Windows/ShaderFormatD3D/Private/D3DShaderCompilerDXC.cpp b/Engine/Source/Developer/Windows/ShaderFormatD3D/Private/D3DShaderCompilerDXC.cpp
--- a/Engine/Source/Developer/Windows/ShaderFormatD3D/Private/D3DShaderCompilerDXC.cpp
+++ b/Engine/Source/Developer/Windows/ShaderFormatD3D/Private/D3DShaderCompilerDXC.cpp
@@ -616,6 +616,11 @@ bool CompileAndProcessD3DShaderDXC(FString& PreprocessedShaderSource,
 	FString RayTracingExports;
 
 	bool bEnable16BitTypes = false;
+	if (Language == ELanguage::SM6)
+	{
+		// 16bit types are SM6.2 whereas Language == ELanguage::SM6 is SM6.6, so their support at runtime is guarented.
+		bEnable16BitTypes = Input.Environment.CompilerFlags.Contains(CFLAG_AllowRealTypes);
+	}
 
 	if (bIsRayTracingShader)
 	{
