Cross-compile HLSL-to-HLSL after FXC crashed due to internal error.

diff --git a/Engine/Source/Developer/ShaderCompilerCommon/Private/ShaderConductorContext.cpp b/Engine/Source/Developer/ShaderCompilerCommon/Private/ShaderConductorContext.cpp
--- a/Engine/Source/Developer/ShaderCompilerCommon/Private/ShaderConductorContext.cpp
+++ b/Engine/Source/Developer/ShaderCompilerCommon/Private/ShaderConductorContext.cpp
@@ -241,6 +241,29 @@ namespace CrossCompiler
 		}
 	}
 
+	static const ANSICHAR* GetHlslVersionString(int32 Version)
+	{
+		switch (Version)
+		{
+		case 50: return "50";
+		case 60: return "60";
+		case 61: return "61";
+		case 62: return "62";
+		case 63: return "63";
+		case 64: return "64";
+		case 65: return "65";
+		case 66: return "66";
+		default: return nullptr;
+		}
+	}
+
+	static void ConvertScTargetDescLanguageHlsl(const FShaderConductorTarget& InTarget, ShaderConductor::Compiler::TargetDesc& OutTargetDesc)
+	{
+		OutTargetDesc.language = ShaderConductor::ShadingLanguage::Hlsl;
+		OutTargetDesc.version = GetHlslVersionString(InTarget.Version);
+		checkf(OutTargetDesc.version, TEXT("Unsupported target shader version for HLSL: SM%d.%d"), InTarget.Version / 10, InTarget.Version % 10);
+	}
+
 	static const ANSICHAR* GetGlslFamilyVersionString(int32 Version)
 	{
 		switch (Version)
@@ -329,6 +352,9 @@ namespace CrossCompiler
 
 		switch (InTarget.Language)
 		{
+		case EShaderConductorLanguage::Hlsl:
+			ConvertScTargetDescLanguageHlsl(InTarget, OutTargetDesc);
+			break;
 		case EShaderConductorLanguage::Glsl:
 		case EShaderConductorLanguage::Essl:
 			ConvertScTargetDescLanguageGlslFamily(InTarget, OutTargetDesc);
diff --git a/Engine/Source/Developer/ShaderCompilerCommon/Public/ShaderConductorContext.h b/Engine/Source/Developer/ShaderCompilerCommon/Public/ShaderConductorContext.h
--- a/Engine/Source/Developer/ShaderCompilerCommon/Public/ShaderConductorContext.h
+++ b/Engine/Source/Developer/ShaderCompilerCommon/Public/ShaderConductorContext.h
@@ -39,6 +39,7 @@ namespace CrossCompiler
 	/** Target high level languages for ShaderConductor output. */
 	enum class EShaderConductorLanguage
 	{
+		Hlsl,
 		Glsl,
 		Essl,
 		Metal_macOS,
@@ -60,6 +61,7 @@ namespace CrossCompiler
 
 		/**
 		Target shader version.
+		Valid values for HLSL: 50, 60, 61, 62, 63, 64, 65, 66.
 		Valid values for Metal family: 20300, 20200, 20100, 20000, 10200, 10100, 10000.
 		Valid values for GLSL family: 310, 320, 330, 430.
 		*/
diff --git a/Engine/Source/Developer/Windows/ShaderFormatD3D/Private/D3DShaderCompiler.cpp b/Engine/Source/Developer/Windows/ShaderFormatD3D/Private/D3DShaderCompiler.cpp
--- a/Engine/Source/Developer/Windows/ShaderFormatD3D/Private/D3DShaderCompiler.cpp
+++ b/Engine/Source/Developer/Windows/ShaderFormatD3D/Private/D3DShaderCompiler.cpp
@@ -10,6 +10,7 @@
 #include "HAL/FileManager.h"
 #include "Serialization/MemoryWriter.h"
 #include "RayTracingDefinitions.h"
+#include "SpirvCommon.h"
 
 DEFINE_LOG_CATEGORY_STATIC(LogD3D11ShaderCompiler, Log, All);
 
@@ -636,6 +637,71 @@ bool CompileAndProcessD3DShaderFXC(FString& PreprocessedShaderSource, const FStr
 			ExceptionInfo
 		);
 
+		// Some materials give FXC a hard time to optimize and the compiler fails with an internal error.
+		if (Result == HRESULT_FROM_WIN32(ERROR_ARITHMETIC_OVERFLOW))
+		{
+			CrossCompiler::FShaderConductorContext CompilerContext;
+
+			// Load shader source into compiler context
+			CompilerContext.LoadSource(PreprocessedShaderSource, Input.VirtualSourceFilePath, EntryPointName, static_cast<EHlslShaderFrequency>(Input.Target.Frequency));
+
+			// Compile HLSL source to SPIR-V binary
+			CrossCompiler::FShaderConductorOptions Options;
+
+			FSpirv Spirv;
+			if (!CompilerContext.CompileHlslToSpirv(Options, Spirv.Data))
+			{
+				CompilerContext.FlushErrors(Output.Errors);
+				return false;
+			}
+
+			// Cross-compile back to HLSL
+			CrossCompiler::FShaderConductorTarget TargetDesc;
+			TargetDesc.Language = CrossCompiler::EShaderConductorLanguage::Hlsl;
+			TargetDesc.Version = 50;
+			TargetDesc.CompileFlags.SetDefine(TEXT("implicit_resource_binding"), 1);
+			TargetDesc.CompileFlags.SetDefine(TEXT("reconstruct_global_uniforms"), 1);
+			TargetDesc.CompileFlags.SetDefine(TEXT("reconstruct_cbuffer_names"), 1);
+			TargetDesc.CompileFlags.SetDefine(TEXT("reconstruct_semantics"), 1);
+
+			TArray<ANSICHAR> CrossCompiledSource;
+			if (!CompilerContext.CompileSpirvToSourceAnsi(Options, TargetDesc, Spirv.GetByteData(), Spirv.GetByteSize(), CrossCompiledSource))
+			{
+				CompilerContext.FlushErrors(Output.Errors);
+				return false;
+			}
+
+			if (bDumpDebugInfo && CrossCompiledSource.Num() > 1)
+			{
+				DumpDebugShaderText(Input, CrossCompiledSource.GetData(), CrossCompiledSource.Num() - 1, TEXT("intermediate.hlsl"));
+			}
+
+			// Compile again with FXC:
+			// SPIRV-Cross will have generated the new shader with "main" as the new entry point.
+			Result = D3DCompileWrapper(
+				D3DCompileFunc,
+				CrossCompiledSource.GetData(),
+				CrossCompiledSource.Num(),
+				TCHAR_TO_ANSI(*Input.VirtualSourceFilePath),
+				/*pDefines=*/ NULL,
+				/*pInclude=*/ NULL,
+				"main",
+				TCHAR_TO_ANSI(ShaderProfile),
+				CompileFlags,
+				0,
+				Shader.GetInitReference(),
+				Errors.GetInitReference(),
+				bException,
+				ExceptionInfo
+			);
+
+			if (SUCCEEDED(Result))
+			{
+				// Let the user know this shader had to be cross-compiled due to a crash in FXC. Only shows up if CVar 'r.ShaderDevelopmentMode' is enabled.
+				Output.Errors.Add(FShaderCompilerError(FString::Printf(TEXT("Cross-compiled shader to intermediate HLSL after first attempt crashed FXC: %s"), *Input.GenerateShaderName())));
+			}
+		}
+
 		if (bException)
 		{
 			const char* CodeName = Win32SehExceptionToString(ExceptionInfo.Code);
