Fix mapping of EHlslShaderFrequency when cross-compiling HLSL-to-HLSL as appendix to CL 17291947.

#rnx

diff --git a/Engine/Source/Developer/ShaderCompilerCommon/Private/ShaderConductorContext.cpp b/Engine/Source/Developer/ShaderCompilerCommon/Private/ShaderConductorContext.cpp
--- a/Engine/Source/Developer/ShaderCompilerCommon/Private/ShaderConductorContext.cpp
+++ b/Engine/Source/Developer/ShaderCompilerCommon/Private/ShaderConductorContext.cpp
@@ -917,4 +917,23 @@ namespace CrossCompiler
 		return IdentifierTable;
 	}
 
+	EHlslShaderFrequency FShaderConductorContext::MapShaderFrequency(EShaderFrequency Frequency)
+	{
+		const EHlslShaderFrequency FrequencyTable[] =
+		{
+			HSF_VertexShader,		// SF_Vertex
+			HSF_InvalidFrequency,	// SF_Mesh
+			HSF_InvalidFrequency,	// SF_Amplification
+			HSF_PixelShader,		// SF_Pixel
+			HSF_GeometryShader,		// SF_Geometry
+			HSF_ComputeShader,		// SF_Compute
+			HSF_RayGen,				// SF_RayGen
+			HSF_RayMiss,			// SF_RayMiss
+			HSF_RayHitGroup,		// SF_RayHitGroup
+			HSF_RayCallable,		// SF_RayCallable
+		};
+		const uint32 FrequencyIndex = static_cast<uint32>(Frequency);
+		return (FrequencyIndex < UE_ARRAY_COUNT(FrequencyTable) ? FrequencyTable[FrequencyIndex] : HSF_InvalidFrequency);
+	}
+
 } // namespace CrossCompiler
diff --git a/Engine/Source/Developer/ShaderCompilerCommon/Public/ShaderConductorContext.h b/Engine/Source/Developer/ShaderCompilerCommon/Public/ShaderConductorContext.h
--- a/Engine/Source/Developer/ShaderCompilerCommon/Public/ShaderConductorContext.h
+++ b/Engine/Source/Developer/ShaderCompilerCommon/Public/ShaderConductorContext.h
@@ -158,6 +158,9 @@ namespace CrossCompiler
 		/** Disassembles the specified SPIR-V module and returns its assembly as text representation. */
 		static bool Disassemble(EShaderConductorIR Language, const void* Binary, uint32 BinaryByteSize, TArray<ANSICHAR>& OutAssemblyText);
 
+		/** Maps the primary shader frequency (see FShaderTarget::Frequency) to EHlslShaderFrequency used by this context. */
+		static EHlslShaderFrequency MapShaderFrequency(EShaderFrequency Frequency);
+
 	public:
 		struct FShaderConductorIntermediates; // Pimpl idiom
 
diff --git a/Engine/Source/Developer/Windows/ShaderFormatD3D/Private/D3DShaderCompiler.cpp b/Engine/Source/Developer/Windows/ShaderFormatD3D/Private/D3DShaderCompiler.cpp
--- a/Engine/Source/Developer/Windows/ShaderFormatD3D/Private/D3DShaderCompiler.cpp
+++ b/Engine/Source/Developer/Windows/ShaderFormatD3D/Private/D3DShaderCompiler.cpp
@@ -643,7 +643,8 @@ bool CompileAndProcessD3DShaderFXC(FString& PreprocessedShaderSource, const FStr
 			CrossCompiler::FShaderConductorContext CompilerContext;
 
 			// Load shader source into compiler context
-			CompilerContext.LoadSource(PreprocessedShaderSource, Input.VirtualSourceFilePath, EntryPointName, static_cast<EHlslShaderFrequency>(Input.Target.Frequency));
+			const EHlslShaderFrequency Frequency = CrossCompiler::FShaderConductorContext::MapShaderFrequency(static_cast<EShaderFrequency>(Input.Target.Frequency));
+			CompilerContext.LoadSource(PreprocessedShaderSource, Input.VirtualSourceFilePath, EntryPointName, Frequency);
 
 			// Compile HLSL source to SPIR-V binary
 			CrossCompiler::FShaderConductorOptions Options;
@@ -678,11 +679,12 @@ bool CompileAndProcessD3DShaderFXC(FString& PreprocessedShaderSource, const FStr
 
 			// Compile again with FXC:
 			// SPIRV-Cross will have generated the new shader with "main" as the new entry point.
+			const FString CrossCompiledSourceFilename = Input.VirtualSourceFilePath + TEXT(".intermediate.hlsl");
 			Result = D3DCompileWrapper(
 				D3DCompileFunc,
 				CrossCompiledSource.GetData(),
-				CrossCompiledSource.Num(),
-				TCHAR_TO_ANSI(*Input.VirtualSourceFilePath),
+				CrossCompiledSource.Num() - 1,
+				TCHAR_TO_ANSI(*CrossCompiledSourceFilename),
 				/*pDefines=*/ NULL,
 				/*pInclude=*/ NULL,
 				"main",
