Migrates D3D to submit all loose shaders parameters in a single RHISetShaderParameter()

diff --git a/Engine/Shaders/Public/ShaderVersion.ush b/Engine/Shaders/Public/ShaderVersion.ush
--- a/Engine/Shaders/Public/ShaderVersion.ush
+++ b/Engine/Shaders/Public/ShaderVersion.ush
@@ -2,5 +2,4 @@
 // Each time the console command is executed it generates a new GUID. As this file is included
 // in Platform.ush (which should be included in any shader) it allows to invalidate the shader DDC.
 // 
-// If you are merging streams and there is a conflict with this GUID you should make a new GUID rather than taking one or the other.
-// GUID = E5401FDD59AF47489260A452B32036EB
+// GUID = 8BC5CFC6-BF36-415E-AAE3-696C1988A664
\ No newline at end of file
diff --git a/Engine/Source/Developer/ShaderCompilerCommon/Private/ShaderCompilerCommon.cpp b/Engine/Source/Developer/ShaderCompilerCommon/Private/ShaderCompilerCommon.cpp
--- a/Engine/Source/Developer/ShaderCompilerCommon/Private/ShaderCompilerCommon.cpp
+++ b/Engine/Source/Developer/ShaderCompilerCommon/Private/ShaderCompilerCommon.cpp
@@ -1078,14 +1078,22 @@ bool FShaderParameterParser::ParseAndMoveShaderParametersToRootConstantBuffer(
 
 void FShaderParameterParser::ValidateShaderParameterType(
 	const FShaderCompilerInput& CompilerInput,
-	const TCHAR* ShaderBindingName,
-	const FParsedShaderParameter& ParsedParameter,
-	int32 BoundSize,
+	const FString& ShaderBindingName,
+	int32 ReflectionOffset,
+	int32 ReflectionSize,
 	FShaderCompilerOutput& CompilerOutput) const
 {
+	const FShaderParameterParser::FParsedShaderParameter& ParsedParameter = FindParameterInfos(ShaderBindingName);
+
 	check(ParsedParameter.IsFound());
 	check(CompilerInput.RootParametersStructure);
 
+	if (ReflectionSize > 0)
+	{
+		// Verify the offset of the parameter comming from shader reflections matches
+		check(ReflectionOffset == ParsedParameter.ConstantBufferOffset);
+	}
+
 	EShaderPlatform ShaderPlatform = EShaderPlatform(CompilerInput.Target.Platform);
 
 	// Validate the shader type.
@@ -1140,7 +1148,7 @@ void FShaderParameterParser::ValidateShaderParameterType(
 			Error.StrippedErrorMessage = FString::Printf(
 				TEXT("Error: Type %s of shader parameter %s in shader mismatch the shader parameter structure: %s expects a %s"),
 				*ParsedParameter.ParsedType,
-				ShaderBindingName,
+				*ShaderBindingName,
 				*CppCodeName,
 				*ExpectedShaderType);
 			GetParameterFileAndLine(ParsedParameter, Error.ErrorVirtualFilePath, Error.ErrorLineString);
@@ -1153,15 +1161,15 @@ void FShaderParameterParser::ValidateShaderParameterType(
 	}
 
 	// Validate parameter size, in case this is an array.
-	if (BoundSize > int32(ParsedParameter.Member->GetMemberSize()))
+	if (ReflectionSize > int32(ParsedParameter.Member->GetMemberSize()))
 	{
 		FString CppCodeName = CompilerInput.RootParametersStructure->GetFullMemberCodeName(ParsedParameter.ConstantBufferOffset);
 
 		FShaderCompilerError Error;
 		Error.StrippedErrorMessage = FString::Printf(
 			TEXT("Error: The size required to bind shader parameter %s is %i bytes, smaller than %s that is %i bytes in the parameter structure."),
-			ShaderBindingName,
-			BoundSize,
+			*ShaderBindingName,
+			ReflectionSize,
 			*CppCodeName,
 			ParsedParameter.Member->GetMemberSize());
 		GetParameterFileAndLine(ParsedParameter, Error.ErrorVirtualFilePath, Error.ErrorLineString);
@@ -1218,13 +1226,15 @@ void FShaderParameterParser::ValidateShaderParameterTypes(
 			return;
 		}
 
+		int32 BoundOffset = 0;
 		int32 BoundSize = 0;
 		if (const FParameterAllocation* ParameterAllocation = ParametersFoundByCompiler.Find(ShaderBindingName))
 		{
+			BoundOffset = ParameterAllocation->BaseIndex;
 			BoundSize = ParameterAllocation->Size;
 		}
 
-		ValidateShaderParameterType(CompilerInput, ShaderBindingName, ParsedParameter, BoundSize, CompilerOutput);
+		ValidateShaderParameterType(CompilerInput, ShaderBindingName, BoundOffset, BoundSize, CompilerOutput);
 	});
 }
 
diff --git a/Engine/Source/Developer/ShaderCompilerCommon/Public/ShaderCompilerCommon.h b/Engine/Source/Developer/ShaderCompilerCommon/Public/ShaderCompilerCommon.h
--- a/Engine/Source/Developer/ShaderCompilerCommon/Public/ShaderCompilerCommon.h
+++ b/Engine/Source/Developer/ShaderCompilerCommon/Public/ShaderCompilerCommon.h
@@ -95,9 +95,9 @@ public:
 	/** Validates the shader parameter in code is compatible with the shader parameter structure. */
 	void ValidateShaderParameterType(
 		const FShaderCompilerInput& CompilerInput,
-		const TCHAR* ShaderBindingName,
-		const FParsedShaderParameter& ParsedParameter,
-		int32 BoundSize,
+		const FString& ShaderBindingName,
+		int32 ReflectionOffset,
+		int32 ReflectionSize,
 		FShaderCompilerOutput& CompilerOutput) const;
 
 	/** Validates shader parameter map is compatible with the shader parameter structure. */
diff --git a/Engine/Source/Developer/Windows/ShaderFormatD3D/Private/D3DShaderCompiler.cpp b/Engine/Source/Developer/Windows/ShaderFormatD3D/Private/D3DShaderCompiler.cpp
--- a/Engine/Source/Developer/Windows/ShaderFormatD3D/Private/D3DShaderCompiler.cpp
+++ b/Engine/Source/Developer/Windows/ShaderFormatD3D/Private/D3DShaderCompiler.cpp
@@ -573,7 +573,10 @@ bool DumpDebugShaderUSF(FString& PreprocessedShaderSource, const FShaderCompiler
 
 // Generate the dumped usf file; call the D3D compiler, gather reflection information and generate the output data
 bool CompileAndProcessD3DShaderFXC(FString& PreprocessedShaderSource, const FString& CompilerPath,
-	uint32 CompileFlags, const FShaderCompilerInput& Input, FString& EntryPointName,
+	uint32 CompileFlags,
+	const FShaderCompilerInput& Input,
+	const FShaderParameterParser& ShaderParameterParser,
+	FString& EntryPointName,
 	const TCHAR* ShaderProfile, bool bSecondPassAferUnusedInputRemoval,
 	TArray<FString>& FilteredErrors, FShaderCompilerOutput& Output)
 {
@@ -921,7 +924,7 @@ bool CompileAndProcessD3DShaderFXC(FString& PreprocessedShaderSource, const FStr
 						if (RemoveUnusedInputs(PreprocessedShaderSource, ShaderInputs, EntryPointName, RemoveErrors))
 						{
 							Output = OriginalOutput;
-							if (!CompileAndProcessD3DShaderFXC(PreprocessedShaderSource, CompilerPath, CompileFlags, Input, EntryPointName, ShaderProfile, true, FilteredErrors, Output))
+							if (!CompileAndProcessD3DShaderFXC(PreprocessedShaderSource, CompilerPath, CompileFlags, Input, ShaderParameterParser, EntryPointName, ShaderProfile, true, FilteredErrors, Output))
 							{
 								// if we failed to compile the shader, propagate the error up
 								return false;
@@ -988,7 +991,9 @@ bool CompileAndProcessD3DShaderFXC(FString& PreprocessedShaderSource, const FStr
 			ExtractParameterMapFromD3DShader<
 				ID3D11ShaderReflection, D3D11_SHADER_DESC, D3D11_SHADER_INPUT_BIND_DESC,
 				ID3D11ShaderReflectionConstantBuffer, D3D11_SHADER_BUFFER_DESC,
-				ID3D11ShaderReflectionVariable, D3D11_SHADER_VARIABLE_DESC>(Input.Target.Platform, BindingSpace, Input.VirtualSourceFilePath, Reflector, ShaderDesc,
+				ID3D11ShaderReflectionVariable, D3D11_SHADER_VARIABLE_DESC>(
+					Input, ShaderParameterParser,
+					BindingSpace, Reflector, ShaderDesc,
 					bGlobalUniformBufferUsed, bDiagnosticBufferUsed,
 					NumSamplers, NumSRVs, NumCBs, NumUAVs,
 					Output, UniformBufferNames, UsedUniformBufferSlots, VendorExtensions);
@@ -1198,7 +1203,7 @@ void CompileD3DShader(const FShaderCompilerInput& Input, FShaderCompilerOutput&
 	FShaderParameterParser ShaderParameterParser;
 	if (!ShaderParameterParser.ParseAndMoveShaderParametersToRootConstantBuffer(
 		Input, Output, PreprocessedShaderSource,
-		Input.IsRayTracingShader() ? TEXT("cbuffer") : nullptr))
+		(Input.IsRayTracingShader() || ShouldUseStableConstantBuffer(Input)) ? TEXT("cbuffer") : nullptr))
 	{
 		// The FShaderParameterParser will add any relevant errors.
 		return;
@@ -1241,7 +1246,7 @@ void CompileD3DShader(const FShaderCompilerInput& Input, FShaderCompilerOutput&
 	TArray<FString> FilteredErrors;
 	if (bUseDXC)
 	{
-		if (!CompileAndProcessD3DShaderDXC(PreprocessedShaderSource, CompileFlags, Input, EntryPointName, ShaderProfile, Language, false, FilteredErrors, Output))
+		if (!CompileAndProcessD3DShaderDXC(PreprocessedShaderSource, CompileFlags, Input, ShaderParameterParser, EntryPointName, ShaderProfile, Language, false, FilteredErrors, Output))
 		{
 			if (!FilteredErrors.Num())
 			{
@@ -1256,7 +1261,7 @@ void CompileD3DShader(const FShaderCompilerInput& Input, FShaderCompilerOutput&
 		FString CompilerPath = FPaths::EngineDir();
 		CompilerPath.Append(TEXT("Binaries/ThirdParty/Windows/DirectX/x64/d3dcompiler_47.dll"));
 
-		if (!CompileAndProcessD3DShaderFXC(PreprocessedShaderSource, CompilerPath, CompileFlags, Input, EntryPointName, ShaderProfile, false, FilteredErrors, Output))
+		if (!CompileAndProcessD3DShaderFXC(PreprocessedShaderSource, CompilerPath, CompileFlags, Input, ShaderParameterParser, EntryPointName, ShaderProfile, false, FilteredErrors, Output))
 		{
 			if (!FilteredErrors.Num())
 			{
diff --git a/Engine/Source/Developer/Windows/ShaderFormatD3D/Private/D3DShaderCompiler.inl b/Engine/Source/Developer/Windows/ShaderFormatD3D/Private/D3DShaderCompiler.inl
--- a/Engine/Source/Developer/Windows/ShaderFormatD3D/Private/D3DShaderCompiler.inl
+++ b/Engine/Source/Developer/Windows/ShaderFormatD3D/Private/D3DShaderCompiler.inl
@@ -247,8 +247,11 @@ public:
 template <typename ID3D1xShaderReflection, typename D3D1x_SHADER_DESC, typename D3D1x_SHADER_INPUT_BIND_DESC,
 	typename ID3D1xShaderReflectionConstantBuffer, typename D3D1x_SHADER_BUFFER_DESC,
 	typename ID3D1xShaderReflectionVariable, typename D3D1x_SHADER_VARIABLE_DESC>
-inline void ExtractParameterMapFromD3DShader(
-		uint32 TargetPlatform, uint32 BindingSpace, const FString& VirtualSourceFilePath, ID3D1xShaderReflection* Reflector, const D3D1x_SHADER_DESC& ShaderDesc,
+	inline void ExtractParameterMapFromD3DShader(
+		const FShaderCompilerInput& Input,
+		const FShaderParameterParser& ShaderParameterParser,
+		uint32 BindingSpace,
+		ID3D1xShaderReflection* Reflector, const D3D1x_SHADER_DESC& ShaderDesc,
 		bool& bGlobalUniformBufferUsed, bool& bDiagnosticBufferUsed, uint32& NumSamplers, uint32& NumSRVs, uint32& NumCBs, uint32& NumUAVs,
 		FShaderCompilerOutput& Output, TArray<FString>& UniformBufferNames, TBitArray<>& UsedUniformBufferSlots, TArray<FShaderCodeVendorExtension>& VendorExtensions)
 {
@@ -270,29 +273,100 @@ inline void ExtractParameterMapFromD3DShader(
 			D3D1x_SHADER_BUFFER_DESC CBDesc;
 			ConstantBuffer->GetDesc(&CBDesc);
 			bool bGlobalCB = (FCStringAnsi::Strcmp(CBDesc.Name, "$Globals") == 0);
+			const bool bIsRootCB = FCString::Strcmp(ANSI_TO_TCHAR(CBDesc.Name), FShaderParametersMetadata::kRootUniformBufferBindingName) == 0;
 
 			if (bGlobalCB)
 			{
-				// Track all of the variables in this constant buffer.
-				for (uint32 ConstantIndex = 0; ConstantIndex < CBDesc.Variables; ConstantIndex++)
+				if (ShouldUseStableConstantBuffer(Input))
 				{
-					ID3D1xShaderReflectionVariable* Variable = ConstantBuffer->GetVariableByIndex(ConstantIndex);
-					D3D1x_SHADER_VARIABLE_DESC VariableDesc;
-					Variable->GetDesc(&VariableDesc);
-					if (VariableDesc.uFlags & D3D_SVF_USED)
+					// Each member found in the global constant buffer means it was not in RootParametersStructure or
+					// it would have been moved by ShaderParameterParser.ParseAndMoveShaderParametersToRootConstantBuffer().
+					for (uint32 ConstantIndex = 0; ConstantIndex < CBDesc.Variables; ConstantIndex++)
 					{
-						bGlobalUniformBufferUsed = true;
+						ID3D1xShaderReflectionVariable* Variable = ConstantBuffer->GetVariableByIndex(ConstantIndex);
+						D3D1x_SHADER_VARIABLE_DESC VariableDesc;
+						Variable->GetDesc(&VariableDesc);
+						if (VariableDesc.uFlags & D3D_SVF_USED)
+						{
+							AddUnboundShaderParameterError(
+								Input,
+								ShaderParameterParser,
+								ANSI_TO_TCHAR(VariableDesc.Name),
+								Output);
+						}
+					}
+				}
+				else
+				{
+					// Track all of the variables in this constant buffer.
+					for (uint32 ConstantIndex = 0; ConstantIndex < CBDesc.Variables; ConstantIndex++)
+					{
+						ID3D1xShaderReflectionVariable* Variable = ConstantBuffer->GetVariableByIndex(ConstantIndex);
+						D3D1x_SHADER_VARIABLE_DESC VariableDesc;
+						Variable->GetDesc(&VariableDesc);
+						if (VariableDesc.uFlags & D3D_SVF_USED)
+						{
+							bGlobalUniformBufferUsed = true;
+
+							Output.ParameterMap.AddParameterAllocation(
+								ANSI_TO_TCHAR(VariableDesc.Name),
+								CBIndex,
+								VariableDesc.StartOffset,
+								VariableDesc.Size,
+								EShaderParameterType::LooseData
+							);
+							UsedUniformBufferSlots[CBIndex] = true;
+						}
+					}
+				}
+			}
+			else if (bIsRootCB && ShouldUseStableConstantBuffer(Input))
+			{
+				if (CBIndex == FShaderParametersMetadata::kRootCBufferBindingIndex)
+				{
+					int32 ConstantBufferSize = 0;
 
+					// Track all of the variables in this constant buffer.
+					for (uint32 ConstantIndex = 0; ConstantIndex < CBDesc.Variables; ConstantIndex++)
+					{
+						ID3D1xShaderReflectionVariable* Variable = ConstantBuffer->GetVariableByIndex(ConstantIndex);
+						D3D1x_SHADER_VARIABLE_DESC VariableDesc;
+						Variable->GetDesc(&VariableDesc);
+						if (VariableDesc.uFlags & D3D_SVF_USED)
+						{
+							FString MemberName(ANSI_TO_TCHAR(VariableDesc.Name));
+							int32 ReflectionSize = VariableDesc.Size;
+							int32 ReflectionOffset = VariableDesc.StartOffset;
+
+							ShaderParameterParser.ValidateShaderParameterType(Input, MemberName, ReflectionOffset, ReflectionSize, Output);
+
+							ConstantBufferSize = FMath::Max(ConstantBufferSize, ReflectionOffset + ReflectionSize);
+						}
+					}
+
+					if (ConstantBufferSize > 0)
+					{
 						Output.ParameterMap.AddParameterAllocation(
-							ANSI_TO_TCHAR(VariableDesc.Name),
-							CBIndex,
-							VariableDesc.StartOffset,
-							VariableDesc.Size,
-							EShaderParameterType::LooseData
-						);
+							FShaderParametersMetadata::kRootUniformBufferBindingName,
+							FShaderParametersMetadata::kRootCBufferBindingIndex,
+							/* Offset = */ uint16(0),
+							/* Size = */ uint16(ConstantBufferSize),
+							EShaderParameterType::LooseData);
+
+						bGlobalUniformBufferUsed = true;
 						UsedUniformBufferSlots[CBIndex] = true;
 					}
 				}
+				else
+				{
+					FString ErrorMessage = FString::Printf(
+						TEXT("Error: %s is expected to always be in the API slot %d, but is actually in slot %d."),
+						FShaderParametersMetadata::kRootUniformBufferBindingName,
+						FShaderParametersMetadata::kRootCBufferBindingIndex,
+						CBIndex);
+					Output.Errors.Add(FShaderCompilerError(*ErrorMessage));
+					Output.bSucceeded = false;
+				}
 			}
 			else
 			{
diff --git a/Engine/Source/Developer/Windows/ShaderFormatD3D/Private/D3DShaderCompilerDXC.cpp b/Engine/Source/Developer/Windows/ShaderFormatD3D/Private/D3DShaderCompilerDXC.cpp
--- a/Engine/Source/Developer/Windows/ShaderFormatD3D/Private/D3DShaderCompilerDXC.cpp
+++ b/Engine/Source/Developer/Windows/ShaderFormatD3D/Private/D3DShaderCompilerDXC.cpp
@@ -600,7 +600,10 @@ static bool RewriteUsingSC(FString& PreprocessedShaderSource, const FShaderCompi
 
 // Generate the dumped usf file; call the D3D compiler, gather reflection information and generate the output data
 bool CompileAndProcessD3DShaderDXC(FString& PreprocessedShaderSource,
-	const uint32 CompileFlags, const FShaderCompilerInput& Input, FString& EntryPointName,
+	const uint32 CompileFlags,
+	const FShaderCompilerInput& Input,
+	const FShaderParameterParser& ShaderParameterParser,
+	FString& EntryPointName,
 	const TCHAR* ShaderProfile, ELanguage Language, bool bProcessingSecondTime,
 	TArray<FString>& FilteredErrors, FShaderCompilerOutput& Output)
 {
@@ -793,7 +796,8 @@ bool CompileAndProcessD3DShaderDXC(FString& PreprocessedShaderSource,
 						ExtractParameterMapFromD3DShader<ID3D12FunctionReflection, D3D12_FUNCTION_DESC, D3D12_SHADER_INPUT_BIND_DESC,
 							ID3D12ShaderReflectionConstantBuffer, D3D12_SHADER_BUFFER_DESC,
 							ID3D12ShaderReflectionVariable, D3D12_SHADER_VARIABLE_DESC>(
-								Input.Target.Platform, AutoBindingSpace, Input.VirtualSourceFilePath, FunctionReflection, FunctionDesc, 
+								Input, ShaderParameterParser,
+								AutoBindingSpace, FunctionReflection, FunctionDesc, 
 								bGlobalUniformBufferUsed, bDiagnosticBufferUsed,
 								NumSamplers, NumSRVs, NumCBs, NumUAVs,
 								Output, UniformBufferNames, UsedUniformBufferSlots, VendorExtensions);
@@ -863,9 +867,10 @@ bool CompileAndProcessD3DShaderDXC(FString& PreprocessedShaderSource,
 			ExtractParameterMapFromD3DShader<ID3D12ShaderReflection, D3D12_SHADER_DESC, D3D12_SHADER_INPUT_BIND_DESC,
 				ID3D12ShaderReflectionConstantBuffer, D3D12_SHADER_BUFFER_DESC,
 				ID3D12ShaderReflectionVariable, D3D12_SHADER_VARIABLE_DESC>(
-					Input.Target.Platform, AutoBindingSpace, Input.VirtualSourceFilePath, ShaderReflection, ShaderDesc,
-						bGlobalUniformBufferUsed, bDiagnosticBufferUsed, 
-						NumSamplers, NumSRVs, NumCBs, NumUAVs,
+					Input, ShaderParameterParser,
+					AutoBindingSpace, ShaderReflection, ShaderDesc,
+					bGlobalUniformBufferUsed, bDiagnosticBufferUsed, 
+					NumSamplers, NumSRVs, NumCBs, NumUAVs,
 					Output, UniformBufferNames, UsedUniformBufferSlots, VendorExtensions);
 
 
diff --git a/Engine/Source/Developer/Windows/ShaderFormatD3D/Private/ShaderFormatD3D.h b/Engine/Source/Developer/Windows/ShaderFormatD3D/Private/ShaderFormatD3D.h
--- a/Engine/Source/Developer/Windows/ShaderFormatD3D/Private/ShaderFormatD3D.h
+++ b/Engine/Source/Developer/Windows/ShaderFormatD3D/Private/ShaderFormatD3D.h
@@ -3,6 +3,7 @@
 #pragma once
 
 #include "Windows/WindowsHWrapper.h"
+#include "ShaderCompilerCommon.h"
 
 // TODO: Lock D3D12 to SM 6.6 min spec
 #define USE_SHADER_MODEL_6_6 0
@@ -22,12 +23,18 @@ void CompileShader_Windows(const struct FShaderCompilerInput& Input,struct FShad
  * @param bSecondPassAferUnusedInputRemoval whether we're compiling the shader second time, after having removed the unused inputs discovered in the first pass
  */
 bool CompileAndProcessD3DShaderFXC(FString& PreprocessedShaderSource, const FString& CompilerPath,
-	uint32 CompileFlags, const FShaderCompilerInput& Input, FString& EntryPointName,
+	uint32 CompileFlags,
+	const FShaderCompilerInput& Input,
+	const FShaderParameterParser& ShaderParameterParser,
+	FString& EntryPointName,
 	const TCHAR* ShaderProfile, bool bSecondPassAferUnusedInputRemoval,
 	TArray<FString>& FilteredErrors, FShaderCompilerOutput& Output);
 
 bool CompileAndProcessD3DShaderDXC(FString& PreprocessedShaderSource,
-	uint32 CompileFlags, const FShaderCompilerInput& Input, FString& EntryPointName,
+	uint32 CompileFlags,
+	const FShaderCompilerInput& Input,
+	const FShaderParameterParser& ShaderParameterParser,
+	FString& EntryPointName,
 	const TCHAR* ShaderProfile, ELanguage Language, bool bProcessingSecondTime,
 	TArray<FString>& FilteredErrors, FShaderCompilerOutput& Output);
 
diff --git a/Engine/Source/Runtime/RenderCore/Public/ShaderCompilerCore.h b/Engine/Source/Runtime/RenderCore/Public/ShaderCompilerCore.h
--- a/Engine/Source/Runtime/RenderCore/Public/ShaderCompilerCore.h
+++ b/Engine/Source/Runtime/RenderCore/Public/ShaderCompilerCore.h
@@ -484,6 +484,20 @@ enum class ESCWErrorCode
 	CrashInsidePlatformCompiler,
 };
 
+
+inline bool ShouldUseStableConstantBuffer(const FShaderCompilerInput& Input)
+{
+	// stable constant buffer is for the FShaderParameterBindings::BindForLegacyShaderParameters() code path.
+	// Ray tracing shaders use FShaderParameterBindings::BindForRootShaderParameters instead.
+	if (Input.IsRayTracingShader())
+	{
+		return false;
+	}
+
+	return Input.RootParametersStructure != nullptr;
+}
+
+
 /**
  * Validates the format of a virtual shader file path.
  * Meant to be use as such: check(CheckVirtualShaderFilePath(VirtualFilePath));
