Added a missing VulkanShaderFormat to a Windows target platform

diff --git a/Engine/Source/Developer/Windows/WindowsTargetPlatform/Private/GenericWindowsTargetPlatform.h b/Engine/Source/Developer/Windows/WindowsTargetPlatform/Private/GenericWindowsTargetPlatform.h
--- a/Engine/Source/Developer/Windows/WindowsTargetPlatform/Private/GenericWindowsTargetPlatform.h
+++ b/Engine/Source/Developer/Windows/WindowsTargetPlatform/Private/GenericWindowsTargetPlatform.h
@@ -256,6 +256,7 @@ public:
 	{
 		OutModuleNames.Add(TEXT("ShaderFormatD3D"));
 		OutModuleNames.Add(TEXT("ShaderFormatOpenGL"));
+		OutModuleNames.Add(TEXT("VulkanShaderFormat"));
 	}
 
 	virtual const class FStaticMeshLODSettings& GetStaticMeshLODSettings( ) const override
