Modified warning text to include notes that crashes are likely when data changes are made without re-instancing enabled.

#rnx

diff --git a/Engine/Source/Developer/Windows/LiveCoding/Private/LiveCodingModule.cpp b/Engine/Source/Developer/Windows/LiveCoding/Private/LiveCodingModule.cpp
--- a/Engine/Source/Developer/Windows/LiveCoding/Private/LiveCodingModule.cpp
+++ b/Engine/Source/Developer/Windows/LiveCoding/Private/LiveCodingModule.cpp
@@ -494,11 +494,18 @@ void FLiveCodingModule::AttemptSyncLivePatching()
 			case commands::PostCompileResult::Success:
 				if (bHasReinstancingOccurred)
 				{
+					if (!IsReinstancingEnabled())
+					{
+						UE_LOG(LogLiveCoding, Warning, TEXT("%s, %s"), *Success, TEXT("data type changes with re-instancing disabled is not supported and will likely lead to a crash"));
+					}
+					else
+					{
 #if WITH_EDITOR
-					UE_LOG(LogLiveCoding, Warning, TEXT("%s, %s"), *Success, TEXT("data type changes may cause packaging to fail if assets reference the new or updated data types"));
+						UE_LOG(LogLiveCoding, Warning, TEXT("%s, %s"), *Success, TEXT("data type changes may cause packaging to fail if assets reference the new or updated data types"));
 #else
-					UE_LOG(LogLiveCoding, Warning, TEXT("%s, %s"), *Success, TEXT("data type changes may cause unexpected failures"));
+						UE_LOG(LogLiveCoding, Warning, TEXT("%s, %s"), *Success, TEXT("data type changes may cause unexpected failures"));
 #endif
+					}
 				}
 				else
 				{
@@ -525,13 +532,21 @@ void FLiveCodingModule::AttemptSyncLivePatching()
 			static const FText FailureDetailText = LOCTEXT("FailureDetail", "Please see Live Coding console for more information.");
 			static const FText CancelledText = LOCTEXT("Cancelled", "Live coding cancelled");
 			static const FText ReinstancingText = LOCTEXT("Reinstancing", "Data type changes may cause packaging to fail if assets reference the new or updated data types.");
+			static const FText DisabledText = LOCTEXT("ReinstancingDisabled", "Data type changes with re-instancing disabled is not supported and will likely lead to a crash.");
 
 			switch (GPostCompileResult)
 			{
 			case commands::PostCompileResult::Success:
 				if (bHasReinstancingOccurred)
 				{
-					ShowNotification(true, SuccessText, &ReinstancingText);
+					if (!IsReinstancingEnabled())
+					{
+						ShowNotification(true, SuccessText, &DisabledText);
+					}
+					else
+					{
+						ShowNotification(true, SuccessText, &ReinstancingText);
+					}
 				}
 				else
 				{
