Fixed issue where LiveCoding fails to reinstance objects when their structure has changed.

#rnx

diff --git a/Engine/Source/Developer/Windows/LiveCodingServer/Private/External/LC_SymbolPatterns.cpp b/Engine/Source/Developer/Windows/LiveCodingServer/Private/External/LC_SymbolPatterns.cpp
--- a/Engine/Source/Developer/Windows/LiveCodingServer/Private/External/LC_SymbolPatterns.cpp
+++ b/Engine/Source/Developer/Windows/LiveCodingServer/Private/External/LC_SymbolPatterns.cpp
@@ -222,7 +222,7 @@ namespace symbolPatterns
 	};
 
 	// BEGIN EPIC MOD
-	extern const char* const UE_STATICS_BLOCK_PATTERNS[6] =
+	extern const char* const UE_STATICS_BLOCK_PATTERNS[7] =
 	{
 		"*Z_Construct_UClass_*_Statics*",
 		"*Z_Construct_UEnum_*_Statics*",
@@ -230,15 +230,17 @@ namespace symbolPatterns
 		"*Z_Construct_UDelegateFunction_*_Statics*",
 		"*Z_Construct_UScriptStruct_*_Statics*",
 		"*Z_Construct_UPackage_*_Statics*",
+		"*Z_CompiledInDeferFile_*_Statics*",
 	};
 
-	extern const char* const UE_REGISTER_PATTERNS[5] =
+	extern const char* const UE_REGISTER_PATTERNS[6] =
 	{
 		"?AutoInitialize",
 		"?Z_CompiledInDeferEnum_UEnum_",
 		"?Z_CompiledInDeferStruct_UScriptStruct_",
 		"?Z_CompiledInDeferCppStructOps_UScriptStruct_",
 		"?Z_CompiledInDeferPackage_UPackage_",
+		"?Z_CompiledInDeferFile_",
 	};
 	// END EPIC MOD
 }
diff --git a/Engine/Source/Developer/Windows/LiveCodingServer/Private/External/LC_SymbolPatterns.h b/Engine/Source/Developer/Windows/LiveCodingServer/Private/External/LC_SymbolPatterns.h
--- a/Engine/Source/Developer/Windows/LiveCodingServer/Private/External/LC_SymbolPatterns.h
+++ b/Engine/Source/Developer/Windows/LiveCodingServer/Private/External/LC_SymbolPatterns.h
@@ -43,7 +43,7 @@ namespace symbolPatterns
 	extern const char* const ANONYMOUS_NAMESPACE_PATTERN;
 
 	// BEGIN EPIC MOD
-	extern const char* const UE_STATICS_BLOCK_PATTERNS[6];
-	extern const char* const UE_REGISTER_PATTERNS[5];
+	extern const char* const UE_STATICS_BLOCK_PATTERNS[7];
+	extern const char* const UE_REGISTER_PATTERNS[6];
 	// END EPIC MOD
 }
