Attempt to fix "Cannot find image section .voltbl" LiveCoding error spam

#rnx

diff --git a/Engine/Source/Developer/Windows/LiveCodingServer/Private/External/LC_SymbolReconstruction.cpp b/Engine/Source/Developer/Windows/LiveCodingServer/Private/External/LC_SymbolReconstruction.cpp
--- a/Engine/Source/Developer/Windows/LiveCodingServer/Private/External/LC_SymbolReconstruction.cpp
+++ b/Engine/Source/Developer/Windows/LiveCodingServer/Private/External/LC_SymbolReconstruction.cpp
@@ -753,7 +753,33 @@ walkOpenSymbols:
 			const symbols::ImageSection* imageSection = symbols::FindImageSectionByName(imageSectionDb, coffSection.name);
 			if (!imageSection)
 			{
-				LC_ERROR_DEV("Cannot find image section %s", coffSection.name.c_str());
+				// BEGIN EPIC MOD
+				if (diaNameIt != diaNameToRva.end())
+				{
+					// fast path.
+					// there is a symbol that matches the exact name of the symbol in the .obj file
+					const uint32_t rva = diaNameIt->second.rva;
+
+					LC_LOG_DEV("Fast path, found symbol %s at 0x%X", missingSymbolName.c_str(), rva);
+
+					symbols::CreateNewSymbol(missingSymbolName, rva, symbolDB);
+
+					openSymbols.push_back(symbol);
+
+					--unknownSymbolsToFind;
+
+					// did we already find all symbols?
+					if (unknownSymbolsToFind == 0u)
+					{
+						LC_LOG_DEV("All symbols known, exiting");
+						return;
+					}
+				}
+				else
+				{
+					LC_ERROR_DEV("Cannot find image section %s", coffSection.name.c_str());
+				}
+				// END EPIC MOD
 				continue;
 			}
 
