LiveCodingConsole no longer generates extraneous error messages regarding calls to NtQuerySystemInformation.

diff --git a/Engine/Source/Developer/Windows/LiveCoding/Private/External/LC_Process.cpp b/Engine/Source/Developer/Windows/LiveCoding/Private/External/LC_Process.cpp
--- a/Engine/Source/Developer/Windows/LiveCoding/Private/External/LC_Process.cpp
+++ b/Engine/Source/Developer/Windows/LiveCoding/Private/External/LC_Process.cpp
@@ -583,7 +583,9 @@ std::vector<Thread::Id> Process::EnumerateThreads(Id processId)
 		processSnapshot = ::malloc(bufferSize);
 
 		// try getting a process snapshot into the provided buffer
-		status = WindowsInternals::NtQuerySystemInformation(WindowsInternals::SystemProcessInformation, processSnapshot, bufferSize, NULL);
+		// BEGIN EPIC MOD
+		status = WindowsInternals::NtQuerySystemInformation.ExecNoResultCheck(WindowsInternals::SystemProcessInformation, processSnapshot, bufferSize, NULL);
+		// END EPIC MOD
 
 		if (status == STATUS_INFO_LENGTH_MISMATCH)
 		{
@@ -597,6 +599,7 @@ std::vector<Thread::Id> Process::EnumerateThreads(Id processId)
 		{
 			// something went wrong
 			// BEGIN EPIC MOD - PVS is having problems dealing with + here. 
+			WindowsInternals::NtQuerySystemInformation.CheckResult(status); // write the error
 			LC_ERROR_USER("Cannot enumerate threads in process (PID: %d)", static_cast<DWORD>(processId));
 			// END EPIC MOD
 			::free(processSnapshot);
diff --git a/Engine/Source/Developer/Windows/LiveCoding/Private/External/LC_WindowsInternalFunctions.h b/Engine/Source/Developer/Windows/LiveCoding/Private/External/LC_WindowsInternalFunctions.h
--- a/Engine/Source/Developer/Windows/LiveCoding/Private/External/LC_WindowsInternalFunctions.h
+++ b/Engine/Source/Developer/Windows/LiveCoding/Private/External/LC_WindowsInternalFunctions.h
@@ -25,7 +25,10 @@ namespace WindowsInternals
 
 		inline R operator()(Args... args) const;
 
-	private:
+		// BEGIN EPIC MODS
+		inline R ExecNoResultCheck(Args... args) const;
+		// END EPIC MODS
+
 		// Helper for letting us check the result for arbitrary return types.
 		// Base template.
 		template <typename T>
@@ -41,6 +44,10 @@ namespace WindowsInternals
 			}
 		}
 
+		// BEGIN EPIC MODS
+	private:
+		// END EPIC MODS
+
 		const char* m_moduleName;
 		const char* m_functionName;
 		PtrToFunction m_function;
@@ -78,6 +85,15 @@ inline R WindowsInternals::Function<R (Args...)>::operator()(Args... args) const
 	return result;
 }
 
+// BEGIN EPIC MODS
+template <typename R, typename... Args>
+inline R WindowsInternals::Function<R(Args...)>::ExecNoResultCheck(Args... args) const
+{
+	const R result = m_function(args...);
+
+	return result;
+}
+// BEGIN EPIC MODS
 
 // These are undocumented functions found in ntdll.dll.
 // We don't call them directly, but use them for "extracting" their signature using decltype.
