LiveCoding feedback to user when the console can't be started has been improved.

diff --git a/Engine/Source/Developer/Windows/LiveCoding/Private/LiveCodingModule.cpp b/Engine/Source/Developer/Windows/LiveCoding/Private/LiveCodingModule.cpp
--- a/Engine/Source/Developer/Windows/LiveCoding/Private/LiveCodingModule.cpp
+++ b/Engine/Source/Developer/Windows/LiveCoding/Private/LiveCodingModule.cpp
@@ -269,6 +269,7 @@ void FLiveCodingModule::EnableForSession(bool bEnable)
 {
 	if (bEnable)
 	{
+		EnableErrorText = FText::GetEmpty();
 		if(!bStarted)
 		{
 			StartLiveCoding();
@@ -297,6 +298,11 @@ bool FLiveCodingModule::IsEnabledForSession() const
 	return bEnabledForSession;
 }
 
+FText FLiveCodingModule::GetEnableErrorText() const
+{
+	return EnableErrorText;
+}
+
 bool FLiveCodingModule::CanEnableForSession() const
 {
 #if !IS_MONOLITHIC
@@ -597,11 +603,13 @@ ILiveCodingModule::FOnPatchCompleteDelegate& FLiveCodingModule::GetOnPatchComple
 
 bool FLiveCodingModule::StartLiveCoding()
 {
+	EnableErrorText = FText::GetEmpty();
 	if(!bStarted)
 	{
 		// Make sure there aren't any hot reload modules already active
 		if (!CanEnableForSession())
 		{
+			EnableErrorText = LOCTEXT("NoLiveCodingCompileAfterHotReload", "Live Coding cannot be enabled while hot-reloaded modules are active. Please close the editor and build from your IDE before restarting.");
 			UE_LOG(LogLiveCoding, Error, TEXT("Unable to start live coding session. Some modules have already been hot reloaded."));
 			return false;
 		}
@@ -610,6 +618,10 @@ bool FLiveCodingModule::StartLiveCoding()
 		GLiveCodingConsolePath = ConsolePathVariable->GetString();
 		if (!FPaths::FileExists(GLiveCodingConsolePath))
 		{
+			FFormatNamedArguments Args;
+			Args.Add(TEXT("Executable"), FText::FromString(GLiveCodingConsolePath));
+			const static FText FormatString = LOCTEXT("LiveCodingMissingExecutable", "Unable to start live coding session. Missing executable '{Executable}'. Use the LiveCoding.ConsolePath console variable to modify.");
+			EnableErrorText = FText::Format(FormatString, Args);
 			UE_LOG(LogLiveCoding, Error, TEXT("Unable to start live coding session. Missing executable '%s'. Use the LiveCoding.ConsolePath console variable to modify."), *GLiveCodingConsolePath);
 			return false;
 		}
@@ -618,6 +630,10 @@ bool FLiveCodingModule::StartLiveCoding()
 		FString SourceProject = SourceProjectVariable->GetString();
 		if (SourceProject.Len() > 0 && !FPaths::FileExists(SourceProject))
 		{
+			FFormatNamedArguments Args;
+			Args.Add(TEXT("ProjectFile"), FText::FromString(SourceProject));
+			const static FText FormatString = LOCTEXT("LiveCodingMissingProjectFile", "Unable to start live coding session. Unable to find source project file '{ProjectFile}'.");
+			EnableErrorText = FText::Format(FormatString, Args);
 			UE_LOG(LogLiveCoding, Error, TEXT("Unable to start live coding session. Unable to find source project file '%s'."), *SourceProject);
 			return false;
 		}
diff --git a/Engine/Source/Developer/Windows/LiveCoding/Private/LiveCodingModule.h b/Engine/Source/Developer/Windows/LiveCoding/Private/LiveCodingModule.h
--- a/Engine/Source/Developer/Windows/LiveCoding/Private/LiveCodingModule.h
+++ b/Engine/Source/Developer/Windows/LiveCoding/Private/LiveCodingModule.h
@@ -33,6 +33,7 @@ public:
 	virtual bool IsEnabledByDefault() const override;
 	virtual void EnableForSession(bool bInEnabled) override;
 	virtual bool IsEnabledForSession() const override;
+	virtual FText GetEnableErrorText() const override;
 	virtual bool AutomaticallyCompileNewClasses() const override;
 	virtual bool CanEnableForSession() const override;
 	virtual bool HasStarted() const override;
@@ -60,6 +61,8 @@ private:
 	TSet<FName> ConfiguredModules;
 	TArray<void*> LppPendingTokens;
 
+	FText EnableErrorText;
+
 	const FString FullEnginePluginsDir;
 	const FString FullProjectDir;
 	const FString FullProjectPluginsDir;
diff --git a/Engine/Source/Developer/Windows/LiveCoding/Public/ILiveCodingModule.h b/Engine/Source/Developer/Windows/LiveCoding/Public/ILiveCodingModule.h
--- a/Engine/Source/Developer/Windows/LiveCoding/Public/ILiveCodingModule.h
+++ b/Engine/Source/Developer/Windows/LiveCoding/Public/ILiveCodingModule.h
@@ -17,6 +17,7 @@ public:
 	virtual void EnableForSession(bool bEnabled) = 0;
 	virtual bool IsEnabledForSession() const = 0;
 	virtual bool CanEnableForSession() const = 0;
+	virtual FText GetEnableErrorText() const = 0;
 	virtual bool AutomaticallyCompileNewClasses() const = 0;
 
 	virtual bool HasStarted() const = 0;
diff --git a/Engine/Source/Editor/LevelEditor/Private/LevelEditorActions.cpp b/Engine/Source/Editor/LevelEditor/Private/LevelEditorActions.cpp
--- a/Engine/Source/Editor/LevelEditor/Private/LevelEditorActions.cpp
+++ b/Engine/Source/Editor/LevelEditor/Private/LevelEditorActions.cpp
@@ -1083,7 +1083,12 @@ void FLevelEditorActionCallbacks::RecompileGameCode_Clicked()
 		}
 		else
 		{
-			FMessageDialog::Open(EAppMsgType::Ok, LOCTEXT("NoLiveCodingCompileAfterHotReload", "Live Coding cannot be enabled while hot-reloaded modules are active. Please close the editor and build from your IDE before restarting."));
+			FText EnableErrorText = LiveCoding->GetEnableErrorText();
+			if (EnableErrorText.IsEmpty())
+			{
+				EnableErrorText = LOCTEXT("NoLiveCodingCompileAfterHotReload", "Live Coding cannot be enabled while hot-reloaded modules are active. Please close the editor and build from your IDE before restarting.");
+			}
+			FMessageDialog::Open(EAppMsgType::Ok, EnableErrorText);
 		}
 		return;
 	}
