Fixing compile issue with type not being defined

#rnx

diff --git a/Engine/Source/Developer/Windows/LiveCoding/Private/LiveCodingModule.cpp b/Engine/Source/Developer/Windows/LiveCoding/Private/LiveCodingModule.cpp
--- a/Engine/Source/Developer/Windows/LiveCoding/Private/LiveCodingModule.cpp
+++ b/Engine/Source/Developer/Windows/LiveCoding/Private/LiveCodingModule.cpp
@@ -298,7 +298,7 @@ bool FLiveCodingModule::IsEnabledForSession() const
 	return bEnabledForSession;
 }
 
-FText FLiveCodingModule::GetEnableErrorText() const
+const FText& FLiveCodingModule::GetEnableErrorText() const
 {
 	return EnableErrorText;
 }
diff --git a/Engine/Source/Developer/Windows/LiveCoding/Private/LiveCodingModule.h b/Engine/Source/Developer/Windows/LiveCoding/Private/LiveCodingModule.h
--- a/Engine/Source/Developer/Windows/LiveCoding/Private/LiveCodingModule.h
+++ b/Engine/Source/Developer/Windows/LiveCoding/Private/LiveCodingModule.h
@@ -33,7 +33,7 @@ public:
 	virtual bool IsEnabledByDefault() const override;
 	virtual void EnableForSession(bool bInEnabled) override;
 	virtual bool IsEnabledForSession() const override;
-	virtual FText GetEnableErrorText() const override;
+	virtual const FText& GetEnableErrorText() const override;
 	virtual bool AutomaticallyCompileNewClasses() const override;
 	virtual bool CanEnableForSession() const override;
 	virtual bool HasStarted() const override;
diff --git a/Engine/Source/Developer/Windows/LiveCoding/Public/ILiveCodingModule.h b/Engine/Source/Developer/Windows/LiveCoding/Public/ILiveCodingModule.h
--- a/Engine/Source/Developer/Windows/LiveCoding/Public/ILiveCodingModule.h
+++ b/Engine/Source/Developer/Windows/LiveCoding/Public/ILiveCodingModule.h
@@ -8,6 +8,8 @@
 
 #define LIVE_CODING_MODULE_NAME "LiveCoding"
 
+class FText;
+
 class ILiveCodingModule : public IModuleInterface
 {
 public:
@@ -17,7 +19,7 @@ public:
 	virtual void EnableForSession(bool bEnabled) = 0;
 	virtual bool IsEnabledForSession() const = 0;
 	virtual bool CanEnableForSession() const = 0;
-	virtual FText GetEnableErrorText() const = 0;
+	virtual const FText& GetEnableErrorText() const = 0;
 	virtual bool AutomaticallyCompileNewClasses() const = 0;
 
 	virtual bool HasStarted() const = 0;
