Fixing compile issue with type not being defined

#rnx

diff --git a/Engine/Source/Developer/Windows/LiveCoding/Private/LiveCodingModule.h b/Engine/Source/Developer/Windows/LiveCoding/Private/LiveCodingModule.h
--- a/Engine/Source/Developer/Windows/LiveCoding/Private/LiveCodingModule.h
+++ b/Engine/Source/Developer/Windows/LiveCoding/Private/LiveCodingModule.h
@@ -6,6 +6,7 @@
 #include "Delegates/Delegate.h"
 #include "Modules/ModuleManager.h"
 #include "Templates/SharedPointer.h"
+#include "Internationalization/Text.h"
 
 struct IConsoleCommand;
 class IConsoleVariable;
