Fix linking issue when doing Live Coding and adding new classes.

#rnx

diff --git a/Engine/Source/Developer/Windows/LiveCodingServer/Private/External/LC_LiveModule.cpp b/Engine/Source/Developer/Windows/LiveCodingServer/Private/External/LC_LiveModule.cpp
--- a/Engine/Source/Developer/Windows/LiveCodingServer/Private/External/LC_LiveModule.cpp
+++ b/Engine/Source/Developer/Windows/LiveCodingServer/Private/External/LC_LiveModule.cpp
@@ -1343,7 +1343,9 @@ void LiveModule::UpdateDirectoryCache(DirectoryCache* cache)
 }
 
 
-LiveModule::ErrorType::Enum LiveModule::Update(FileAttributeCache* fileCache, DirectoryCache* directoryCache, UpdateType::Enum updateType, const types::vector<symbols::ModifiedObjFile>& modifiedOrNewObjFiles)
+// BEGIN EPIC MOD
+LiveModule::ErrorType::Enum LiveModule::Update(FileAttributeCache* fileCache, DirectoryCache* directoryCache, UpdateType::Enum updateType, const types::vector<symbols::ModifiedObjFile>& modifiedOrNewObjFiles, const types::vector<std::wstring>& additionalLibraries)
+// END EPIC MOD
 {
 	telemetry::Scope updateScope("Update live module");
 
@@ -3050,6 +3052,20 @@ LiveModule::ErrorType::Enum LiveModule::Update(FileAttributeCache* fileCache, Di
 			linkerOptions += L"\"\n";
 		}
 	}
+	// BEGIN EPIC MOD
+	{
+		const size_t count = additionalLibraries.size();
+		for (size_t i = 0u; i < count; ++i)
+		{
+			const std::wstring& libPath = additionalLibraries[i];
+			LC_LOG_DEV("Pulling in additional LIB file %s", libPath.c_str());
+
+			linkerOptions += L"\"";
+			linkerOptions += libPath;
+			linkerOptions += L"\"\n";
+		}
+	}
+	// END EPIC MOD
 
 	// add UE4-specific helper library to support NatVis visualizers when debugging
 	if (appSettings::g_ue4EnableNatVisSupport->GetValue())
diff --git a/Engine/Source/Developer/Windows/LiveCodingServer/Private/External/LC_LiveModule.h b/Engine/Source/Developer/Windows/LiveCodingServer/Private/External/LC_LiveModule.h
--- a/Engine/Source/Developer/Windows/LiveCodingServer/Private/External/LC_LiveModule.h
+++ b/Engine/Source/Developer/Windows/LiveCodingServer/Private/External/LC_LiveModule.h
@@ -79,7 +79,9 @@ public:
 	// in DEFAULT mode, this checks for file modifications, compiles files automatically, builds a patch containing changes, and loads them into the host application.
 	// in EXTERNAL_BUILD_SYSTEM mode, this does not compile files but builds a patch containing modified .objs, loading the patch into the host application.
 	// optionally, an array of modified or new .objs can be given in this mode, which builds a patch containing these files, not checking for any other modifications.
-	ErrorType::Enum Update(FileAttributeCache* fileCache, DirectoryCache* directoryCache, UpdateType::Enum updateType, const types::vector<symbols::ModifiedObjFile>& modifiedOrNewObjFiles);
+	// BEGIN EPIC MOD
+	ErrorType::Enum Update(FileAttributeCache* fileCache, DirectoryCache* directoryCache, UpdateType::Enum updateType, const types::vector<symbols::ModifiedObjFile>& modifiedOrNewObjFiles, const types::vector<std::wstring>& additionalLibraries);
+	// END EPIC MOD
 	bool InstallCompiledPatches(LiveProcess* liveProcess, void* originalModuleBase);
 
 	const std::wstring& GetModuleName(void) const;
diff --git a/Engine/Source/Developer/Windows/LiveCodingServer/Private/External/LC_ServerCommandThread.cpp b/Engine/Source/Developer/Windows/LiveCodingServer/Private/External/LC_ServerCommandThread.cpp
--- a/Engine/Source/Developer/Windows/LiveCodingServer/Private/External/LC_ServerCommandThread.cpp
+++ b/Engine/Source/Developer/Windows/LiveCodingServer/Private/External/LC_ServerCommandThread.cpp
@@ -114,6 +114,9 @@ ServerCommandThread::ServerCommandThread(MainFrame* mainFrame, const wchar_t* co
 	, m_commandThreads()
 	, m_manualRecompileTriggered(false)
 	, m_liveModuleToModifiedOrNewObjFiles()
+// BEGIN EPIC MOD
+	, m_liveModuleToAdditionalLibraries()
+// END EPIC MOD
 	, m_restartCS()
 	, m_restartJob(nullptr)
 	, m_restartedProcessCount(0u)
@@ -801,7 +804,7 @@ void ServerCommandThread::CompileChanges(bool didAllProcessesMakeProgress, comma
 		GLiveCodingServer->GetStatusChangeDelegate().ExecuteIfBound(L"Compiling changes for live coding...");
 
 		// Keep retrying the compile until we've added all the required modules
-		TMap<FString, TArray<FString>> ModuleToObjectFiles;
+		FModuleToModuleFiles ModuleToModuleFiles;
 		ELiveCodingCompileReason CompileReason = ELiveCodingCompileReason::Initial;
 		for (;;)
 		{
@@ -814,24 +817,20 @@ void ServerCommandThread::CompileChanges(bool didAllProcessesMakeProgress, comma
 
 			// Execute the compile
 			TArray<FString> RequiredModules;
-			ELiveCodingCompileResult CompileResult = CompileDelegate.Execute(Targets, ValidModules, RequiredModules, ModuleToObjectFiles, CompileReason);
+			ELiveCodingCompileResult CompileResult = CompileDelegate.Execute(Targets, ValidModules, RequiredModules, ModuleToModuleFiles, CompileReason);
 			if (CompileResult == ELiveCodingCompileResult::Success)
 			{
 				break;
 			}
 			else if (CompileResult == ELiveCodingCompileResult::Canceled)
 			{
-				// BEGIN EPIC MOD
 				postCompileResult = commands::PostCompileResult::Cancelled;
-				// END EPIC MOD
 				GLiveCodingServer->GetCompileFinishedDelegate().ExecuteIfBound(ELiveCodingResult::Error, L"Compilation canceled.");
 				return;
 			}
 			else if (CompileResult == ELiveCodingCompileResult::Failure)
 			{
-				// BEGIN EPIC MOD
 				postCompileResult = commands::PostCompileResult::Failure;
-				// END EPIC MOD
 				GLiveCodingServer->GetCompileFinishedDelegate().ExecuteIfBound(ELiveCodingResult::Error, L"Compilation error.");
 				return;
 			}
@@ -839,9 +838,7 @@ void ServerCommandThread::CompileChanges(bool didAllProcessesMakeProgress, comma
 			// Enable any lazy-loaded modules that we need
 			if (!RequiredModules.IsEmpty() && !EnableRequiredModules(RequiredModules))
 			{
-				// BEGIN EPIC MOD
 				postCompileResult = commands::PostCompileResult::Failure;
-				// END EPIC MOD
 				GLiveCodingServer->GetCompileFinishedDelegate().ExecuteIfBound(ELiveCodingResult::Error, L"Compilation error.");
 				return;
 			}
@@ -859,7 +856,7 @@ void ServerCommandThread::CompileChanges(bool didAllProcessesMakeProgress, comma
 			ValidModuleFileNames.insert(liveModule->GetModuleName());
 		}
 
-		for(const TPair<FString, TArray<FString>>& Pair : ModuleToObjectFiles)
+		for(const TPair<FString, FModuleFiles>& Pair : ModuleToModuleFiles)
 		{
 			std::wstring ModuleFileName = Filesystem::NormalizePath(*Pair.Key).GetString();
 			if(ValidModuleFileNames.find(ModuleFileName) == ValidModuleFileNames.end())
@@ -891,7 +888,7 @@ void ServerCommandThread::CompileChanges(bool didAllProcessesMakeProgress, comma
 			}
 
 			types::vector<symbols::ModifiedObjFile> ObjectFiles;
-			for(const FString& ObjectFile : Pair.Value)
+			for (const FString& ObjectFile : Pair.Value.Objects)
 			{
 				std::wstring NormalizedObjectFile = Filesystem::NormalizePath(*ObjectFile).GetString();
 
@@ -928,6 +925,13 @@ void ServerCommandThread::CompileChanges(bool didAllProcessesMakeProgress, comma
 			}
 
 			m_liveModuleToModifiedOrNewObjFiles.insert(std::make_pair(ModuleFileName, std::move(ObjectFiles)));
+
+			types::vector<std::wstring> AdditionalLibraries;
+			for (const FString& ObjectFile : Pair.Value.Libraries)
+			{
+				AdditionalLibraries.push_back(Filesystem::NormalizePath(*ObjectFile).GetString());
+			}
+			m_liveModuleToAdditionalLibraries.insert(std::make_pair(ModuleFileName, std::move(AdditionalLibraries)));
 		}
 	}
 
@@ -958,6 +962,10 @@ void ServerCommandThread::CompileChanges(bool didAllProcessesMakeProgress, comma
 		? LiveModule::UpdateType::DEFAULT
 		: LiveModule::UpdateType::NO_CLIENT_COMMUNICATION;
 
+	// BEGIN EPIC MOD
+	static const types::vector<std::wstring> emptyLibFiles;
+	// END EPIC MOD
+
 	// has the user given us at least one modified or new .obj file for at least one of the modules?
 	const bool hasAtLeastOneOptionalObj = (m_liveModuleToModifiedOrNewObjFiles.size() != 0u);
 	for (size_t i = 0u; i < count; ++i)
@@ -976,16 +984,23 @@ void ServerCommandThread::CompileChanges(bool didAllProcessesMakeProgress, comma
 			}
 			else
 			{
+				// BEGIN EPIC MOD
+				const auto libFilesIt = m_liveModuleToAdditionalLibraries.find(liveModule->GetModuleName());
+				const types::vector<std::wstring>& libFiles = libFilesIt != m_liveModuleToAdditionalLibraries.end() ? libFilesIt->second : emptyLibFiles;
+
 				// build a patch with the given list of .objs for this module
 				const types::vector<symbols::ModifiedObjFile>& objFiles = objFilesIt->second;
-				moduleUpdateError = liveModule->Update(&fileCache, m_directoryCache, updateType, objFiles);
+				moduleUpdateError = liveModule->Update(&fileCache, m_directoryCache, updateType, objFiles, libFiles);
+				// END EPIC MOD
 			}
 		}
 		else
 		{
 			// no optional .objs were given, update all live modules regularly
 			types::vector<symbols::ModifiedObjFile> emptyObjs;
-			moduleUpdateError = liveModule->Update(&fileCache, m_directoryCache, updateType, emptyObjs);
+			// BEGIN EPIC MOD
+			moduleUpdateError = liveModule->Update(&fileCache, m_directoryCache, updateType, emptyObjs, emptyLibFiles);
+			// END EPIC MOD
 		}
 
 		// only accept new error conditions for this module if there haven't been any updates until now.
@@ -1396,6 +1411,9 @@ Thread::ReturnValue ServerCommandThread::CompileThread(void)
 			// clear API recompiles
 			m_manualRecompileTriggered = false;
 			m_liveModuleToModifiedOrNewObjFiles.clear();
+			// BEGIN EPIC MOD
+			m_liveModuleToAdditionalLibraries.clear();
+			// END EPIC MOD
 		}
 
 		m_restartCS.Leave();
diff --git a/Engine/Source/Developer/Windows/LiveCodingServer/Private/External/LC_ServerCommandThread.h b/Engine/Source/Developer/Windows/LiveCodingServer/Private/External/LC_ServerCommandThread.h
--- a/Engine/Source/Developer/Windows/LiveCodingServer/Private/External/LC_ServerCommandThread.h
+++ b/Engine/Source/Developer/Windows/LiveCodingServer/Private/External/LC_ServerCommandThread.h
@@ -185,6 +185,10 @@ private:
 	bool m_manualRecompileTriggered;
 	types::unordered_map<std::wstring, types::vector<symbols::ModifiedObjFile>> m_liveModuleToModifiedOrNewObjFiles;
 
+	// BEGIN EPIC MOD
+	types::unordered_map<std::wstring, types::vector<std::wstring>> m_liveModuleToAdditionalLibraries;
+	// END EPIC MOD
+
 	// restart mechanism
 	CriticalSection m_restartCS;
 	void* m_restartJob;
diff --git a/Engine/Source/Developer/Windows/LiveCodingServer/Public/ILiveCodingServer.h b/Engine/Source/Developer/Windows/LiveCodingServer/Public/ILiveCodingServer.h
--- a/Engine/Source/Developer/Windows/LiveCodingServer/Public/ILiveCodingServer.h
+++ b/Engine/Source/Developer/Windows/LiveCodingServer/Public/ILiveCodingServer.h
@@ -46,6 +46,13 @@ public:
 	virtual void ShutdownModule() override = 0;
 };
 
+struct FModuleFiles
+{
+	TArray<FString> Objects;
+	TArray<FString> Libraries;
+};
+typedef TMap<FString, FModuleFiles> FModuleToModuleFiles;
+
 class ILiveCodingServer : public IModularFeature
 {
 public:
@@ -68,8 +75,7 @@ public:
 	DECLARE_DELEGATE_TwoParams(FLogOutputDelegate, ELiveCodingLogVerbosity, const wchar_t*);
 	virtual FLogOutputDelegate& GetLogOutputDelegate() = 0;
 
-	typedef TMap<FString, TArray<FString>> FModuleToObjectFiles;
-	DECLARE_DELEGATE_RetVal_FiveParams(ELiveCodingCompileResult, FCompileDelegate, const TArray<FString>&, const TArray<FString>&, TArray<FString>&, FModuleToObjectFiles&, ELiveCodingCompileReason);
+	DECLARE_DELEGATE_RetVal_FiveParams(ELiveCodingCompileResult, FCompileDelegate, const TArray<FString>&, const TArray<FString>&, TArray<FString>&, FModuleToModuleFiles&, ELiveCodingCompileReason);
 	virtual FCompileDelegate& GetCompileDelegate() = 0;
 
 	DECLARE_DELEGATE(FCompileStartedDelegate);
diff --git a/Engine/Source/Programs/LiveCodingConsole/Private/LiveCodingConsole.cpp b/Engine/Source/Programs/LiveCodingConsole/Private/LiveCodingConsole.cpp
--- a/Engine/Source/Programs/LiveCodingConsole/Private/LiveCodingConsole.cpp
+++ b/Engine/Source/Programs/LiveCodingConsole/Private/LiveCodingConsole.cpp
@@ -265,7 +265,7 @@ private:
 		LogWidget->AppendLine(GetLogColor(Verbosity), MoveTemp(Text));
 	}
 
-	ELiveCodingCompileResult CompilePatch(const TArray<FString>& Targets, const TArray<FString>& ValidModules, TArray<FString>& RequiredModules, TMap<FString, TArray<FString>>& ModuleToObjectFiles, ELiveCodingCompileReason CompileReason)
+	ELiveCodingCompileResult CompilePatch(const TArray<FString>& Targets, const TArray<FString>& ValidModules, TArray<FString>& RequiredModules, FModuleToModuleFiles& ModuleToModuleFiles, ELiveCodingCompileReason CompileReason)
 	{
 		// Update the compile start time. This gets copied into the last patch time once a patch has been confirmed to have been applied.
 		NextPatchStartTime = FDateTime::UtcNow();
@@ -385,10 +385,19 @@ private:
 			{
 				if (FileManager.GetTimeStamp(*ObjectFileName) > MinTimeStamp)
 				{
-					ModuleToObjectFiles.FindOrAdd(Pair.Key).Add(ObjectFileName);
+					ModuleToModuleFiles.FindOrAdd(Pair.Key).Objects.Add(ObjectFileName);
 				}
 			}
 		}
+
+		// Add all of the additional libraries
+		for (TPair<FString, TArray<FString>>& Pair : Manifest.Libraries)
+		{
+			for (const FString& Library : Pair.Value)
+			{
+				ModuleToModuleFiles.FindOrAdd(Pair.Key).Libraries.Add(Library);
+			}
+		}
 		return ELiveCodingCompileResult::Success;
 	}
 
diff --git a/Engine/Source/Programs/LiveCodingConsole/Private/LiveCodingManifest.cpp b/Engine/Source/Programs/LiveCodingConsole/Private/LiveCodingManifest.cpp
--- a/Engine/Source/Programs/LiveCodingConsole/Private/LiveCodingManifest.cpp
+++ b/Engine/Source/Programs/LiveCodingConsole/Private/LiveCodingManifest.cpp
@@ -94,6 +94,21 @@ bool FLiveCodingManifest::Parse(FJsonObject& Object, FString& OutFailReason)
 			}
 			InputFiles.Add(Input->AsString());
 		}
+
+		const TArray<TSharedPtr<FJsonValue>>* JsonLibraries;
+		if (ModuleObject.TryGetArrayField(TEXT("Libraries"), JsonLibraries))
+		{
+			TArray<FString>& LibraryFiles = Libraries.Add(OutputFile);
+			for (const TSharedPtr<FJsonValue>& Library : *JsonLibraries)
+			{
+				if (Library->Type != EJson::String)
+				{
+					OutFailReason = TEXT("invalid module library field");
+					return false;
+				}
+				LibraryFiles.Add(Library->AsString());
+			}
+		}
 	}
 
 	return true;
diff --git a/Engine/Source/Programs/LiveCodingConsole/Private/LiveCodingManifest.h b/Engine/Source/Programs/LiveCodingConsole/Private/LiveCodingManifest.h
--- a/Engine/Source/Programs/LiveCodingConsole/Private/LiveCodingManifest.h
+++ b/Engine/Source/Programs/LiveCodingConsole/Private/LiveCodingManifest.h
@@ -11,6 +11,7 @@ struct FLiveCodingManifest
 	FString LinkerPath;
 	TMap<FString, FString> LinkerEnvironment;
 	TMap<FString, TArray<FString>> BinaryToObjectFiles;
+	TMap<FString, TArray<FString>> Libraries;
 
 	bool Read(const TCHAR* FileName, FString& OutFailReason);
 	bool Parse(FJsonObject& Object, FString& OutFailReason);
diff --git a/Engine/Source/Programs/UnrealBuildTool/System/HotReload.cs b/Engine/Source/Programs/UnrealBuildTool/System/HotReload.cs
--- a/Engine/Source/Programs/UnrealBuildTool/System/HotReload.cs
+++ b/Engine/Source/Programs/UnrealBuildTool/System/HotReload.cs
@@ -1342,6 +1342,16 @@ namespace UnrealBuildTool
 							}
 							Writer.WriteArrayEnd();
 
+							Writer.WriteArrayStart("Libraries");
+							foreach (FileItem InputFile in Action.PrerequisiteItems)
+							{
+								if (InputFile.HasExtension(".lib"))
+								{
+									Writer.WriteValue(InputFile.FullName);
+								}
+							}
+							Writer.WriteArrayEnd();
+
 							Writer.WriteObjectEnd();
 						}
 					}
