Live coding now handles adding a feature pack as long as the project already contained source files.  Otherwise it requests that the game be built from the IDE.  Will also request that the game be built from the IDE when adding a class to a blueprint only project.

#rnx

diff --git a/Engine/Source/Developer/Windows/LiveCoding/Private/LiveCodingModule.cpp b/Engine/Source/Developer/Windows/LiveCoding/Private/LiveCodingModule.cpp
--- a/Engine/Source/Developer/Windows/LiveCoding/Private/LiveCodingModule.cpp
+++ b/Engine/Source/Developer/Windows/LiveCoding/Private/LiveCodingModule.cpp
@@ -13,6 +13,7 @@
 #include "Misc/Paths.h"
 #include "Misc/ConfigCacheIni.h"
 #include "Misc/MessageDialog.h"
+#include "Misc/ScopedSlowTask.h"
 #include "LiveCodingSettings.h"
 #include "ISettingsModule.h"
 #include "ISettingsSection.h"
@@ -164,7 +165,7 @@ void FLiveCodingModule::StartupModule()
 	CompileCommand = ConsoleManager.RegisterConsoleCommand(
 		TEXT("LiveCoding.Compile"),
 		TEXT("Initiates a live coding compile"),
-		FConsoleCommandDelegate::CreateRaw(this, &FLiveCodingModule::Compile),
+		FConsoleCommandDelegate::CreateLambda([this] { Compile(ELiveCodingCompileFlags::None, nullptr); }),
 		ECVF_Cheat
 	);
 
@@ -332,16 +333,63 @@ void FLiveCodingModule::ShowConsole()
 
 void FLiveCodingModule::Compile()
 {
-	if(!GIsCompileActive)
+	Compile(ELiveCodingCompileFlags::None, nullptr);
+}
+
+inline bool ReturnResults(ELiveCodingCompileResult InResult, ELiveCodingCompileResult* OutResult)
+{
+	if (OutResult != nullptr)
 	{
-		EnableForSession(true);
-		if(bStarted)
-		{
-			UpdateModules(); // Need to do this immediately rather than waiting until next tick
-			LppTriggerRecompile();
-			GIsCompileActive = true;
-		}
+		*OutResult = InResult;
+	}
+	return InResult == ELiveCodingCompileResult::Success || InResult == ELiveCodingCompileResult::NoChanges || InResult == ELiveCodingCompileResult::InProgress;
+}
+
+bool FLiveCodingModule::Compile(ELiveCodingCompileFlags CompileFlags, ELiveCodingCompileResult* Result)
+{
+	if (GIsCompileActive)
+	{
+		return ReturnResults(ELiveCodingCompileResult::CompileStillActive, Result);
+	}
+
+	EnableForSession(true);
+	if (!bStarted)
+	{
+		return ReturnResults(ELiveCodingCompileResult::NotStarted, Result);
 	}
+
+	// Need to do this immediately rather than waiting until next tick
+	UpdateModules(); 
+
+	// Trigger the recompile
+	GIsCompileActive = true;
+	LastResults = ELiveCodingCompileResult::Failure;
+	LppTriggerRecompile();
+
+	// If we aren't waiting, just return now
+	if (!EnumHasAnyFlags(CompileFlags, ELiveCodingCompileFlags::WaitForCompletion))
+	{
+		return ReturnResults(ELiveCodingCompileResult::InProgress, Result);
+	}
+
+	// Wait until we are no longer compiling.  Cancellation is handled via other mechanisms and
+	// need not be detected in this loop.  GIsCompileActive will be cleared.
+	FText StatusUpdate = LOCTEXT("CompileStatusMessage", "Compiling...");
+	FScopedSlowTask SlowTask(0, StatusUpdate, GIsSlowTask);
+	SlowTask.MakeDialog();
+
+	// Wait until the compile completes
+	while (GIsCompileActive)
+	{
+		SlowTask.EnterProgressFrame(0.0f);
+		AttemptSyncLivePatching();
+		FPlatformProcess::Sleep(0.01f);
+	}
+
+	// A final sync to get the result and complete the process
+	AttemptSyncLivePatching();
+
+	return ReturnResults(LastResults, Result);
 }
 
 bool FLiveCodingModule::IsCompiling() const
@@ -488,6 +536,7 @@ void FLiveCodingModule::AttemptSyncLivePatching()
 		{
 			LppSyncPoint();
 		}
+
 		if (!GIsCompileActive)
 		{
 			static const FString Success("Live coding succeeded");
@@ -498,6 +547,7 @@ void FLiveCodingModule::AttemptSyncLivePatching()
 			switch (GPostCompileResult)
 			{
 			case commands::PostCompileResult::Success:
+				LastResults = ELiveCodingCompileResult::Success;
 				if (bHasReinstancingOccurred)
 				{
 					if (!IsReinstancingEnabled())
@@ -519,15 +569,19 @@ void FLiveCodingModule::AttemptSyncLivePatching()
 				}
 				break;
 			case commands::PostCompileResult::NoChanges:
+				LastResults = ELiveCodingCompileResult::NoChanges;
 				UE_LOG(LogLiveCoding, Display, TEXT("%s, %s"), *Success, TEXT("no code changes detected"));
 				break;
 			case commands::PostCompileResult::Cancelled:
-				UE_LOG(LogLiveCoding, Error, TEXT("Live coding cancelled"));
+				LastResults = ELiveCodingCompileResult::Cancelled;
+				UE_LOG(LogLiveCoding, Error, TEXT("Live coding canceled"));
 				break;
 			case commands::PostCompileResult::Failure:
+				LastResults = ELiveCodingCompileResult::Failure;
 				UE_LOG(LogLiveCoding, Error, TEXT("Live coding failed, please see Live console for more information"));
 				break;
 			default:
+				LastResults = ELiveCodingCompileResult::Failure;
 				check(false);
 			}
 
diff --git a/Engine/Source/Developer/Windows/LiveCoding/Private/LiveCodingModule.h b/Engine/Source/Developer/Windows/LiveCoding/Private/LiveCodingModule.h
--- a/Engine/Source/Developer/Windows/LiveCoding/Private/LiveCodingModule.h
+++ b/Engine/Source/Developer/Windows/LiveCoding/Private/LiveCodingModule.h
@@ -40,6 +40,7 @@ public:
 	virtual bool HasStarted() const override;
 	virtual void ShowConsole() override;
 	virtual void Compile() override;
+	virtual bool Compile(ELiveCodingCompileFlags CompileFlags, ELiveCodingCompileResult* Result) override;
 	virtual bool IsCompiling() const override;
 	virtual void Tick() override;
 	virtual FOnPatchCompleteDelegate& GetOnPatchCompleteDelegate() override;
@@ -59,6 +60,7 @@ private:
 	bool bUpdateModulesInTick = false;
 	bool bHasReinstancingOccurred = false;
 	bool bHasPatchBeenLoaded = false;
+	ELiveCodingCompileResult LastResults = ELiveCodingCompileResult::Success;
 	TSet<FName> ConfiguredModules;
 	TArray<void*> LppPendingTokens;
 
diff --git a/Engine/Source/Developer/Windows/LiveCoding/Public/ILiveCodingModule.h b/Engine/Source/Developer/Windows/LiveCoding/Public/ILiveCodingModule.h
--- a/Engine/Source/Developer/Windows/LiveCoding/Public/ILiveCodingModule.h
+++ b/Engine/Source/Developer/Windows/LiveCoding/Public/ILiveCodingModule.h
@@ -10,6 +10,24 @@
 
 class FText;
 
+enum class ELiveCodingCompileFlags : uint8
+{
+	None = 0,
+	WaitForCompletion = 1 << 0,
+};
+ENUM_CLASS_FLAGS(ELiveCodingCompileFlags)
+
+enum class ELiveCodingCompileResult : uint8
+{
+	Success,			// Compile completed and there were changes 
+	NoChanges,			// Compile completed and there were no changes
+	InProgress,			// Compile started but wait for completion was not specified
+	CompileStillActive,	// A prior compile request is still active
+	NotStarted,			// Live coding monitor could not be started
+	Failure,			// Complete completed but there was an error
+	Cancelled,			// Compile was cancelled
+};
+
 class ILiveCodingModule : public IModuleInterface
 {
 public:
@@ -26,6 +44,7 @@ public:
 
 	virtual void ShowConsole() = 0;
 	virtual void Compile() = 0;
+	virtual bool Compile(ELiveCodingCompileFlags CompileFlags, ELiveCodingCompileResult* Result) = 0;
 	virtual bool IsCompiling() const = 0;
 	virtual void Tick() = 0;
 
diff --git a/Engine/Source/Editor/GameProjectGeneration/Private/GameProjectUtils.cpp b/Engine/Source/Editor/GameProjectGeneration/Private/GameProjectUtils.cpp
--- a/Engine/Source/Editor/GameProjectGeneration/Private/GameProjectUtils.cpp
+++ b/Engine/Source/Editor/GameProjectGeneration/Private/GameProjectUtils.cpp
@@ -4135,9 +4135,15 @@ GameProjectUtils::EAddCodeToProjectResult GameProjectUtils::AddCodeToProject_Int
 	ILiveCodingModule* LiveCoding = FModuleManager::GetModulePtr<ILiveCodingModule>(LIVE_CODING_MODULE_NAME);
 	if (LiveCoding != nullptr && LiveCoding->IsEnabledForSession())
 	{
-		if (bProjectHadCodeFiles && LiveCoding->AutomaticallyCompileNewClasses())
+		if (!bProjectHadCodeFiles)
 		{
-			LiveCoding->Compile();
+			FMessageDialog::Open(EAppMsgType::Ok, LOCTEXT("LiveCodingNoSources", "Project now includes sources, please close the editor and build from your IDE."));
+			return EAddCodeToProjectResult::Succeeded;
+		}
+
+		if (LiveCoding->AutomaticallyCompileNewClasses())
+		{
+			LiveCoding->Compile(ELiveCodingCompileFlags::None, nullptr);
 			OutReloadStatus = EReloadStatus::Reloaded;
 		}
 		return EAddCodeToProjectResult::Succeeded;
diff --git a/Engine/Source/Editor/UnrealEd/Private/Factories/PackFactory.cpp b/Engine/Source/Editor/UnrealEd/Private/Factories/PackFactory.cpp
--- a/Engine/Source/Editor/UnrealEd/Private/Factories/PackFactory.cpp
+++ b/Engine/Source/Editor/UnrealEd/Private/Factories/PackFactory.cpp
@@ -683,7 +683,19 @@ UObject* UPackFactory::FactoryCreateBinary
 					ILiveCodingModule* LiveCoding = FModuleManager::GetModulePtr<ILiveCodingModule>(LIVE_CODING_MODULE_NAME);
 					if (LiveCoding != nullptr && LiveCoding->IsEnabledForSession())
 					{
-						FMessageDialog::Open(EAppMsgType::Ok, NSLOCTEXT("PackFactory", "CannotCompileWithLiveCoding", "Unable to compile source code while Live Coding is enabled. Please close the editor and build from your IDE."));
+						if (bProjectHadSourceFiles)
+						{
+							if (!LiveCoding->Compile(ELiveCodingCompileFlags::WaitForCompletion, nullptr))
+							{
+								FMessageDialog::Open(EAppMsgType::Ok, NSLOCTEXT("PackFactory", "LiveCodingFailedToCompile", "Failed to compile sources, please close the editor and build from your IDE."));
+							}
+						}
+						else
+						{
+							FMessageDialog::Open(EAppMsgType::Ok, NSLOCTEXT("PackFactory", "LiveCodingNoSources", "Project now includes sources, please close the editor and build from your IDE."));
+						}
+
+						// Don't allow hot-reload to try to compile
 						bCompileSource = false;
 					}
 				}
