Fixing tooltips in platform settings

diff --git a/Engine/Plugins/Runtime/AR/Microsoft/HoloLensAR/Source/HoloLensPlatformEditor/Private/HoloLensTargetSettings.h b/Engine/Plugins/Runtime/AR/Microsoft/HoloLensAR/Source/HoloLensPlatformEditor/Private/HoloLensTargetSettings.h
--- a/Engine/Plugins/Runtime/AR/Microsoft/HoloLensAR/Source/HoloLensPlatformEditor/Private/HoloLensTargetSettings.h
+++ b/Engine/Plugins/Runtime/AR/Microsoft/HoloLensAR/Source/HoloLensPlatformEditor/Private/HoloLensTargetSettings.h
@@ -142,15 +142,15 @@ public:
 	UPROPERTY(EditAnywhere, config, Category = Capabilities)
 	bool bSetDefaultCapabilities = true;
 
-	/** Which of the currently enabled spatialization plugins to use on Windows. */
+	/** Which of the currently enabled spatialization plugins to use. */
 	UPROPERTY(config, EditAnywhere, Category = "Audio")
 	FString SpatializationPlugin;
 
-	/** Which of the currently enabled reverb plugins to use on Windows. */
+	/** Which of the currently enabled reverb plugins to use. */
 	UPROPERTY(config, EditAnywhere, Category = "Audio")
 	FString ReverbPlugin;
 
-	/** Which of the currently enabled occlusion plugins to use on Windows. */
+	/** Which of the currently enabled occlusion plugins to use. */
 	UPROPERTY(config, EditAnywhere, Category = "Audio")
 	FString OcclusionPlugin;
 
diff --git a/Engine/Source/Developer/Linux/LinuxTargetPlatform/Classes/LinuxTargetSettings.h b/Engine/Source/Developer/Linux/LinuxTargetPlatform/Classes/LinuxTargetSettings.h
--- a/Engine/Source/Developer/Linux/LinuxTargetPlatform/Classes/LinuxTargetSettings.h
+++ b/Engine/Source/Developer/Linux/LinuxTargetPlatform/Classes/LinuxTargetSettings.h
@@ -23,15 +23,15 @@ public:
 
 	GENERATED_UCLASS_BODY()
 
-	/** Which of the currently enabled spatialization plugins to use on Windows. */
+	/** Which of the currently enabled spatialization plugins to use. */
 	UPROPERTY(config, EditAnywhere, Category = "Audio")
 	FString SpatializationPlugin;
 
-	/** Which of the currently enabled reverb plugins to use on Windows. */
+	/** Which of the currently enabled reverb plugins to use. */
 	UPROPERTY(config, EditAnywhere, Category = "Audio")
 	FString ReverbPlugin;
 
-	/** Which of the currently enabled occlusion plugins to use on Windows. */
+	/** Which of the currently enabled occlusion plugins to use. */
 	UPROPERTY(config, EditAnywhere, Category = "Audio")
 	FString OcclusionPlugin;
 
diff --git a/Engine/Source/Developer/Mac/MacTargetPlatform/Classes/MacTargetSettings.h b/Engine/Source/Developer/Mac/MacTargetPlatform/Classes/MacTargetSettings.h
--- a/Engine/Source/Developer/Mac/MacTargetPlatform/Classes/MacTargetSettings.h
+++ b/Engine/Source/Developer/Mac/MacTargetPlatform/Classes/MacTargetSettings.h
@@ -122,15 +122,15 @@ public:
 	UPROPERTY(config, EditAnywhere, Category = "Audio", meta = (ClampMin = "0", UIMin = "0", DisplayName = "Number of Source Workers"))
 	int32 AudioNumSourceWorkers;
 
-	/** Which of the currently enabled spatialization plugins to use on Windows. */
+	/** Which of the currently enabled spatialization plugins to use. */
 	UPROPERTY(config, EditAnywhere, Category = "Audio")
 	FString SpatializationPlugin;
 
-	/** Which of the currently enabled reverb plugins to use on Windows. */
+	/** Which of the currently enabled reverb plugins to use. */
 	UPROPERTY(config, EditAnywhere, Category = "Audio")
 	FString ReverbPlugin;
 
-	/** Which of the currently enabled occlusion plugins to use on Windows. */
+	/** Which of the currently enabled occlusion plugins to use. */
 	UPROPERTY(config, EditAnywhere, Category = "Audio")
 	FString OcclusionPlugin;
 	
diff --git a/Engine/Source/Developer/Windows/WindowsTargetPlatform/Classes/WindowsTargetSettings.h b/Engine/Source/Developer/Windows/WindowsTargetPlatform/Classes/WindowsTargetSettings.h
--- a/Engine/Source/Developer/Windows/WindowsTargetPlatform/Classes/WindowsTargetSettings.h
+++ b/Engine/Source/Developer/Windows/WindowsTargetPlatform/Classes/WindowsTargetSettings.h
@@ -81,19 +81,19 @@ public:
 	UPROPERTY(config, EditAnywhere, Category = "Audio", meta = (ClampMin = "0", UIMin = "0", DisplayName = "Number of Source Workers"))
 	int32 AudioNumSourceWorkers;
 
-	/** Which of the currently enabled spatialization plugins to use on Windows. */
+	/** Which of the currently enabled spatialization plugins to use. */
 	UPROPERTY(config, EditAnywhere, Category = "Audio")
 	FString SpatializationPlugin;
 
-	/** Which of the currently enabled source data override plugins to use on Windows. */
+	/** Which of the currently enabled source data override plugins to use. */
 	UPROPERTY(config, EditAnywhere, Category = "Audio")
 	FString SourceDataOverridePlugin;
 
-	/** Which of the currently enabled reverb plugins to use on Windows. */
+	/** Which of the currently enabled reverb plugins to use. */
 	UPROPERTY(config, EditAnywhere, Category = "Audio")
 	FString ReverbPlugin;
 
-	/** Which of the currently enabled occlusion plugins to use on Windows. */
+	/** Which of the currently enabled occlusion plugins to use. */
 	UPROPERTY(config, EditAnywhere, Category = "Audio")
 	FString OcclusionPlugin;
 
diff --git a/Engine/Source/Runtime/Android/AndroidRuntimeSettings/Classes/AndroidRuntimeSettings.h b/Engine/Source/Runtime/Android/AndroidRuntimeSettings/Classes/AndroidRuntimeSettings.h
--- a/Engine/Source/Runtime/Android/AndroidRuntimeSettings/Classes/AndroidRuntimeSettings.h
+++ b/Engine/Source/Runtime/Android/AndroidRuntimeSettings/Classes/AndroidRuntimeSettings.h
@@ -502,15 +502,15 @@ public:
 	UPROPERTY(config, EditAnywhere, Category = "Audio", meta = (ClampMin = "0", UIMin = "0", DisplayName = "Number of Source Workers"))
 	int32 AudioNumSourceWorkers;
 
-	/** Which of the currently enabled spatialization plugins to use on Windows. */
+	/** Which of the currently enabled spatialization plugins to use. */
 	UPROPERTY(config, EditAnywhere, Category = "Audio")
 	FString SpatializationPlugin;
 
-	/** Which of the currently enabled reverb plugins to use on Windows. */
+	/** Which of the currently enabled reverb plugins to use. */
 	UPROPERTY(config, EditAnywhere, Category = "Audio")
 	FString ReverbPlugin;
 
-	/** Which of the currently enabled occlusion plugins to use on Windows. */
+	/** Which of the currently enabled occlusion plugins to use. */
 	UPROPERTY(config, EditAnywhere, Category = "Audio")
 	FString OcclusionPlugin;
 
diff --git a/Engine/Source/Runtime/IOS/IOSRuntimeSettings/Classes/IOSRuntimeSettings.h b/Engine/Source/Runtime/IOS/IOSRuntimeSettings/Classes/IOSRuntimeSettings.h
--- a/Engine/Source/Runtime/IOS/IOSRuntimeSettings/Classes/IOSRuntimeSettings.h
+++ b/Engine/Source/Runtime/IOS/IOSRuntimeSettings/Classes/IOSRuntimeSettings.h
@@ -492,15 +492,15 @@ public:
 	UPROPERTY(config, EditAnywhere, Category = "Audio", meta = (ClampMin = "0", UIMin = "0", DisplayName = "Number of Source Workers"))
 	int32 AudioNumSourceWorkers;
 
-	/** Which of the currently enabled spatialization plugins to use on Windows. */
+	/** Which of the currently enabled spatialization plugins to use. */
 	UPROPERTY(config, EditAnywhere, Category = "Audio")
 	FString SpatializationPlugin;
 
-	/** Which of the currently enabled reverb plugins to use on Windows. */
+	/** Which of the currently enabled reverb plugins to use. */
 	UPROPERTY(config, EditAnywhere, Category = "Audio")
 	FString ReverbPlugin;
 
-	/** Which of the currently enabled occlusion plugins to use on Windows. */
+	/** Which of the currently enabled occlusion plugins to use. */
 	UPROPERTY(config, EditAnywhere, Category = "Audio")
 	FString OcclusionPlugin;
 
