DXC provides the PDB name along with the DXC_OUT_PDB blob so we should use that instead of trying to (incorrectly) generate the hash ourselves. Removing the unnecessary string blob argument from all the other calls to GetOutput.

diff --git a/Engine/Source/Developer/Windows/ShaderFormatD3D/Private/D3DShaderCompilerDXC.cpp b/Engine/Source/Developer/Windows/ShaderFormatD3D/Private/D3DShaderCompilerDXC.cpp
--- a/Engine/Source/Developer/Windows/ShaderFormatD3D/Private/D3DShaderCompilerDXC.cpp
+++ b/Engine/Source/Developer/Windows/ShaderFormatD3D/Private/D3DShaderCompilerDXC.cpp
@@ -300,12 +300,11 @@ static HRESULT D3DCompileToDxil(const char* SourceText, FDxcArguments& Arguments
 	CompileResult->GetStatus(&CompileResultCode);
 	if (SUCCEEDED(CompileResultCode))
 	{
-		TRefCountPtr<IDxcBlobUtf16> Dummy;
 		checkf(CompileResult->HasOutput(DXC_OUT_OBJECT), TEXT("No object code found!"));
-		VERIFYHRESULT(CompileResult->GetOutput(DXC_OUT_OBJECT, IID_PPV_ARGS(OutDxilBlob.GetInitReference()), Dummy.GetInitReference()));
+		VERIFYHRESULT(CompileResult->GetOutput(DXC_OUT_OBJECT, IID_PPV_ARGS(OutDxilBlob.GetInitReference()), nullptr));
 
 		checkf(CompileResult->HasOutput(DXC_OUT_REFLECTION), TEXT("No reflection found!"));
-		VERIFYHRESULT(CompileResult->GetOutput(DXC_OUT_REFLECTION, IID_PPV_ARGS(OutReflectionBlob.GetInitReference()), Dummy.GetInitReference()));
+		VERIFYHRESULT(CompileResult->GetOutput(DXC_OUT_REFLECTION, IID_PPV_ARGS(OutReflectionBlob.GetInitReference()), nullptr));
 
 		if (Arguments.ShouldDump())
 		{
@@ -318,27 +317,16 @@ static HRESULT D3DCompileToDxil(const char* SourceText, FDxcArguments& Arguments
 			FString DxilFile = Arguments.GetDumpDisassemblyFilename().LeftChop(7) + TEXT("_refl.dxil");
 			SaveDxcBlobToFile(OutDxilBlob, DxilFile);
 
-			if (CompileResult->HasOutput(DXC_OUT_PDB) && CompileResult->HasOutput(DXC_OUT_SHADER_HASH))
+			if (CompileResult->HasOutput(DXC_OUT_PDB))
 			{
 				TRefCountPtr<IDxcBlob> PdbBlob;
-				VERIFYHRESULT(CompileResult->GetOutput(DXC_OUT_PDB, IID_PPV_ARGS(PdbBlob.GetInitReference()), Dummy.GetInitReference()));
+				TRefCountPtr<IDxcBlobUtf16> PdbNameBlob;
+				VERIFYHRESULT(CompileResult->GetOutput(DXC_OUT_PDB, IID_PPV_ARGS(PdbBlob.GetInitReference()), PdbNameBlob.GetInitReference()));
 
-				TRefCountPtr<IDxcBlob> HashBlob;
-				VERIFYHRESULT(CompileResult->GetOutput(DXC_OUT_SHADER_HASH, IID_PPV_ARGS(HashBlob.GetInitReference()), Dummy.GetInitReference()));
-
-				check(sizeof(DxcShaderHash) == HashBlob->GetBufferSize());
-				const DxcShaderHash* ShaderHash = (DxcShaderHash*)HashBlob->GetBufferPointer();
-
-				FString HashName;
-				static_assert(sizeof(DxcShaderHash::HashDigest) == 16, "Hash changed");
-				for (int32 Index = 0; Index < 16; ++Index)
-				{
-					HashName += FString::Printf(TEXT("%02x"), ShaderHash->HashDigest[Index]);
-				}
+				const FString PdbName = PdbNameBlob->GetStringPointer();
 
 				// Dump pdb (.d3dasm -> .pdb)
-				//#todo-rco: Need to put this in a central location
-				FString PdbFile = Arguments.GetDumpDebugInfoPath() / (HashName + TEXT(".lld"));
+				const FString PdbFile = Arguments.GetDumpDebugInfoPath() / PdbName;
 				SaveDxcBlobToFile(PdbBlob, PdbFile);
 			}
 		}
