UnrealBuildTool: Initial support for Clang 13
Hack to workaround clang-cl bug in LC_StringUtil.cpp is finally no longer necessary with Clang 13

Outstanding Issues:
Unresloved externals in ClangEditor (FRHI*)
error: ordered comparison of function pointers ('Electra::fastdelegate::DelegateMemento::GenericFuncPtr' (aka
    'void (*)()') and 'Electra::fastdelegate::DelegateMemento::GenericFuncPtr') [-Werror,-Wordered-compare-function-
    pointers]

#rnx

diff --git a/Engine/Source/Developer/Windows/LiveCoding/Private/External/LC_StringUtil.cpp b/Engine/Source/Developer/Windows/LiveCoding/Private/External/LC_StringUtil.cpp
--- a/Engine/Source/Developer/Windows/LiveCoding/Private/External/LC_StringUtil.cpp
+++ b/Engine/Source/Developer/Windows/LiveCoding/Private/External/LC_StringUtil.cpp
@@ -6,11 +6,8 @@
 
 #include "LC_StringUtil.h"
 // BEGIN EPIC MOD
-#if defined(__clang__) && defined(_MSC_VER) && _MSVC_LANG > 201402L
-#if __clang_major__ > 12
+#if defined(__clang__) && defined(_MSC_VER) && _MSVC_LANG > 201402L && __clang_major__ < 13
 // For reference: https://bugs.llvm.org/show_bug.cgi?id=41226#c16
-#error Re-check whether this workaround is still needed for clang 13.0 and higher.
-#endif
 #include <wchar.h>
 #endif
 
@@ -167,11 +164,8 @@ namespace detail
 	// Temporary replacement for std::wstring::find
 	static size_t WideFindOffset(const std::wstring& haystack, const std::wstring& needle)
 	{
-#if defined(__clang__) && defined(_MSC_VER) && _MSVC_LANG > 201402L
-#if __clang_major__ > 12
+#if defined(__clang__) && defined(_MSC_VER) && _MSVC_LANG > 201402L && __clang_major__ < 13
 		// For reference: https://bugs.llvm.org/show_bug.cgi?id=41226#c16
-		#error Re-check whether this workaround is still needed for clang 13.0 and higher.
-#endif
 		const wchar_t* found = ::wcsstr(haystack.c_str(), needle.c_str());
 		if (found)
 		{
diff --git a/Engine/Source/Programs/UnrealBuildTool/Platform/Android/AndroidToolChain.cs b/Engine/Source/Programs/UnrealBuildTool/Platform/Android/AndroidToolChain.cs
--- a/Engine/Source/Programs/UnrealBuildTool/Platform/Android/AndroidToolChain.cs
+++ b/Engine/Source/Programs/UnrealBuildTool/Platform/Android/AndroidToolChain.cs
@@ -605,6 +605,11 @@ namespace UnrealBuildTool
 			Result += " -fno-PIE";
 
 			Result += " -Wno-unused-variable";
+			if (CompilerVersionGreaterOrEqual(13, 0, 0))
+			{
+				Result += " -Wno-unused-but-set-variable";
+				Result += " -Wno-unused-but-set-parameter";
+			}
 			// this will hide the warnings about static functions in headers that aren't used in every single .cpp file
 			Result += " -Wno-unused-function";
 			// this hides the "enumeration value 'XXXXX' not handled in switch [-Wswitch]" warnings - we should maybe remove this at some point and add UE_LOG(, Fatal, ) to default cases
diff --git a/Engine/Source/Programs/UnrealBuildTool/Platform/Linux/LinuxToolChain.cs b/Engine/Source/Programs/UnrealBuildTool/Platform/Linux/LinuxToolChain.cs
--- a/Engine/Source/Programs/UnrealBuildTool/Platform/Linux/LinuxToolChain.cs
+++ b/Engine/Source/Programs/UnrealBuildTool/Platform/Linux/LinuxToolChain.cs
@@ -682,6 +682,11 @@ namespace UnrealBuildTool
 			}
 
 			Result += " -Wno-unused-variable";
+			if (CompilerVersionGreaterOrEqual(13, 0, 0))
+			{
+				Result += " -Wno-unused-but-set-variable";
+				Result += " -Wno-unused-but-set-parameter";
+			}
 			// this will hide the warnings about static functions in headers that aren't used in every single .cpp file
 			Result += " -Wno-unused-function";
 			// this hides the "enumeration value 'XXXXX' not handled in switch [-Wswitch]" warnings - we should maybe remove this at some point and add UE_LOG(, Fatal, ) to default cases
