Fix static analysis warning in IDxcResult::GetOutput function call.

#rnx

diff --git a/Engine/Source/Developer/Windows/ShaderFormatD3D/Private/D3DShaderCompilerDXC.cpp b/Engine/Source/Developer/Windows/ShaderFormatD3D/Private/D3DShaderCompilerDXC.cpp
--- a/Engine/Source/Developer/Windows/ShaderFormatD3D/Private/D3DShaderCompilerDXC.cpp
+++ b/Engine/Source/Developer/Windows/ShaderFormatD3D/Private/D3DShaderCompilerDXC.cpp
@@ -300,11 +300,13 @@ static HRESULT D3DCompileToDxil(const char* SourceText, FDxcArguments& Arguments
 	CompileResult->GetStatus(&CompileResultCode);
 	if (SUCCEEDED(CompileResultCode))
 	{
+		TRefCountPtr<IDxcBlobUtf16> ObjectCodeNameBlob; // Dummy name blob to silence static analysis warning
 		checkf(CompileResult->HasOutput(DXC_OUT_OBJECT), TEXT("No object code found!"));
-		VERIFYHRESULT(CompileResult->GetOutput(DXC_OUT_OBJECT, IID_PPV_ARGS(OutDxilBlob.GetInitReference()), nullptr));
+		VERIFYHRESULT(CompileResult->GetOutput(DXC_OUT_OBJECT, IID_PPV_ARGS(OutDxilBlob.GetInitReference()), ObjectCodeNameBlob.GetInitReference()));
 
+		TRefCountPtr<IDxcBlobUtf16> ReflectionNameBlob; // Dummy name blob to silence static analysis warning
 		checkf(CompileResult->HasOutput(DXC_OUT_REFLECTION), TEXT("No reflection found!"));
-		VERIFYHRESULT(CompileResult->GetOutput(DXC_OUT_REFLECTION, IID_PPV_ARGS(OutReflectionBlob.GetInitReference()), nullptr));
+		VERIFYHRESULT(CompileResult->GetOutput(DXC_OUT_REFLECTION, IID_PPV_ARGS(OutReflectionBlob.GetInitReference()), ReflectionNameBlob.GetInitReference()));
 
 		if (Arguments.ShouldDump())
 		{
