Status Analysis fix.

diff --git a/Engine/Source/Developer/Windows/ShaderFormatD3D/Private/D3DShaderCompilerDXC.cpp b/Engine/Source/Developer/Windows/ShaderFormatD3D/Private/D3DShaderCompilerDXC.cpp
--- a/Engine/Source/Developer/Windows/ShaderFormatD3D/Private/D3DShaderCompilerDXC.cpp
+++ b/Engine/Source/Developer/Windows/ShaderFormatD3D/Private/D3DShaderCompilerDXC.cpp
@@ -296,6 +296,11 @@ static HRESULT D3DCompileToDxil(const char* SourceText, FDxcArguments& Arguments
 	TRefCountPtr<IDxcResult> CompileResult;
 	VERIFYHRESULT(DXCCompileWrapper(Compiler, TextBlob, Arguments, CompileResult));
 
+	if (!CompileResult.IsValid())
+	{
+		return -1;
+	}
+
 	HRESULT CompileResultCode;
 	CompileResult->GetStatus(&CompileResultCode);
 	if (SUCCEEDED(CompileResultCode))
