Fix for HRESULT return value

diff --git a/Engine/Source/Developer/Windows/ShaderFormatD3D/Private/D3DShaderCompilerDXC.cpp b/Engine/Source/Developer/Windows/ShaderFormatD3D/Private/D3DShaderCompilerDXC.cpp
--- a/Engine/Source/Developer/Windows/ShaderFormatD3D/Private/D3DShaderCompilerDXC.cpp
+++ b/Engine/Source/Developer/Windows/ShaderFormatD3D/Private/D3DShaderCompilerDXC.cpp
@@ -298,7 +298,7 @@ static HRESULT D3DCompileToDxil(const char* SourceText, FDxcArguments& Arguments
 
 	if (!CompileResult.IsValid())
 	{
-		return -1;
+		return E_FAIL;
 	}
 
 	HRESULT CompileResultCode;
