Fixing PC SM6.6 support for RayTracing shaders.

diff --git a/Engine/Source/Developer/Windows/ShaderFormatD3D/Private/D3DShaderCompiler.cpp b/Engine/Source/Developer/Windows/ShaderFormatD3D/Private/D3DShaderCompiler.cpp
--- a/Engine/Source/Developer/Windows/ShaderFormatD3D/Private/D3DShaderCompiler.cpp
+++ b/Engine/Source/Developer/Windows/ShaderFormatD3D/Private/D3DShaderCompiler.cpp
@@ -108,8 +108,7 @@ static const TCHAR* GetShaderProfileName(FShaderTarget Target, bool bForceSM6)
 		case SF_RayMiss:
 		case SF_RayHitGroup:
 		case SF_RayCallable:
-			// return USE_SHADER_MODEL_6_6 ? TEXT("lib_6_6") : TEXT("lib_6_5");
-			return TEXT("lib_6_5"); // TODO: Intentionally using SM6.5 as a workaround for outstanding driver bug.
+			return USE_SHADER_MODEL_6_6 ? TEXT("lib_6_6") : TEXT("lib_6_5");
 		}
 	}
 	else if(Target.Platform == SP_PCD3D_SM5)
diff --git a/Engine/Source/Developer/Windows/ShaderFormatD3D/Private/ShaderFormatD3D.cpp b/Engine/Source/Developer/Windows/ShaderFormatD3D/Private/ShaderFormatD3D.cpp
--- a/Engine/Source/Developer/Windows/ShaderFormatD3D/Private/ShaderFormatD3D.cpp
+++ b/Engine/Source/Developer/Windows/ShaderFormatD3D/Private/ShaderFormatD3D.cpp
@@ -19,7 +19,7 @@ class FShaderFormatD3D : public IShaderFormat
 		UE_SHADER_PCD3D_SHARED_VER = 2,
 
 		/** Version for shader format, this becomes part of the DDC key. */
-		UE_SHADER_PCD3D_SM6_VER = UE_SHADER_PCD3D_SHARED_VER + 2,
+		UE_SHADER_PCD3D_SM6_VER = UE_SHADER_PCD3D_SHARED_VER + 3,
 		UE_SHADER_PCD3D_SM5_VER = UE_SHADER_PCD3D_SHARED_VER + 9,
 		UE_SHADER_PCD3D_ES3_1_VER = UE_SHADER_PCD3D_SHARED_VER + 8,
 	};
