Hololens: Dedicated Shader Platform

We can't share DDPI infrastructure with PC D3D ES3.1 because of conflicts from sharing. We need a separate shader platform.
For now, we'll just 'intrusively' plumb the new Hololens shader platform. In the future, we want to move as much code out into
Hololens-specific platform code to reduce conflict with shared paths that don't need to know about Hololens.

diff --git a/Engine/Config/HoloLens/DataDrivenPlatformInfo.ini b/Engine/Config/HoloLens/DataDrivenPlatformInfo.ini
--- a/Engine/Config/HoloLens/DataDrivenPlatformInfo.ini
+++ b/Engine/Config/HoloLens/DataDrivenPlatformInfo.ini
@@ -24,3 +24,10 @@ bUsesHostCompiler=true
 bCanUseCrashReporter=false
 bUATClosesAfterLaunch=false
 PlatformGroupName=Mobile
+
+[ShaderPlatform D3D_ES3_1_HOLOLENS]
+Language=D3D
+MaxFeatureLevel=ES3_1
+bIsMobile=true
+bSupportsInstancedStereo=true
+bSupportsManualVertexFetch=false
diff --git a/Engine/Plugins/Runtime/AR/Microsoft/HoloLensAR/Source/HoloLensTargetPlatform/Private/HoloLensTargetPlatform.cpp b/Engine/Plugins/Runtime/AR/Microsoft/HoloLensAR/Source/HoloLensTargetPlatform/Private/HoloLensTargetPlatform.cpp
--- a/Engine/Plugins/Runtime/AR/Microsoft/HoloLensAR/Source/HoloLensTargetPlatform/Private/HoloLensTargetPlatform.cpp
+++ b/Engine/Plugins/Runtime/AR/Microsoft/HoloLensAR/Source/HoloLensTargetPlatform/Private/HoloLensTargetPlatform.cpp
@@ -96,18 +96,20 @@ bool FHoloLensTargetPlatform::SupportsFeature(ETargetPlatformFeatures Feature) c
 	}
 }
 
-static FName NAME_PCD3D_ES3_1(TEXT("PCD3D_ES31"));
+// Hololens previously used PCD3D_ES3_1 as a shader format/platform until we
+// had to split to properly support DataDrivenPlatformInfo
+static FName NAME_D3D_ES3_1_HOLOLENS(TEXT("D3D_ES3_1_HOLOLENS"));
 static FName NAME_PCD3D_SM5(TEXT("PCD3D_SM5"));
 
 void FHoloLensTargetPlatform::GetAllPossibleShaderFormats(TArray<FName>& OutFormats) const
 {
-	OutFormats.AddUnique(NAME_PCD3D_ES3_1);
+	OutFormats.AddUnique(NAME_D3D_ES3_1_HOLOLENS);
 	OutFormats.AddUnique(NAME_PCD3D_SM5);
 }
 
 void FHoloLensTargetPlatform::GetAllTargetedShaderFormats(TArray<FName>& OutFormats) const
 {
-	OutFormats.AddUnique(NAME_PCD3D_ES3_1);
+	OutFormats.AddUnique(NAME_D3D_ES3_1_HOLOLENS);
 	OutFormats.AddUnique(NAME_PCD3D_SM5);
 }
 
diff --git a/Engine/Plugins/Runtime/WindowsMixedReality/Source/WindowsMixedRealityRHI/Private/WindowsMixedRealityRHI.cpp b/Engine/Plugins/Runtime/WindowsMixedReality/Source/WindowsMixedRealityRHI/Private/WindowsMixedRealityRHI.cpp
--- a/Engine/Plugins/Runtime/WindowsMixedReality/Source/WindowsMixedRealityRHI/Private/WindowsMixedRealityRHI.cpp
+++ b/Engine/Plugins/Runtime/WindowsMixedReality/Source/WindowsMixedRealityRHI/Private/WindowsMixedRealityRHI.cpp
@@ -146,11 +146,15 @@ FDynamicRHI* FWindowsMixedRealityRHIModule::CreateRHI(ERHIFeatureLevel::Type Req
 {
 #if PLATFORM_HOLOLENS
 	GMaxRHIFeatureLevel = ERHIFeatureLevel::ES3_1;
-	GMaxRHIShaderPlatform = SP_PCD3D_ES3_1;
+	GMaxRHIShaderPlatform = SP_D3D_ES3_1_HOLOLENS;
 #endif
 
 	GShaderPlatformForFeatureLevel[ERHIFeatureLevel::ES2_REMOVED] = SP_NumPlatforms;
+#if PLATFORM_HOLOLENS
+	GShaderPlatformForFeatureLevel[ERHIFeatureLevel::ES3_1] = SP_D3D_ES3_1_HOLOLENS;
+#else
 	GShaderPlatformForFeatureLevel[ERHIFeatureLevel::ES3_1] = SP_PCD3D_ES3_1;
+#endif
 	GShaderPlatformForFeatureLevel[ERHIFeatureLevel::SM4_REMOVED] = SP_NumPlatforms;
 	GShaderPlatformForFeatureLevel[ERHIFeatureLevel::SM5] = SP_PCD3D_SM5;
 
diff --git a/Engine/Source/Developer/SharedSettingsWidgets/Private/ShaderFormatsPropertyDetails.cpp b/Engine/Source/Developer/SharedSettingsWidgets/Private/ShaderFormatsPropertyDetails.cpp
--- a/Engine/Source/Developer/SharedSettingsWidgets/Private/ShaderFormatsPropertyDetails.cpp
+++ b/Engine/Source/Developer/SharedSettingsWidgets/Private/ShaderFormatsPropertyDetails.cpp
@@ -69,7 +69,10 @@ FText FShaderFormatsPropertyDetails::GetFriendlyNameFromRHINameMac(const FString
 		case SP_VULKAN_PCES3_1:
 		case SP_VULKAN_ES3_1_ANDROID:
 			FriendlyRHIName = LOCTEXT("VulkanES31", "Vulkan (ES 3.1)");
-			break;	
+			break;
+		case SP_D3D_ES3_1_HOLOLENS:
+			FriendlyRHIName = LOCTEXT("D3DES31HL", "Direct3D (ES3.1, Hololens)");
+			break;
 		default:
 			break;
 	}
diff --git a/Engine/Source/Developer/Windows/ShaderFormatD3D/Private/D3DShaderCompiler.cpp b/Engine/Source/Developer/Windows/ShaderFormatD3D/Private/D3DShaderCompiler.cpp
--- a/Engine/Source/Developer/Windows/ShaderFormatD3D/Private/D3DShaderCompiler.cpp
+++ b/Engine/Source/Developer/Windows/ShaderFormatD3D/Private/D3DShaderCompiler.cpp
@@ -134,11 +134,11 @@ static const TCHAR* GetShaderProfileName(FShaderTarget Target, bool bForceSM6)
 			return TEXT("lib_6_3");
 		}
 	}
-	else if (Target.Platform == SP_PCD3D_ES3_1)
+	else if ((Target.Platform == SP_PCD3D_ES3_1) || (Target.Platform == SP_D3D_ES3_1_HOLOLENS))
 	{
 		checkSlow(Target.Frequency == SF_Vertex ||
 			Target.Frequency == SF_Pixel ||
-			Target.Frequency == SF_Geometry || 
+			Target.Frequency == SF_Geometry ||
 			Target.Frequency == SF_Compute);
 
 		//set defines and profiles for the appropriate shader paths
@@ -1304,7 +1304,7 @@ void CompileShader_Windows(const FShaderCompilerInput& Input,FShaderCompilerOutp
 	}
 	else if (Language == ELanguage::ES3_1)
 	{
-		check(Input.Target.Platform == SP_PCD3D_ES3_1);
+		check((Input.Target.Platform == SP_PCD3D_ES3_1) || (Input.Target.Platform == SP_D3D_ES3_1_HOLOLENS));
 		AdditionalDefines.SetDefine(TEXT("ES3_1_PROFILE"), 1);
 	}
 	else
diff --git a/Engine/Source/Developer/Windows/ShaderFormatD3D/Private/ShaderFormatD3D.cpp b/Engine/Source/Developer/Windows/ShaderFormatD3D/Private/ShaderFormatD3D.cpp
--- a/Engine/Source/Developer/Windows/ShaderFormatD3D/Private/ShaderFormatD3D.cpp
+++ b/Engine/Source/Developer/Windows/ShaderFormatD3D/Private/ShaderFormatD3D.cpp
@@ -11,6 +11,7 @@
 static FName NAME_PCD3D_SM6(TEXT("PCD3D_SM6"));
 static FName NAME_PCD3D_SM5(TEXT("PCD3D_SM5"));
 static FName NAME_PCD3D_ES3_1(TEXT("PCD3D_ES31"));
+static FName NAME_D3D_ES3_1_HOLOLENS(TEXT("D3D_ES3_1_HOLOLENS"));
 
 class FShaderFormatD3D : public IShaderFormat
 {
@@ -22,11 +23,12 @@ class FShaderFormatD3D : public IShaderFormat
 		UE_SHADER_PCD3D_SM6_VER = UE_SHADER_PCD3D_SHARED_VER + 3,
 		UE_SHADER_PCD3D_SM5_VER = UE_SHADER_PCD3D_SHARED_VER + 9,
 		UE_SHADER_PCD3D_ES3_1_VER = UE_SHADER_PCD3D_SHARED_VER + 8,
+		UE_SHADER_D3D_ES3_1_HOLOLENS_VER = UE_SHADER_PCD3D_ES3_1_VER,
 	};
 
 	void CheckFormat(FName Format) const
 	{
-		check(Format == NAME_PCD3D_SM6 || Format == NAME_PCD3D_SM5 || Format == NAME_PCD3D_ES3_1);
+		check(Format == NAME_PCD3D_SM6 || Format == NAME_PCD3D_SM5 || Format == NAME_PCD3D_ES3_1 || Format == NAME_D3D_ES3_1_HOLOLENS);
 	}
 
 	uint32 DxcVersionHash = 0;
@@ -62,6 +64,11 @@ public:
 			// Shader DXC signature is intentionally not included, as ES3_1 target always uses legacy compiler.
 			return UE_SHADER_PCD3D_ES3_1_VER;
 		}
+		else if (Format == NAME_D3D_ES3_1_HOLOLENS)
+		{
+			// Shader DXC signature is intentionally not included, as ES3_1 target always uses legacy compiler.
+			return UE_SHADER_D3D_ES3_1_HOLOLENS_VER;
+		}
 		checkf(0, TEXT("Unknown Format %s"), *Format.ToString());
 		return 0;
 	}
@@ -71,6 +78,7 @@ public:
 		OutFormats.Add(NAME_PCD3D_SM6);
 		OutFormats.Add(NAME_PCD3D_SM5);
 		OutFormats.Add(NAME_PCD3D_ES3_1);
+		OutFormats.Add(NAME_D3D_ES3_1_HOLOLENS);
 	}
 
 	virtual void CompileShader(FName Format, const struct FShaderCompilerInput& Input, struct FShaderCompilerOutput& Output,const FString& WorkingDirectory) const
@@ -88,6 +96,10 @@ public:
 		{
 			CompileShader_Windows(Input, Output, WorkingDirectory, ELanguage::ES3_1);
 		}
+		else if (Format == NAME_D3D_ES3_1_HOLOLENS)
+		{
+			CompileShader_Windows(Input, Output, WorkingDirectory, ELanguage::ES3_1);
+		}
 		else
 		{
 			checkf(0, TEXT("Unknown format %s"), *Format.ToString());
@@ -143,6 +155,13 @@ public:
 			Input.Environment.SetDefine(TEXT("__SHADER_TARGET_MAJOR"), 5);
 			Input.Environment.SetDefine(TEXT("__SHADER_TARGET_MINOR"), 0);
 		}
+		else if (Input.ShaderFormat == NAME_D3D_ES3_1_HOLOLENS)
+		{
+			Input.Environment.SetDefine(TEXT("ES3_1_PROFILE"), 1);
+			Input.Environment.SetDefine(TEXT("COMPILER_DXC"), 0);
+			Input.Environment.SetDefine(TEXT("__SHADER_TARGET_MAJOR"), 5);
+			Input.Environment.SetDefine(TEXT("__SHADER_TARGET_MINOR"), 0);
+		}
 		else
 		{
 			checkf(0, TEXT("Unknown format %s"), *Input.ShaderFormat.ToString());
diff --git a/Engine/Source/Editor/DetailCustomizations/Private/WindowsTargetSettingsDetails.cpp b/Engine/Source/Editor/DetailCustomizations/Private/WindowsTargetSettingsDetails.cpp
--- a/Engine/Source/Editor/DetailCustomizations/Private/WindowsTargetSettingsDetails.cpp
+++ b/Engine/Source/Editor/DetailCustomizations/Private/WindowsTargetSettingsDetails.cpp
@@ -64,6 +64,10 @@ static FText GetFriendlyNameFromWindowsRHIName(const FString& InRHIName)
 	{
 		FriendlyRHIName = LOCTEXT("DirectXES31", "DirectX Mobile Emulation (ES3.1)");
 	}
+	else if (InRHIName == TEXT("D3D_ES3_1_HOLOLENS"))
+	{
+		FriendlyRHIName = LOCTEXT("DirectXES31HL", "DirectX Hololens (ES3.1)");
+	}
 	else if (InRHIName == TEXT("SF_VULKAN_SM5"))
 	{
 		FriendlyRHIName = LOCTEXT("VulkanSM5", "Vulkan (SM5)");
diff --git a/Engine/Source/Editor/MaterialEditor/Private/MaterialStats.cpp b/Engine/Source/Editor/MaterialEditor/Private/MaterialStats.cpp
--- a/Engine/Source/Editor/MaterialEditor/Private/MaterialStats.cpp
+++ b/Engine/Source/Editor/MaterialEditor/Private/MaterialStats.cpp
@@ -444,6 +444,9 @@ void FMaterialStats::BuildShaderPlatformDB()
 	AddShaderPlatform(EPlatformCategoryType::Desktop, SP_PCD3D_SM5, TEXT("DirectX SM5"), true, true, TEXT("Desktop, DirectX, Shader Model 5"));
 	AddShaderPlatform(EPlatformCategoryType::Desktop, SP_PCD3D_SM6, TEXT("DirectX SM6"), true, true, TEXT("Desktop, DirectX, Shader Model 6"));
 	AddShaderPlatform(EPlatformCategoryType::Desktop, SP_PCD3D_ES3_1, TEXT("DirectX ES 3.1"), true, true, TEXT("Desktop, DirectX, ES 3.1"));
+
+	// Desktop is the closest fit for Hololens (for now)
+	AddShaderPlatform(EPlatformCategoryType::Desktop, SP_D3D_ES3_1_HOLOLENS, TEXT("DirectX ES 3.1 Hololens"), true, true, TEXT("Hololens, DirectX, ES 3.1"));
 #endif
 
 	// Vulkan
diff --git a/Engine/Source/Runtime/RHI/Private/HoloLens/HoloLensDynamicRHI.cpp b/Engine/Source/Runtime/RHI/Private/HoloLens/HoloLensDynamicRHI.cpp
--- a/Engine/Source/Runtime/RHI/Private/HoloLens/HoloLensDynamicRHI.cpp
+++ b/Engine/Source/Runtime/RHI/Private/HoloLens/HoloLensDynamicRHI.cpp
@@ -113,7 +113,7 @@ FDynamicRHI* PlatformCreateDynamicRHI()
 	}
 
 	GMaxRHIFeatureLevel = ERHIFeatureLevel::ES3_1;
-	GMaxRHIShaderPlatform = SP_PCD3D_ES3_1;
+	GMaxRHIShaderPlatform = SP_D3D_ES3_1_HOLOLENS;
 
 	return DynamicRHI;
 }
diff --git a/Engine/Source/Runtime/RHI/Private/RHI.cpp b/Engine/Source/Runtime/RHI/Private/RHI.cpp
--- a/Engine/Source/Runtime/RHI/Private/RHI.cpp
+++ b/Engine/Source/Runtime/RHI/Private/RHI.cpp
@@ -1694,6 +1694,7 @@ FName ShaderPlatformToPlatformName(EShaderPlatform Platform)
 	case SP_OPENGL_PCES3_1:
 	case SP_VULKAN_PCES3_1:
 	case SP_VULKAN_SM5:
+	case SP_D3D_ES3_1_HOLOLENS:
 		return NAME_PLATFORM_WINDOWS;
 	case SP_VULKAN_ES3_1_ANDROID:
 	case SP_VULKAN_SM5_ANDROID:
@@ -2132,6 +2133,7 @@ FString LexToString(EShaderPlatform Platform, bool bError)
 	case SP_VULKAN_PCES3_1: return TEXT("VULKAN_PCES3_1");
 	case SP_VULKAN_SM5: return TEXT("VULKAN_SM5");
 	case SP_VULKAN_SM5_ANDROID: return TEXT("VULKAN_SM5_ANDROID");
+	case SP_D3D_ES3_1_HOLOLENS: return TEXT("D3D_ES3_1_HOLOLENS");
 
 	default:
 		if (FStaticShaderPlatformNames::IsStaticPlatform(Platform))
diff --git a/Engine/Source/Runtime/RHI/Public/RHIDefinitions.h b/Engine/Source/Runtime/RHI/Public/RHIDefinitions.h
--- a/Engine/Source/Runtime/RHI/Public/RHIDefinitions.h
+++ b/Engine/Source/Runtime/RHI/Public/RHIDefinitions.h
@@ -90,6 +90,7 @@ enum EShaderPlatform
 	//---------------------------------------------------------------------------------
 	SP_VULKAN_SM5_ANDROID			= SP_StaticPlatform_Last+1,
 	SP_PCD3D_SM6,
+	SP_D3D_ES3_1_HOLOLENS,
 
 	SP_NumPlatforms,
 	SP_NumBits						= 7,
diff --git a/Engine/Source/Runtime/RHI/Public/RHIShaderFormatDefinitions.inl b/Engine/Source/Runtime/RHI/Public/RHIShaderFormatDefinitions.inl
--- a/Engine/Source/Runtime/RHI/Public/RHIShaderFormatDefinitions.inl
+++ b/Engine/Source/Runtime/RHI/Public/RHIShaderFormatDefinitions.inl
@@ -30,6 +30,7 @@ static FName NAME_VULKAN_ES3_1(TEXT("SF_VULKAN_ES31"));
 static FName NAME_VULKAN_SM5(TEXT("SF_VULKAN_SM5"));
 static FName NAME_VULKAN_SM5_ANDROID(TEXT("SF_VULKAN_SM5_ANDROID"));
 
+static FName NAME_D3D_ES3_1_HOLOLENS(TEXT("D3D_ES3_1_HOLOLENS"));
 
 static FName ShaderPlatformToShaderFormatName(EShaderPlatform Platform)
 {
@@ -69,6 +70,8 @@ static FName ShaderPlatformToShaderFormatName(EShaderPlatform Platform)
 		return NAME_VULKAN_SM5;
 	case SP_VULKAN_SM5_ANDROID:
 		return NAME_VULKAN_SM5_ANDROID;
+	case SP_D3D_ES3_1_HOLOLENS:
+		return NAME_D3D_ES3_1_HOLOLENS;
 
 	default:
 		if (FStaticShaderPlatformNames::IsStaticPlatform(Platform))
@@ -105,6 +108,8 @@ static EShaderPlatform ShaderFormatNameToShaderPlatform(FName ShaderFormat)
 	if (ShaderFormat == NAME_VULKAN_SM5)				return SP_VULKAN_SM5;
 	if (ShaderFormat == NAME_VULKAN_SM5_ANDROID)		return SP_VULKAN_SM5_ANDROID;
 
+	if (ShaderFormat == NAME_D3D_ES3_1_HOLOLENS)		return SP_D3D_ES3_1_HOLOLENS;
+
 	for (int32 StaticPlatform = SP_StaticPlatform_First; StaticPlatform <= SP_StaticPlatform_Last; ++StaticPlatform)
 	{
 		if (ShaderFormat == FStaticShaderPlatformNames::Get().GetShaderFormat(EShaderPlatform(StaticPlatform)))
diff --git a/Engine/Source/Runtime/RHI/Public/RHIShaderPlatformDefinitions.inl b/Engine/Source/Runtime/RHI/Public/RHIShaderPlatformDefinitions.inl
--- a/Engine/Source/Runtime/RHI/Public/RHIShaderPlatformDefinitions.inl
+++ b/Engine/Source/Runtime/RHI/Public/RHIShaderPlatformDefinitions.inl
@@ -30,6 +30,7 @@ static FText GetFriendlyShaderPlatformName(const EShaderPlatform InShaderPlatfor
 	case SP_PCD3D_ES3_1:
 	case SP_VULKAN_PCES3_1:
 	case SP_OPENGL_PCES3_1:
+	case SP_D3D_ES3_1_HOLOLENS:
 	{
 		static const FText Description = NSLOCTEXT("FriendlyShaderPlatformNames", "Generic_ES31_loc", "ES31");
 		return Description;
diff --git a/Engine/Source/Runtime/Renderer/Private/PostProcess/PostProcessMaterial.cpp b/Engine/Source/Runtime/Renderer/Private/PostProcess/PostProcessMaterial.cpp
--- a/Engine/Source/Runtime/Renderer/Private/PostProcess/PostProcessMaterial.cpp
+++ b/Engine/Source/Runtime/Renderer/Private/PostProcess/PostProcessMaterial.cpp
@@ -436,7 +436,8 @@ FScreenPassTexture AddPostProcessMaterialPass(
 	// We can re-use the scene color texture as the render target if we're not simultaneously reading from it.
 	// This is only necessary to do if we're going to be priming content from the render target since it avoids
 	// the copy. Otherwise, we just allocate a new render target.
-	if (!Output.IsValid() && !MaterialShaderMap->UsesSceneTexture(PPI_PostProcessInput0) && bPrimeOutputColor && !bForceIntermediateTarget && Inputs.bAllowSceneColorInputAsOutput && GMaxRHIShaderPlatform != SP_PCD3D_ES3_1)
+	const bool bValidShaderPlatform = (GMaxRHIShaderPlatform != SP_PCD3D_ES3_1) && (GMaxRHIShaderPlatform != SP_D3D_ES3_1_HOLOLENS); // This might actually work with Hololens GPUs, but we haven't enabled it before
+	if (!Output.IsValid() && !MaterialShaderMap->UsesSceneTexture(PPI_PostProcessInput0) && bPrimeOutputColor && !bForceIntermediateTarget && Inputs.bAllowSceneColorInputAsOutput && bValidShaderPlatform)
 	{
 		Output = FScreenPassRenderTarget(SceneColor, ERenderTargetLoadAction::ELoad);
 	}
diff --git a/Engine/Source/Runtime/Windows/D3D11RHI/Private/Windows/WindowsD3D11Device.cpp b/Engine/Source/Runtime/Windows/D3D11RHI/Private/Windows/WindowsD3D11Device.cpp
--- a/Engine/Source/Runtime/Windows/D3D11RHI/Private/Windows/WindowsD3D11Device.cpp
+++ b/Engine/Source/Runtime/Windows/D3D11RHI/Private/Windows/WindowsD3D11Device.cpp
@@ -1061,13 +1061,17 @@ FDynamicRHI* FD3D11DynamicRHIModule::CreateRHI(ERHIFeatureLevel::Type RequestedF
 {
 #if PLATFORM_HOLOLENS
 	GMaxRHIFeatureLevel = ERHIFeatureLevel::ES3_1;
-	GMaxRHIShaderPlatform = SP_PCD3D_ES3_1;
+	GMaxRHIShaderPlatform = SP_D3D_ES3_1_HOLOLENS;
 #endif
 
 	IDXGIFactory1* DXGIFactory1;
 	VERIFYD3D11RESULT(ChosenAdapter.DXGIAdapter->GetParent(__uuidof(DXGIFactory1), reinterpret_cast<void**>(&DXGIFactory1)));
 
+#if PLATFORM_HOLOLENS
+	GShaderPlatformForFeatureLevel[ERHIFeatureLevel::ES3_1] = SP_D3D_ES3_1_HOLOLENS;
+#else
 	GShaderPlatformForFeatureLevel[ERHIFeatureLevel::ES3_1] = SP_PCD3D_ES3_1;
+#endif
 	GShaderPlatformForFeatureLevel[ERHIFeatureLevel::SM5] = SP_PCD3D_SM5;
 
 	ERHIFeatureLevel::Type PreviewFeatureLevel;
