Windows target platform can be extended by Windows-based platform extensions.

#rnx

diff --git a/Engine/Source/Developer/Windows/WindowsTargetPlatform/Private/GenericWindowsTargetPlatform.h b/Engine/Source/Developer/Windows/WindowsTargetPlatform/Private/GenericWindowsTargetPlatform.h
--- a/Engine/Source/Developer/Windows/WindowsTargetPlatform/Private/GenericWindowsTargetPlatform.h
+++ b/Engine/Source/Developer/Windows/WindowsTargetPlatform/Private/GenericWindowsTargetPlatform.h
@@ -26,7 +26,12 @@
 /**
  * Template for Windows target platforms
  */
+#if PLATFORM_WINDOWS
+template<typename TProperties, typename TTargetDevice = TLocalPcTargetDevice<PLATFORM_64BITS>>
+#else
 template<typename TProperties>
+#endif
+
 class TGenericWindowsTargetPlatform : public TTargetPlatformBase<TProperties>
 {
 public:
@@ -39,7 +44,7 @@ public:
 	{
 #if PLATFORM_WINDOWS
 		// only add local device if actually running on Windows
-		LocalDevice = MakeShareable(new TLocalPcTargetDevice<PLATFORM_64BITS>(*this));
+		LocalDevice = MakeShareable(new TTargetDevice(*this));
 #endif
 
 	#if WITH_ENGINE
diff --git a/Engine/Source/Developer/Windows/WindowsTargetPlatform/Private/LocalPcTargetDevice.h b/Engine/Source/Developer/Windows/WindowsTargetPlatform/Private/LocalPcTargetDevice.h
--- a/Engine/Source/Developer/Windows/WindowsTargetPlatform/Private/LocalPcTargetDevice.h
+++ b/Engine/Source/Developer/Windows/WindowsTargetPlatform/Private/LocalPcTargetDevice.h
@@ -6,6 +6,8 @@
 #include "Misc/Guid.h"
 #include "HAL/FileManager.h"
 #include "HAL/PlatformProcess.h"
+#include "Interfaces/ITargetDevice.h"
+#include "Interfaces/ITargetPlatform.h"
 
 #include "Windows/WindowsHWrapper.h"
 #include "Windows/AllowWindowsPlatformTypes.h"
@@ -319,7 +321,7 @@ protected:
 		return (GetLastError() == ERROR_SUCCESS);
 	}
 
-private:
+protected:
 
 	// Holds the collection of processes that were started using the Run() method.
 	TMap<FGuid, void*> Processes;
