Use DXC re-compilation when FXC crashed or ran out of memory.

#rnx

diff --git a/Engine/Source/Developer/Windows/ShaderFormatD3D/Private/D3DShaderCompiler.cpp b/Engine/Source/Developer/Windows/ShaderFormatD3D/Private/D3DShaderCompiler.cpp
--- a/Engine/Source/Developer/Windows/ShaderFormatD3D/Private/D3DShaderCompiler.cpp
+++ b/Engine/Source/Developer/Windows/ShaderFormatD3D/Private/D3DShaderCompiler.cpp
@@ -640,7 +640,7 @@ bool CompileAndProcessD3DShaderFXC(FString& PreprocessedShaderSource, const FStr
 		);
 
 		// Some materials give FXC a hard time to optimize and the compiler fails with an internal error.
-		if (Result == HRESULT_FROM_WIN32(ERROR_ARITHMETIC_OVERFLOW))
+		if (Result == HRESULT_FROM_WIN32(ERROR_ARITHMETIC_OVERFLOW) || Result == E_OUTOFMEMORY || bException)
 		{
 			CrossCompiler::FShaderConductorContext CompilerContext;
 
