enable BC6+7 texture support in games that use the ES3.1 RHI
ES3.1 is used as a fast renderer mode, but still supports DX11 textures
log a message when BC6+7 is encoded as uncompressed
delete the unused GetVirtualTextureLayerFormat

diff --git a/Engine/Source/Developer/Mac/MacTargetPlatform/Private/GenericMacTargetPlatform.h b/Engine/Source/Developer/Mac/MacTargetPlatform/Private/GenericMacTargetPlatform.h
--- a/Engine/Source/Developer/Mac/MacTargetPlatform/Private/GenericMacTargetPlatform.h
+++ b/Engine/Source/Developer/Mac/MacTargetPlatform/Private/GenericMacTargetPlatform.h
@@ -170,136 +170,6 @@ public:
 		}
 	}
 
-	virtual FName GetVirtualTextureLayerFormat(int32 SourceFormat, bool bAllowCompression, bool bNoAlpha, bool bSupportDX11TextureFormats, int32 Settings) const override
-	{
-		FName TextureFormatName = NAME_None;
-
-		// Supported texture format names.
-		static FName NameDXT1(TEXT("DXT1"));
-		static FName NameDXT3(TEXT("DXT3"));
-		static FName NameDXT5(TEXT("DXT5"));
-		static FName NameDXT5n(TEXT("DXT5n"));
-		static FName NameBC4(TEXT("BC4"));
-		static FName NameBC5(TEXT("BC5"));
-		static FName NameBGRA8(TEXT("BGRA8"));
-		static FName NameXGXR8(TEXT("XGXR8"));
-		static FName NameG8(TEXT("G8"));
-		static FName NameG16(TEXT("G16"));
-		static FName NameVU8(TEXT("VU8"));
-		static FName NameRGBA16F(TEXT("RGBA16F"));
-		static FName NameR16F(TEXT("R16F"));
-		static FName NameBC6H(TEXT("BC6H"));
-		static FName NameBC7(TEXT("BC7"));
-
-		// Note: We can't use things here like autoDXT here which defer the exact choice to the compressor as it would mean that
-		// some textures on a VT layer may get a different format than others.
-		// We need to guarantee the format to be the same for all textures on the layer so we need to decide on the exact final format here.
-
-		bool bUseDXT5NormalMap = false;
-		FString UseDXT5NormalMapsString;
-		if (this->GetConfigSystem()->GetString(TEXT("SystemSettings"), TEXT("Compat.UseDXT5NormalMaps"), UseDXT5NormalMapsString, GEngineIni))
-		{
-			bUseDXT5NormalMap = FCString::ToBool(*UseDXT5NormalMapsString);
-		}
-
-		// Determine the pixel format of the (un/)compressed texture
-		if (!bAllowCompression)
-		{
-			if (SourceFormat == TSF_RGBA16F)
-			{
-				TextureFormatName = NameRGBA16F;
-			}
-			else if (SourceFormat == TSF_G16)
-			{
-				TextureFormatName = NameG16;
-			}
-			else if (SourceFormat == TSF_G8 || Settings == TC_Grayscale)
-			{
-				TextureFormatName = NameG8;
-			}
-			else if (Settings == TC_Normalmap && bUseDXT5NormalMap)
-			{
-				TextureFormatName = NameXGXR8;
-			}
-			else
-			{
-				TextureFormatName = NameBGRA8;
-			}
-		}
-		else if (Settings == TC_HDR)
-		{
-			TextureFormatName = NameRGBA16F;
-		}
-		else if (Settings == TC_Normalmap)
-		{
-			TextureFormatName = bUseDXT5NormalMap ? NameDXT5n : NameBC5;
-		}
-		else if (Settings == TC_Displacementmap)
-		{
-			TextureFormatName = NameG8;
-		}
-		else if (Settings == TC_VectorDisplacementmap)
-		{
-			TextureFormatName = NameBGRA8;
-		}
-		else if (Settings == TC_Grayscale)
-		{
-			TextureFormatName = NameG8;
-		}
-		else if (Settings == TC_Alpha)
-		{
-			TextureFormatName = NameBC4;
-		}
-		else if (Settings == TC_DistanceFieldFont)
-		{
-			TextureFormatName = NameG8;
-		}
-		else if (Settings == TC_HDR_Compressed)
-		{
-			TextureFormatName = NameBC6H;
-		}
-		else if (Settings == TC_BC7)
-		{
-			TextureFormatName = NameBC7;
-		}
-		else if (Settings == TC_HalfFloat)
-		{
-			TextureFormatName = NameR16F;
-		}
-		else if (bNoAlpha)
-		{
-			TextureFormatName = NameDXT1;
-		}
-		else
-		{
-			TextureFormatName = NameDXT5;
-		}
-
-		/*
-		FIXME: IS this still relevant for VT ? comes from the texture variant
-		// Some PC GPUs don't support sRGB read from G8 textures (e.g. AMD DX10 cards on ShaderModel3.0)
-		// This solution requires 4x more memory but a lot of PC HW emulate the format anyway
-		if ((TextureFormatName == NameG8) && Texture->SRGB && !SupportsFeature(ETargetPlatformFeatures::GrayscaleSRGB))
-		{
-		TextureFormatName = NameBGRA8;
-		}*/
-
-		// fallback to non-DX11 formats if one was chosen, but we can't use it
-		if (!bSupportDX11TextureFormats)
-		{
-			if (TextureFormatName == NameBC6H)
-			{
-				TextureFormatName = NameRGBA16F;
-			}
-			else if (TextureFormatName == NameBC7)
-			{
-				TextureFormatName = NameDXT5;
-			}
-		}
-
-		return TextureFormatName;
-	}
-
 	virtual bool SupportsLQCompressionTextureFormat() const override { return false; };
 
 	virtual const UTextureLODSettings& GetTextureLODSettings() const override
diff --git a/Engine/Source/Developer/TargetPlatform/Public/Common/TargetPlatformBase.h b/Engine/Source/Developer/TargetPlatform/Public/Common/TargetPlatformBase.h
--- a/Engine/Source/Developer/TargetPlatform/Public/Common/TargetPlatformBase.h
+++ b/Engine/Source/Developer/TargetPlatform/Public/Common/TargetPlatformBase.h
@@ -88,14 +88,6 @@ public:
 		return Format;
 	}
 
-	virtual FName GetVirtualTextureLayerFormat(
-		int32 SourceFormat,
-		bool bAllowCompression, bool bNoAlpha,
-		bool bSupportDX11TextureFormats, int32 Settings) const override
-	{
-		return FName();
-	}
-
 	virtual bool SupportsLQCompressionTextureFormat() const override { return true; };
 
 #endif //WITH_ENGINE
diff --git a/Engine/Source/Developer/TargetPlatform/Public/Interfaces/ITargetPlatform.h b/Engine/Source/Developer/TargetPlatform/Public/Interfaces/ITargetPlatform.h
--- a/Engine/Source/Developer/TargetPlatform/Public/Interfaces/ITargetPlatform.h
+++ b/Engine/Source/Developer/TargetPlatform/Public/Interfaces/ITargetPlatform.h
@@ -573,26 +573,7 @@ public:
 	 * might want to use a single specific variant for virtual textures to reduce fragmentation
 	 */
 	virtual FName FinalizeVirtualTextureLayerFormat(FName Format) const = 0;
-
-	/**
-	* Gets the texture format to use for a virtual texturing layer. In order to make a better guess
-	* some parameters are passed to this function.
-	*
-	* @param SourceFormat The raw uncompressed source format (ETextureSourceFormat) the texture is stored in.
-	* @param bAllowCompression Allow a compressed (lossy) format to be chosen.
-	* @param bNoAlpha The chosen format doesn't need to have an alpha channel.
-	* @param bSupportDX11TextureFormats Allow choosing a texture format which is supported only on DX11 hardware (e.g. BC7 etc).
-	* @param settings A hint as to what kind of data is to be stored in the layer or an explicit request for a certain format (TextureCompressionSettings enum value).
-	* @return The chosen format to use for the layer on this platform based on the given input. May return the None name ( Fname::IsNone == true) on platforms that do not support VT.
-	* 
-	* FIXME: Is it better to pass in the UVirtualTexture and layer index instead of all these arguments?! Less encapsulated but cleaner??
-	* this would also mean we don't have to hide the SourceFormat and TextureCompressionSettings enums
-	*/
-	virtual FName GetVirtualTextureLayerFormat(
-		int32 SourceFormat,
-		bool bAllowCompression, bool bNoAlpha,
-		bool bSupportDX11TextureFormats, int32 TextureCompressionSettings) const = 0;
-	
+		
 	//whether R5G6B5 and B5G5R5A1 is supported
 	virtual bool SupportsLQCompressionTextureFormat() const = 0;
 	
diff --git a/Engine/Source/Developer/Windows/WindowsTargetPlatform/Private/GenericWindowsTargetPlatform.h b/Engine/Source/Developer/Windows/WindowsTargetPlatform/Private/GenericWindowsTargetPlatform.h
--- a/Engine/Source/Developer/Windows/WindowsTargetPlatform/Private/GenericWindowsTargetPlatform.h
+++ b/Engine/Source/Developer/Windows/WindowsTargetPlatform/Private/GenericWindowsTargetPlatform.h
@@ -59,40 +59,35 @@ public:
 		static FName NAME_PCD3D_SM6(TEXT("PCD3D_SM6"));
 		static FName NAME_PCD3D_SM5(TEXT("PCD3D_SM5"));
 		static FName NAME_VULKAN_SM5(TEXT("SF_VULKAN_SM5"));
+		static FName NAME_PCD3D_ES3_1(TEXT("PCD3D_ES31"));
+		static FName NAME_SF_VULKAN_ES31(TEXT("SF_VULKAN_ES31"));
+		static FName NAME_OPENGL_150_ES3_1(TEXT("GLSL_150_ES31"));
 
 		bSupportDX11TextureFormats = true;
 		bSupportCompressedVolumeTexture = true;
 
 		for (FName TargetedShaderFormat : TargetedShaderFormats)
 		{
-			EShaderPlatform ShaderPlatform = SP_NumPlatforms;
-			// Can't use ShaderFormatToLegacyShaderPlatform() because of link dependency
+			if (TargetedShaderFormat == NAME_PCD3D_SM6 ||
+				TargetedShaderFormat == NAME_PCD3D_SM5 ||
+				TargetedShaderFormat == NAME_VULKAN_SM5 ||
+				TargetedShaderFormat == NAME_PCD3D_ES3_1)
 			{
-				if (TargetedShaderFormat == NAME_PCD3D_SM6)
-				{
-					ShaderPlatform = SP_PCD3D_SM6;
-				}
-				else if (TargetedShaderFormat == NAME_PCD3D_SM5)
-				{
-					ShaderPlatform = SP_PCD3D_SM5;
-				}
-				else if (TargetedShaderFormat == NAME_VULKAN_SM5)
-				{
-					ShaderPlatform = SP_VULKAN_SM5;
-				}
-			}
+				// DX11 formats okay
 
-			// If we're targeting only DX11 we can use DX11 texture formats. Otherwise we'd have to compress fallbacks and increase the size of cooked content significantly.
-			if (ShaderPlatform != SP_PCD3D_SM6 && ShaderPlatform != SP_PCD3D_SM5 && ShaderPlatform != SP_VULKAN_SM5)
+				// technically PCD3D_ES3_1 might not support DX11 formats, but in UE5 we require them
+				//   PCD3D_ES3_1 is used as a low-spec version of DX11
+			}
+			else
 			{
+				// some TargetedShaderFormat doesn't support DX11 formats
+				// must turn it off altogether
 				bSupportDX11TextureFormats = false;
+				break;
 			}
 		}
 
 		// If we are targeting ES3.1, we also must cook encoded HDR reflection captures
-		static FName NAME_SF_VULKAN_ES31(TEXT("SF_VULKAN_ES31"));
-		static FName NAME_OPENGL_150_ES3_1(TEXT("GLSL_150_ES31"));
-		static FName NAME_PCD3D_ES3_1(TEXT("PCD3D_ES31"));
 		bRequiresEncodedHDRReflectionCaptures =	TargetedShaderFormats.Contains(NAME_SF_VULKAN_ES31)
 												|| TargetedShaderFormats.Contains(NAME_OPENGL_150_ES3_1)
 												|| TargetedShaderFormats.Contains(NAME_PCD3D_ES3_1);
@@ -285,136 +280,6 @@ public:
 		}
 	}
 
-	FName GetVirtualTextureLayerFormat(int32 SourceFormat, bool bAllowCompression, bool bNoAlpha, bool bDX11TextureFormatsSupported, int32 Settings) const override
-	{
-		FName TextureFormatName = NAME_None;
-
-		// Supported texture format names.
-		static FName NameDXT1(TEXT("DXT1"));
-		static FName NameDXT3(TEXT("DXT3"));
-		static FName NameDXT5(TEXT("DXT5"));
-		static FName NameDXT5n(TEXT("DXT5n"));
-		static FName NameBC4(TEXT("BC4"));
-		static FName NameBC5(TEXT("BC5"));
-		static FName NameBGRA8(TEXT("BGRA8"));
-		static FName NameXGXR8(TEXT("XGXR8"));
-		static FName NameG8(TEXT("G8"));
-		static FName NameG16(TEXT("G16"));
-		static FName NameVU8(TEXT("VU8"));
-		static FName NameRGBA16F(TEXT("RGBA16F"));
-		static FName NameR16F(TEXT("R16F"));
-		static FName NameBC6H(TEXT("BC6H"));
-		static FName NameBC7(TEXT("BC7"));
-
-		// Note: We can't use things here like autoDXT here which defer the exact choice to the compressor as it would mean that
-		// some textures on a VT layer may get a different format than others.
-		// We need to guarantee the format to be the same for all textures on the layer so we need to decide on the exact final format here.
-
-		bool bUseDXT5NormalMap = false;
-		FString UseDXT5NormalMapsString;
-		if (this->GetConfigSystem()->GetString(TEXT("SystemSettings"), TEXT("Compat.UseDXT5NormalMaps"), UseDXT5NormalMapsString, GEngineIni))
-		{
-			bUseDXT5NormalMap = FCString::ToBool(*UseDXT5NormalMapsString);
-		}
-
-		// Determine the pixel format of the (un/)compressed texture
-		if (!bAllowCompression)
-		{
-			if (SourceFormat == TSF_RGBA16F)
-			{
-				TextureFormatName = NameRGBA16F;
-			}
-			else if (SourceFormat == TSF_G16)
-			{
-				TextureFormatName = NameG16;
-			}
-			else if (SourceFormat == TSF_G8 || Settings == TC_Grayscale)
-			{
-				TextureFormatName = NameG8;
-			}
-			else if (Settings == TC_Normalmap && bUseDXT5NormalMap)
-			{
-				TextureFormatName = NameXGXR8;
-			}
-			else
-			{
-				TextureFormatName = NameBGRA8;
-			}
-		}
-		else if (Settings == TC_HDR)
-		{
-			TextureFormatName = NameRGBA16F;
-		}
-		else if (Settings == TC_Normalmap)
-		{
-			TextureFormatName = bUseDXT5NormalMap ? NameDXT5n : NameBC5;
-		}
-		else if (Settings == TC_Displacementmap)
-		{
-			TextureFormatName = NameG8;
-		}
-		else if (Settings == TC_VectorDisplacementmap)
-		{
-			TextureFormatName = NameBGRA8;
-		}
-		else if (Settings == TC_Grayscale)
-		{
-			TextureFormatName = NameG8;
-		}
-		else if (Settings == TC_Alpha)
-		{
-			TextureFormatName = NameBC4;
-		}
-		else if (Settings == TC_DistanceFieldFont)
-		{
-			TextureFormatName = NameG8;
-		}
-		else if (Settings == TC_HDR_Compressed)
-		{
-			TextureFormatName = NameBC6H;
-		}
-		else if (Settings == TC_BC7)
-		{
-			TextureFormatName = NameBC7;
-		}
-		else if (Settings == TC_HalfFloat)
-		{
-			TextureFormatName = NameR16F;
-		}
-		else if (bNoAlpha)
-		{
-			TextureFormatName = NameDXT1;
-		}
-		else
-		{
-			TextureFormatName = NameDXT5;
-		}
-
-		/*
-		FIXME: IS this still relevant for VT ? comes from the texture variant
-		// Some PC GPUs don't support sRGB read from G8 textures (e.g. AMD DX10 cards on ShaderModel3.0)
-		// This solution requires 4x more memory but a lot of PC HW emulate the format anyway
-		if ((TextureFormatName == NameG8) && Texture->SRGB && !SupportsFeature(ETargetPlatformFeatures::GrayscaleSRGB))
-		{
-		TextureFormatName = NameBGRA8;
-		}*/
-
-		// fallback to non-DX11 formats if one was chosen, but we can't use it
-		if (!bDX11TextureFormatsSupported)
-		{
-			if (TextureFormatName == NameBC6H)
-			{
-				TextureFormatName = NameRGBA16F;
-			}
-			else if (TextureFormatName == NameBC7)
-			{
-				TextureFormatName = NameDXT5;
-			}
-		}
-
-		return TextureFormatName;
-	}
-
 	virtual const UTextureLODSettings& GetTextureLODSettings() const override
 	{
 		return *TextureLODSettings;
diff --git a/Engine/Source/Runtime/Engine/Private/Texture.cpp b/Engine/Source/Runtime/Engine/Private/Texture.cpp
--- a/Engine/Source/Runtime/Engine/Private/Texture.cpp
+++ b/Engine/Source/Runtime/Engine/Private/Texture.cpp
@@ -2491,10 +2491,14 @@ FName GetDefaultTextureFormatName( const ITargetPlatform* TargetPlatform, const
 	{
 		if (TextureFormatName == NameBC6H)
 		{
+			UE_LOG(LogTexture, Display, TEXT("BC6H Texture expanded to uncompressed RGBA16F because !SupportDX11 : %s"), *Texture->GetPathName());
+
 			TextureFormatName = NameRGBA16F;
 		}
 		else if (TextureFormatName == NameBC7)
 		{
+			UE_LOG(LogTexture, Display, TEXT("BC7 Texture expanded to uncompressed BGRA8 because !SupportDX11 : %s"), *Texture->GetPathName());
+
 			TextureFormatName = NameBGRA8;
 		}
 	}
