Patch separate atomic counter resource in HLSL precompilation step.

#rnx

diff --git a/Engine/Source/Developer/Windows/ShaderFormatD3D/Private/D3DShaderCompiler.cpp b/Engine/Source/Developer/Windows/ShaderFormatD3D/Private/D3DShaderCompiler.cpp
--- a/Engine/Source/Developer/Windows/ShaderFormatD3D/Private/D3DShaderCompiler.cpp
+++ b/Engine/Source/Developer/Windows/ShaderFormatD3D/Private/D3DShaderCompiler.cpp
@@ -579,19 +579,47 @@ static void PatchSpirvForPrecompilation(FSpirv& Spirv)
 	}
 }
 
+// @todo-lh: use ANSI string class whenever UE core gets one
 static void PatchHlslForPrecompilation(TArray<ANSICHAR>& HlslSource)
 {
-	const FAnsiStringView HlslSourceView = FAnsiStringView(HlslSource.GetData(), HlslSource.Num() - 1);
-	const int32 RootShaderParameterSourceLocation = HlslSourceView.Find("cbuffer RootShaderParameters");
+	FString HlslSourceString = ANSI_TO_TCHAR(HlslSource.GetData());
+
+	// Patch SPIRV-Cross renaming to retain original member names in RootShaderParameters cbuffer
+	const int32 RootShaderParameterSourceLocation = HlslSourceString.Find("cbuffer RootShaderParameters");
 	if (RootShaderParameterSourceLocation != INDEX_NONE)
 	{
-		FString HlslSourceString = ANSI_TO_TCHAR(HlslSource.GetData());
 		HlslSourceString.ReplaceInline(TEXT("cbuffer RootShaderParameters"), TEXT("cbuffer _RootShaderParameters"), ESearchCase::CaseSensitive);
 		HlslSourceString.ReplaceInline(TEXT("_RootShaderParameters_"), TEXT(""), ESearchCase::CaseSensitive);
-		HlslSource.SetNum(HlslSourceString.Len() + 1);
-		FMemory::Memcpy(HlslSource.GetData(), TCHAR_TO_ANSI(*HlslSourceString), HlslSourceString.Len());
-		HlslSource[HlslSourceString.Len()] = '\0';
 	}
+
+	// Patch separation of atomic counters: replace declarations of all counter_var_... declarations by their original buffer resource.
+	const FString CounterPrefix = TEXT("counter_var_");
+	const FString CounterDeclPrefix = TEXT("RWByteAddressBuffer ") + CounterPrefix;
+
+	for (int32 ReadPos = 0, NextReadPos = 0;
+		 (NextReadPos = HlslSourceString.Find(CounterDeclPrefix, ESearchCase::CaseSensitive, ESearchDir::FromStart, ReadPos)) != INDEX_NONE;
+		 ReadPos = NextReadPos)
+	{
+		// Find original resource name without "counter_var_" prefix
+		const int32 ResourceNameStartPos = NextReadPos + CounterDeclPrefix.Len();
+		const int32 ResourceNameEndPos = HlslSourceString.Find(TEXT(";"), ESearchCase::CaseSensitive, ESearchDir::FromStart, ResourceNameStartPos);
+		if (ResourceNameEndPos != INDEX_NONE)
+		{
+			const FString ResourceName = HlslSourceString.Mid(NextReadPos + CounterDeclPrefix.Len(), ResourceNameEndPos - ResourceNameStartPos);
+			const FString ResourceCounterName = HlslSourceString.Mid(NextReadPos + CounterDeclPrefix.Len() - CounterPrefix.Len(), ResourceNameEndPos - ResourceNameStartPos + CounterPrefix.Len());
+
+			// Remove current "RWByteAddressBuffer counter_var_*;" resource declaration line
+			HlslSourceString.RemoveAt(NextReadPos, ResourceNameEndPos - NextReadPos + 1);
+
+			// Remove all "counter_var_" prefixes for the current resource
+			HlslSourceString.ReplaceInline(*ResourceCounterName, *ResourceName, ESearchCase::CaseSensitive);
+		}
+	}
+
+	// Return new HLSL source
+	HlslSource.SetNum(HlslSourceString.Len() + 1);
+	FMemory::Memcpy(HlslSource.GetData(), TCHAR_TO_ANSI(*HlslSourceString), HlslSourceString.Len());
+	HlslSource[HlslSourceString.Len()] = '\0';
 }
 
 // Generate the dumped usf file; call the D3D compiler, gather reflection information and generate the output data
