Added uint64_t type to lexer grammar of HlslParser and improved warning output for removal of unused shader interpolators.

#rnx

diff --git a/Engine/Source/Developer/ShaderCompilerCommon/Private/HlslExpressionParser.inl b/Engine/Source/Developer/ShaderCompilerCommon/Private/HlslExpressionParser.inl
--- a/Engine/Source/Developer/ShaderCompilerCommon/Private/HlslExpressionParser.inl
+++ b/Engine/Source/Developer/ShaderCompilerCommon/Private/HlslExpressionParser.inl
@@ -269,6 +269,28 @@ namespace CrossCompiler
 		case EHlslToken::Uint4x3:
 		case EHlslToken::Uint4x4:
 
+		case EHlslToken::Uint64_t:
+		case EHlslToken::Uint64_t1:
+		case EHlslToken::Uint64_t2:
+		case EHlslToken::Uint64_t3:
+		case EHlslToken::Uint64_t4:
+		case EHlslToken::Uint64_t1x1:
+		case EHlslToken::Uint64_t1x2:
+		case EHlslToken::Uint64_t1x3:
+		case EHlslToken::Uint64_t1x4:
+		case EHlslToken::Uint64_t2x1:
+		case EHlslToken::Uint64_t2x2:
+		case EHlslToken::Uint64_t2x3:
+		case EHlslToken::Uint64_t2x4:
+		case EHlslToken::Uint64_t3x1:
+		case EHlslToken::Uint64_t3x2:
+		case EHlslToken::Uint64_t3x3:
+		case EHlslToken::Uint64_t3x4:
+		case EHlslToken::Uint64_t4x1:
+		case EHlslToken::Uint64_t4x2:
+		case EHlslToken::Uint64_t4x3:
+		case EHlslToken::Uint64_t4x4:
+
 		case EHlslToken::Half:
 		case EHlslToken::Half1:
 		case EHlslToken::Half2:
diff --git a/Engine/Source/Developer/ShaderCompilerCommon/Private/HlslLexer.cpp b/Engine/Source/Developer/ShaderCompilerCommon/Private/HlslLexer.cpp
--- a/Engine/Source/Developer/ShaderCompilerCommon/Private/HlslLexer.cpp
+++ b/Engine/Source/Developer/ShaderCompilerCommon/Private/HlslLexer.cpp
@@ -289,6 +289,28 @@ namespace CrossCompiler
 				InsertToken(TEXT("uint3x4"), EHlslToken::Uint3x4);
 				InsertToken(TEXT("uint4x4"), EHlslToken::Uint4x4);
 
+				InsertToken(TEXT("uint64_t"), EHlslToken::Uint64_t);
+				InsertToken(TEXT("uint64_t1"), EHlslToken::Uint64_t1);
+				InsertToken(TEXT("uint64_t2"), EHlslToken::Uint64_t2);
+				InsertToken(TEXT("uint64_t3"), EHlslToken::Uint64_t3);
+				InsertToken(TEXT("uint64_t4"), EHlslToken::Uint64_t4);
+				InsertToken(TEXT("uint64_t1x1"), EHlslToken::Uint64_t1x1);
+				InsertToken(TEXT("uint64_t2x1"), EHlslToken::Uint64_t2x1);
+				InsertToken(TEXT("uint64_t3x1"), EHlslToken::Uint64_t3x1);
+				InsertToken(TEXT("uint64_t4x1"), EHlslToken::Uint64_t4x1);
+				InsertToken(TEXT("uint64_t1x2"), EHlslToken::Uint64_t1x2);
+				InsertToken(TEXT("uint64_t2x2"), EHlslToken::Uint64_t2x2);
+				InsertToken(TEXT("uint64_t3x2"), EHlslToken::Uint64_t3x2);
+				InsertToken(TEXT("uint64_t4x2"), EHlslToken::Uint64_t4x2);
+				InsertToken(TEXT("uint64_t1x3"), EHlslToken::Uint64_t1x3);
+				InsertToken(TEXT("uint64_t2x3"), EHlslToken::Uint64_t2x3);
+				InsertToken(TEXT("uint64_t3x3"), EHlslToken::Uint64_t3x3);
+				InsertToken(TEXT("uint64_t4x3"), EHlslToken::Uint64_t4x3);
+				InsertToken(TEXT("uint64_t1x4"), EHlslToken::Uint64_t1x4);
+				InsertToken(TEXT("uint64_t2x4"), EHlslToken::Uint64_t2x4);
+				InsertToken(TEXT("uint64_t3x4"), EHlslToken::Uint64_t3x4);
+				InsertToken(TEXT("uint64_t4x4"), EHlslToken::Uint64_t4x4);
+
 				InsertToken(TEXT("half"), EHlslToken::Half);
 				InsertToken(TEXT("half1"), EHlslToken::Half1);
 				InsertToken(TEXT("half2"), EHlslToken::Half2);
diff --git a/Engine/Source/Developer/ShaderCompilerCommon/Private/HlslLexer.h b/Engine/Source/Developer/ShaderCompilerCommon/Private/HlslLexer.h
--- a/Engine/Source/Developer/ShaderCompilerCommon/Private/HlslLexer.h
+++ b/Engine/Source/Developer/ShaderCompilerCommon/Private/HlslLexer.h
@@ -151,6 +151,28 @@ namespace CrossCompiler
 		Uint3x4,
 		Uint4x4,
 
+		Uint64_t,
+		Uint64_t1,
+		Uint64_t2,
+		Uint64_t3,
+		Uint64_t4,
+		Uint64_t1x1,
+		Uint64_t2x1,
+		Uint64_t3x1,
+		Uint64_t4x1,
+		Uint64_t1x2,
+		Uint64_t2x2,
+		Uint64_t3x2,
+		Uint64_t4x2,
+		Uint64_t1x3,
+		Uint64_t2x3,
+		Uint64_t3x3,
+		Uint64_t4x3,
+		Uint64_t1x4,
+		Uint64_t2x4,
+		Uint64_t3x4,
+		Uint64_t4x4,
+
 		Half,
 		Half1,
 		Half2,
diff --git a/Engine/Source/Developer/Windows/ShaderFormatD3D/Private/D3DShaderCompiler.cpp b/Engine/Source/Developer/Windows/ShaderFormatD3D/Private/D3DShaderCompiler.cpp
--- a/Engine/Source/Developer/Windows/ShaderFormatD3D/Private/D3DShaderCompiler.cpp
+++ b/Engine/Source/Developer/Windows/ShaderFormatD3D/Private/D3DShaderCompiler.cpp
@@ -1020,10 +1020,12 @@ bool CompileAndProcessD3DShaderFXC(FString& PreprocessedShaderSource, const FStr
 						else
 						{
 							UE_LOG(LogD3D11ShaderCompiler, Warning, TEXT("Failed to remove unused inputs from shader: %s"), *Input.GenerateShaderName());
-							for (int32 Index = 0; Index < RemoveErrors.Num(); ++Index)
+							for (const FString& ErrorMessage : RemoveErrors)
 							{
+								// Add error to shader output but also make sure the error shows up on build farm by emitting a log entry
+								UE_LOG(LogD3D11ShaderCompiler, Warning, TEXT("%s"), *ErrorMessage);
 								FShaderCompilerError NewError;
-								NewError.StrippedErrorMessage = RemoveErrors[Index];
+								NewError.StrippedErrorMessage = ErrorMessage;
 								Output.Errors.Add(NewError);
 							}
 							Output.bFailedRemovingUnused = true;
@@ -1240,12 +1242,13 @@ void CompileD3DShader(const FShaderCompilerInput& Input, FShaderCompilerOutput&
 		TArray<FString> Errors;
 		if (!RemoveUnusedOutputs(PreprocessedShaderSource, UsedOutputs, Exceptions, EntryPointName, Errors))
 		{
-			DumpDebugShaderUSF(PreprocessedShaderSource, Input);
 			UE_LOG(LogD3D11ShaderCompiler, Warning, TEXT("Failed to remove unused outputs from shader: %s"), *Input.GenerateShaderName());
-			for (int32 Index = 0; Index < Errors.Num(); ++Index)
+			for (const FString& ErrorReport : Errors)
 			{
+				// Add error to shader output but also make sure the error shows up on build farm by emitting a log entry
+				UE_LOG(LogD3D11ShaderCompiler, Warning, TEXT("%s"), *ErrorReport);
 				FShaderCompilerError NewError;
-				NewError.StrippedErrorMessage = Errors[Index];
+				NewError.StrippedErrorMessage = ErrorReport;
 				Output.Errors.Add(NewError);
 			}
 			Output.bFailedRemovingUnused = true;
