Add scoped cpu markers to various D3D shader compiler functions to make this easier to profile when shader compilation is done in process.

diff --git a/Engine/Source/Developer/ShaderCompilerCommon/Private/ShaderCompilerCommon.cpp b/Engine/Source/Developer/ShaderCompilerCommon/Private/ShaderCompilerCommon.cpp
--- a/Engine/Source/Developer/ShaderCompilerCommon/Private/ShaderCompilerCommon.cpp
+++ b/Engine/Source/Developer/ShaderCompilerCommon/Private/ShaderCompilerCommon.cpp
@@ -1315,6 +1315,8 @@ void TransformStringIntoCharacterArray(FString& PreprocessedShaderSource);
 // Replace all uniform buffer struct member references (View.WorldToClip) with a flattened name that removes the struct dependency (View_WorldToClip)
 void RemoveUniformBuffersFromSource(const FShaderCompilerEnvironment& Environment, FString& PreprocessedShaderSource)
 {
+	TRACE_CPUPROFILER_EVENT_SCOPE(RemoveUniformBuffersFromSource);
+
 	TMap<FString, TArray<FUniformBufferMemberInfo>> UniformBufferNameToMembers;
 	UniformBufferNameToMembers.Reserve(Environment.UniformBufferMap.Num());
 
diff --git a/Engine/Source/Developer/ShaderPreprocessor/Private/ShaderPreprocessor.cpp b/Engine/Source/Developer/ShaderPreprocessor/Private/ShaderPreprocessor.cpp
--- a/Engine/Source/Developer/ShaderPreprocessor/Private/ShaderPreprocessor.cpp
+++ b/Engine/Source/Developer/ShaderPreprocessor/Private/ShaderPreprocessor.cpp
@@ -278,14 +278,17 @@ bool PreprocessShader(
 		ANSICHAR* McppOutAnsi = NULL;
 		ANSICHAR* McppErrAnsi = NULL;
 
-		McppResult = mcpp_run(
-			McppOptionsANSI.GetData(),
-			McppOptionsANSI.Num(),
-			TCHAR_TO_ANSI(*ShaderInput.VirtualSourceFilePath),
-			&McppOutAnsi,
-			&McppErrAnsi,
-			FileLoader.GetMcppInterface()
-		);
+		{
+			TRACE_CPUPROFILER_EVENT_SCOPE(mcpp_run);
+			McppResult = mcpp_run(
+				McppOptionsANSI.GetData(),
+				McppOptionsANSI.Num(),
+				TCHAR_TO_ANSI(*ShaderInput.VirtualSourceFilePath),
+				&McppOutAnsi,
+				&McppErrAnsi,
+				FileLoader.GetMcppInterface()
+			);
+		}
 
 		McppOutput = McppOutAnsi;
 		McppErrors = McppErrAnsi;
diff --git a/Engine/Source/Developer/Windows/ShaderFormatD3D/Private/D3DShaderCompiler.cpp b/Engine/Source/Developer/Windows/ShaderFormatD3D/Private/D3DShaderCompiler.cpp
--- a/Engine/Source/Developer/Windows/ShaderFormatD3D/Private/D3DShaderCompiler.cpp
+++ b/Engine/Source/Developer/Windows/ShaderFormatD3D/Private/D3DShaderCompiler.cpp
@@ -631,6 +631,8 @@ bool CompileAndProcessD3DShaderFXC(FString& PreprocessedShaderSource, const FStr
 	const TCHAR* ShaderProfile, bool bSecondPassAferUnusedInputRemoval,
 	TArray<FString>& FilteredErrors, FShaderCompilerOutput& Output)
 {
+	TRACE_CPUPROFILER_EVENT_SCOPE(CompileAndProcessD3DShaderFXC);
+
 	auto AnsiSourceFile = StringCast<ANSICHAR>(*PreprocessedShaderSource);
 
 	// Write out the preprocessed file and a batch file to compile it if requested (DumpDebugInfoPath is valid)
@@ -1142,6 +1144,7 @@ bool CompileAndProcessD3DShaderFXC(FString& PreprocessedShaderSource, const FStr
 
 void CompileD3DShader(const FShaderCompilerInput& Input, FShaderCompilerOutput& Output, FShaderCompilerDefinitions& AdditionalDefines, const FString& WorkingDirectory, ELanguage Language)
 {
+	TRACE_CPUPROFILER_EVENT_SCOPE(CompileD3DShader);
 	FString PreprocessedShaderSource;
 	const bool bIsRayTracingShader = Input.IsRayTracingShader();
 	const bool bUseDXC =
diff --git a/Engine/Source/Developer/Windows/ShaderFormatD3D/Private/D3DShaderCompilerDXC.cpp b/Engine/Source/Developer/Windows/ShaderFormatD3D/Private/D3DShaderCompilerDXC.cpp
--- a/Engine/Source/Developer/Windows/ShaderFormatD3D/Private/D3DShaderCompilerDXC.cpp
+++ b/Engine/Source/Developer/Windows/ShaderFormatD3D/Private/D3DShaderCompilerDXC.cpp
@@ -602,6 +602,8 @@ bool CompileAndProcessD3DShaderDXC(FString& PreprocessedShaderSource,
 	const TCHAR* ShaderProfile, ELanguage Language, bool bProcessingSecondTime,
 	TArray<FString>& FilteredErrors, FShaderCompilerOutput& Output)
 {
+	TRACE_CPUPROFILER_EVENT_SCOPE(CompileAndProcessD3DShaderDXC);
+
 	auto AnsiSourceFile = StringCast<ANSICHAR>(*PreprocessedShaderSource);
 
 	const bool bIsRayTracingShader = Input.IsRayTracingShader();
diff --git a/Engine/Source/Runtime/Engine/Private/ShaderCompiler/ShaderCompiler.cpp b/Engine/Source/Runtime/Engine/Private/ShaderCompiler/ShaderCompiler.cpp
--- a/Engine/Source/Runtime/Engine/Private/ShaderCompiler/ShaderCompiler.cpp
+++ b/Engine/Source/Runtime/Engine/Private/ShaderCompiler/ShaderCompiler.cpp
@@ -2547,6 +2547,8 @@ void FShaderCompileThreadRunnable::CompileDirectlyThroughDll()
 
 void FShaderCompileUtilities::ExecuteShaderCompileJob(FShaderCommonCompileJob& Job)
 {
+	TRACE_CPUPROFILER_EVENT_SCOPE(FShaderCompileUtilities::ExecuteShaderCompileJob);
+
 	check(!Job.bFinalized);
 
 	static ITargetPlatformManagerModule& TPM = GetTargetPlatformManagerRef();
