Set InstructionCount when compiling shaders with DXC

- This fixes an issue where all SM6 compiles were reporting 0 instructions.

diff --git a/Engine/Source/Developer/Windows/ShaderFormatD3D/Private/D3DShaderCompilerDXC.cpp b/Engine/Source/Developer/Windows/ShaderFormatD3D/Private/D3DShaderCompilerDXC.cpp
--- a/Engine/Source/Developer/Windows/ShaderFormatD3D/Private/D3DShaderCompilerDXC.cpp
+++ b/Engine/Source/Developer/Windows/ShaderFormatD3D/Private/D3DShaderCompilerDXC.cpp
@@ -876,6 +876,7 @@ bool CompileAndProcessD3DShaderDXC(FString& PreprocessedShaderSource,
 					NumSamplers, NumSRVs, NumCBs, NumUAVs,
 					Output, UniformBufferNames, UsedUniformBufferSlots, VendorExtensions);
 
+			NumInstructions = ShaderDesc.InstructionCount;
 			Output.bSucceeded = true;
 		}
 
