Reorganizing reflection handling to funnel through common functions.
This is a prerequisite for upcoming bindless changes as it centralizes key points where we register the different binding types.

diff --git a/Engine/Source/Developer/ShaderCompilerCommon/Private/ShaderCompilerCommon.cpp b/Engine/Source/Developer/ShaderCompilerCommon/Private/ShaderCompilerCommon.cpp
--- a/Engine/Source/Developer/ShaderCompilerCommon/Private/ShaderCompilerCommon.cpp
+++ b/Engine/Source/Developer/ShaderCompilerCommon/Private/ShaderCompilerCommon.cpp
@@ -469,6 +469,115 @@ TCHAR* FindNextUniformBufferReference(TCHAR* SearchPtr, const TCHAR* SearchStrin
 	return nullptr;
 }
 
+void HandleReflectedGlobalConstantBufferMember(
+	const FString& MemberName,
+	uint32 ConstantBufferIndex,
+	int32 ReflectionOffset,
+	int32 ReflectionSize,
+	FShaderCompilerOutput& Output
+)
+{
+	Output.ParameterMap.AddParameterAllocation(
+		*MemberName,
+		ConstantBufferIndex,
+		ReflectionOffset,
+		ReflectionSize,
+		EShaderParameterType::LooseData);
+}
+
+void HandleReflectedRootConstantBufferMember(
+	const FShaderCompilerInput& Input,
+	const FShaderParameterParser& ShaderParameterParser,
+	const FString& MemberName,
+	int32 ReflectionOffset,
+	int32 ReflectionSize,
+	FShaderCompilerOutput& Output
+)
+{
+	ShaderParameterParser.ValidateShaderParameterType(Input, MemberName, ReflectionOffset, ReflectionSize, Output);
+}
+
+void HandleReflectedRootConstantBuffer(
+	int32 ConstantBufferSize,
+	FShaderCompilerOutput& CompilerOutput
+)
+{
+	CompilerOutput.ParameterMap.AddParameterAllocation(
+		FShaderParametersMetadata::kRootUniformBufferBindingName,
+		FShaderParametersMetadata::kRootCBufferBindingIndex,
+		0,
+		static_cast<uint16>(ConstantBufferSize),
+		EShaderParameterType::LooseData);
+}
+
+void HandleReflectedUniformBuffer(
+	const FString& UniformBufferName,
+	int32 ReflectionSlot,
+	int32 BaseIndex,
+	int32 BufferSize,
+	FShaderCompilerOutput& CompilerOutput
+)
+{
+	CompilerOutput.ParameterMap.AddParameterAllocation(
+		*UniformBufferName,
+		ReflectionSlot,
+		BaseIndex,
+		BufferSize,
+		EShaderParameterType::UniformBuffer
+	);
+}
+
+void HandleReflectedShaderResource(
+	const FString& ResourceName,
+	int32 BindOffset,
+	int32 ReflectionSlot,
+	int32 BindCount,
+	FShaderCompilerOutput& CompilerOutput
+)
+{
+	CompilerOutput.ParameterMap.AddParameterAllocation(
+		*ResourceName,
+		BindOffset,
+		ReflectionSlot,
+		BindCount,
+		EShaderParameterType::SRV
+	);
+}
+
+void HandleReflectedShaderUAV(
+	const FString& UAVName,
+	int32 BindOffset,
+	int32 ReflectionSlot,
+	int32 BindCount,
+	FShaderCompilerOutput& CompilerOutput
+)
+{
+	CompilerOutput.ParameterMap.AddParameterAllocation(
+		*UAVName,
+		BindOffset,
+		ReflectionSlot,
+		BindCount,
+		EShaderParameterType::UAV
+	);
+}
+
+void HandleReflectedShaderSampler(
+	const FString& SamplerName,
+	int32 BindOffset,
+	int32 ReflectionSlot,
+	int32 BindCount,
+	FShaderCompilerOutput& CompilerOutput
+)
+{
+	CompilerOutput.ParameterMap.AddParameterAllocation(
+		*SamplerName,
+		BindOffset,
+		ReflectionSlot,
+		BindCount,
+		EShaderParameterType::Sampler
+	);
+}
+
 void AddNoteToDisplayShaderParameterStructureOnCppSide(
 	const FShaderParametersMetadata* ParametersStructure,
 	FShaderCompilerOutput& CompilerOutput)
diff --git a/Engine/Source/Developer/ShaderCompilerCommon/Public/ShaderCompilerCommon.h b/Engine/Source/Developer/ShaderCompilerCommon/Public/ShaderCompilerCommon.h
--- a/Engine/Source/Developer/ShaderCompilerCommon/Public/ShaderCompilerCommon.h
+++ b/Engine/Source/Developer/ShaderCompilerCommon/Public/ShaderCompilerCommon.h
@@ -139,6 +139,100 @@ private:
 	bool bMovedLoosedParametersToRootConstantBuffer = false;
 };
 
+extern SHADERCOMPILERCOMMON_API void HandleReflectedGlobalConstantBufferMember(
+	const FString& MemberName,
+	uint32 ConstantBufferIndex,
+	int32 ReflectionOffset,
+	int32 ReflectionSize,
+	FShaderCompilerOutput& CompilerOutput
+);
+
+extern SHADERCOMPILERCOMMON_API void HandleReflectedRootConstantBufferMember(
+	const FShaderCompilerInput& Input,
+	const FShaderParameterParser& ShaderParameterParser,
+	const FString& MemberName,
+	int32 ReflectionOffset,
+	int32 ReflectionSize,
+	FShaderCompilerOutput& CompilerOutput
+);
+
+extern SHADERCOMPILERCOMMON_API void HandleReflectedRootConstantBuffer(
+	int32 ConstantBufferSize,
+	FShaderCompilerOutput& CompilerOutput
+);
+
+extern SHADERCOMPILERCOMMON_API void HandleReflectedUniformBuffer(
+	const FString& UniformBufferName,
+	int32 ReflectionSlot,
+	int32 BaseIndex,
+	int32 BufferSize,
+	FShaderCompilerOutput& CompilerOutput
+);
+
+inline void HandleReflectedUniformBuffer(const FString& UniformBufferName, int32 ReflectionSlot, int32 BufferSize, FShaderCompilerOutput& CompilerOutput)
+{
+	HandleReflectedUniformBuffer(UniformBufferName, ReflectionSlot, 0, BufferSize, CompilerOutput);
+}
+
+inline void HandleReflectedUniformBuffer(const FString& UniformBufferName, int32 ReflectionSlot, FShaderCompilerOutput& CompilerOutput)
+{
+	HandleReflectedUniformBuffer(UniformBufferName, ReflectionSlot, 0, CompilerOutput);
+}
+
+extern SHADERCOMPILERCOMMON_API void HandleReflectedShaderResource(
+	const FString& ResourceName,
+	int32 BindOffset,
+	int32 ReflectionSlot,
+	int32 BindCount,
+	FShaderCompilerOutput& CompilerOutput
+);
+
+inline void HandleReflectedShaderResource(const FString& ResourceName, int32 ReflectionSlot, int32 BindCount, FShaderCompilerOutput& CompilerOutput)
+{
+	HandleReflectedShaderResource(ResourceName, 0, ReflectionSlot, BindCount, CompilerOutput);
+}
+
+inline void HandleReflectedShaderResource(const FString& ResourceName, int32 ReflectionSlot, FShaderCompilerOutput& CompilerOutput)
+{
+	HandleReflectedShaderResource(ResourceName, ReflectionSlot, 1, CompilerOutput);
+}
+
+extern SHADERCOMPILERCOMMON_API void HandleReflectedShaderUAV(
+	const FString& UAVName,
+	int32 BindOffset,
+	int32 ReflectionSlot,
+	int32 BindCount,
+	FShaderCompilerOutput& CompilerOutput
+);
+
+inline void HandleReflectedShaderUAV(const FString& UAVName, int32 ReflectionSlot, int32 BindCount, FShaderCompilerOutput& CompilerOutput)
+{
+	HandleReflectedShaderUAV(UAVName, 0, ReflectionSlot, BindCount, CompilerOutput);
+}
+
+inline void HandleReflectedShaderUAV(const FString& UAVName, int32 ReflectionSlot, FShaderCompilerOutput& CompilerOutput)
+{
+	HandleReflectedShaderUAV(UAVName, ReflectionSlot, 1, CompilerOutput);
+}
+
+extern SHADERCOMPILERCOMMON_API void HandleReflectedShaderSampler(
+	const FString& SamplerName,
+	int32 BindOffset,
+	int32 ReflectionSlot,
+	int32 BindCount,
+	FShaderCompilerOutput& CompilerOutput
+);
+
+inline void HandleReflectedShaderSampler(const FString& SamplerName, int32 ReflectionSlot, int32 BindCount, FShaderCompilerOutput& CompilerOutput)
+{
+	HandleReflectedShaderSampler(SamplerName, 0, ReflectionSlot, BindCount, CompilerOutput);
+}
+
+inline void HandleReflectedShaderSampler(const FString& SamplerName, int32 ReflectionSlot, FShaderCompilerOutput& CompilerOutput)
+{
+	HandleReflectedShaderSampler(SamplerName, ReflectionSlot, 1, CompilerOutput);
+}
+
 /** Adds a note to CompilerOutput.Error about where the shader parameter structure is on C++ side. */
 extern SHADERCOMPILERCOMMON_API void AddNoteToDisplayShaderParameterStructureOnCppSide(
 	const FShaderParametersMetadata* ParametersStructure,
diff --git a/Engine/Source/Developer/ShaderFormatOpenGL/Private/OpenGLShaderCompiler.cpp b/Engine/Source/Developer/ShaderFormatOpenGL/Private/OpenGLShaderCompiler.cpp
--- a/Engine/Source/Developer/ShaderFormatOpenGL/Private/OpenGLShaderCompiler.cpp
+++ b/Engine/Source/Developer/ShaderFormatOpenGL/Private/OpenGLShaderCompiler.cpp
@@ -712,7 +712,7 @@ void FOpenGLFrontend::BuildShaderOutput(
 		}
 		else
 		{
-			ParameterMap.AddParameterAllocation(*UniformBlock.Name, UBIndex, 0, 0, EShaderParameterType::UniformBuffer);
+			HandleReflectedUniformBuffer(UniformBlock.Name, UBIndex, ShaderOutput);
 		}
 		Header.Bindings.NumUniformBuffers++;
 	}
@@ -743,13 +743,13 @@ void FOpenGLFrontend::BuildShaderOutput(
 		if (bIgnore)
 			continue;
 
-		ParameterMap.AddParameterAllocation(
-			*PackedGlobal.Name,
+		HandleReflectedGlobalConstantBufferMember(
+			PackedGlobal.Name,
 			PackedGlobal.PackedType,
-			PackedGlobal.Offset * BytesPerComponent,
-			PackedGlobal.Count * BytesPerComponent,
-			EShaderParameterType::LooseData
-			);
+			PackedGlobal.Offset* BytesPerComponent,
+			PackedGlobal.Count* BytesPerComponent,
+			ShaderOutput
+		);
 
 		uint16& Size = PackedGlobalArraySize.FindOrAdd(PackedGlobal.PackedType);
 		Size = FMath::Max<uint16>(BytesPerComponent * (PackedGlobal.Offset + PackedGlobal.Count), Size);
@@ -769,7 +769,7 @@ void FOpenGLFrontend::BuildShaderOutput(
 		}
 		else
 		{
-			ParameterMap.AddParameterAllocation(*PackedUB.Attribute.Name, PackedUB.Attribute.Index, 0, 0, EShaderParameterType::UniformBuffer);
+			HandleReflectedUniformBuffer(PackedUB.Attribute.Name, PackedUB.Attribute.Index, ShaderOutput);
 		}
 		Header.Bindings.NumUniformBuffers++;
 
@@ -884,13 +884,7 @@ void FOpenGLFrontend::BuildShaderOutput(
 		}
 		else
 		{
-		ParameterMap.AddParameterAllocation(
-			*Sampler.Name,
-			0,
-			Sampler.Offset,
-			Sampler.Count,
-			EShaderParameterType::SRV
-			);
+			HandleReflectedShaderResource(Sampler.Name, Sampler.Offset, Sampler.Count, ShaderOutput);
 		}
 
 		Header.Bindings.NumSamplers = FMath::Max<uint8>(
@@ -907,13 +901,7 @@ void FOpenGLFrontend::BuildShaderOutput(
 			}
 			else
 			{
-				ParameterMap.AddParameterAllocation(
-					*SamplerState,
-					0,
-					Sampler.Offset,
-					Sampler.Count,
-					EShaderParameterType::Sampler
-					);
+				HandleReflectedShaderSampler(SamplerState, Sampler.Offset, Sampler.Count, ShaderOutput);
 			}
 		}
 	}
@@ -930,13 +918,7 @@ void FOpenGLFrontend::BuildShaderOutput(
 		}
 		else
 		{
-		ParameterMap.AddParameterAllocation(
-			*UAV.Name,
-			0,
-			UAV.Offset,
-			UAV.Count,
-			EShaderParameterType::UAV
-			);
+			HandleReflectedShaderUAV(UAV.Name, UAV.Offset, UAV.Count, ShaderOutput);
 		}
 
 		Header.Bindings.NumUAVs = FMath::Max<uint8>(
diff --git a/Engine/Source/Developer/VulkanShaderFormat/Private/VulkanShaderCompiler.cpp b/Engine/Source/Developer/VulkanShaderFormat/Private/VulkanShaderCompiler.cpp
--- a/Engine/Source/Developer/VulkanShaderFormat/Private/VulkanShaderCompiler.cpp
+++ b/Engine/Source/Developer/VulkanShaderFormat/Private/VulkanShaderCompiler.cpp
@@ -996,7 +996,9 @@ static void BuildShaderOutput(
 			check(VulkanBindingIndex != -1);
 			check(!UsedUniformBufferSlots[VulkanBindingIndex]);
 			UsedUniformBufferSlots[VulkanBindingIndex] = true;
-			ParameterMap.AddParameterAllocation(*UniformBlock.Name, VulkanBindingIndex, 0, 0, EShaderParameterType::UniformBuffer);
+
+			HandleReflectedUniformBuffer(UniformBlock.Name, VulkanBindingIndex, ShaderOutput);
+
 			++OLDHeader.SerializedBindings.NumUniformBuffers;
 			NEWEntryTypes.Add(*UniformBlock.Name, FVulkanShaderHeader::UniformBuffer);
 		}
@@ -1059,14 +1061,15 @@ static void BuildShaderOutput(
 		}
 		check(Found != -1);
 
-		ParameterMap.AddParameterAllocation(
-			*PackedGlobal.Name,
+		HandleReflectedGlobalConstantBufferMember(
+			PackedGlobal.Name,
 			Found,
 			PackedGlobal.Offset * BytesPerComponent,
 			PackedGlobal.Count * BytesPerComponent,
-			EShaderParameterType::LooseData
-			);
-		NEWEntryTypes.Add(*PackedGlobal.Name, FVulkanShaderHeader::PackedGlobal);
+			ShaderOutput
+		);
+
+		NEWEntryTypes.Add(PackedGlobal.Name, FVulkanShaderHeader::PackedGlobal);
 
 		uint32& Size = PackedGlobalArraySize.FindOrAdd((CrossCompiler::EPackedTypeName)PackedGlobal.PackedType);
 		Size = FMath::Max<uint32>(BytesPerComponent * (PackedGlobal.Offset + PackedGlobal.Count), Size);
@@ -1080,7 +1083,9 @@ static void BuildShaderOutput(
 		//check(PackedUB.Attribute.Index == Header.SerializedBindings.NumUniformBuffers);
 		check(!UsedUniformBufferSlots[OLDHeader.UNUSED_NumNonGlobalUBs]);
 		UsedUniformBufferSlots[OLDHeader.UNUSED_NumNonGlobalUBs] = true;
-		ParameterMap.AddParameterAllocation(*PackedUB.Attribute.Name, OLDHeader.UNUSED_NumNonGlobalUBs++, PackedUB.Attribute.Index, 0, EShaderParameterType::UniformBuffer);
+
+		HandleReflectedUniformBuffer(PackedUB.Attribute.Name, OLDHeader.UNUSED_NumNonGlobalUBs++, PackedUB.Attribute.Index, 0, ShaderOutput);
+
 		NEWEntryTypes.Add(PackedUB.Attribute.Name, FVulkanShaderHeader::PackedGlobal);
 	}
 
@@ -1207,13 +1212,9 @@ static void BuildShaderOutput(
 		auto& Binding = HlslccBindings[HlslccBindingIndex];
 		int32 BindingIndex = Spirv.FindBinding(Binding.Name, true);
 		check(BindingIndex != -1);
-		ParameterMap.AddParameterAllocation(
-			*Name,
-			0,
-			BindingIndex,
-			1,
-			EShaderParameterType::Sampler
-		);
+
+		HandleReflectedShaderSampler(Name, BindingIndex, ShaderOutput);
+
 		NEWEntryTypes.Add(Name, FVulkanShaderHeader::Global);
 
 		// Count only samplers states, not textures
@@ -1224,13 +1225,9 @@ static void BuildShaderOutput(
 	{
 		int32 VulkanBindingIndex = Spirv.FindBinding(Sampler.Name, true);
 		check(VulkanBindingIndex != -1);
-		ParameterMap.AddParameterAllocation(
-			*Sampler.Name,
-			Sampler.Offset,
-			VulkanBindingIndex,
-			Sampler.Count,
-			EShaderParameterType::SRV
-		);
+
+		HandleReflectedShaderResource(Sampler.Name, Sampler.Offset, VulkanBindingIndex, Sampler.Count, ShaderOutput);
+
 		NEWEntryTypes.Add(Sampler.Name, FVulkanShaderHeader::Global);
 
 		for (auto& SamplerState : Sampler.SamplerStates)
@@ -1239,13 +1236,9 @@ static void BuildShaderOutput(
 			{
 				// ParameterMap does not use a TMultiMap, so we cannot push the same entry to it more than once!  if we try to, we've done something wrong...
 				check(!ParameterMap.ContainsParameterAllocation(*SamplerState));
-				ParameterMap.AddParameterAllocation(
-					*SamplerState,
-					Sampler.Offset,
-					VulkanBindingIndex,
-					Sampler.Count,
-					EShaderParameterType::Sampler
-				);
+
+				HandleReflectedShaderSampler(SamplerState, Sampler.Offset, VulkanBindingIndex, Sampler.Count, ShaderOutput);
+
 				NEWEntryTypes.Add(SamplerState, FVulkanShaderHeader::Global);
 
 				// Count compiled texture-samplers as output samplers
@@ -1259,13 +1252,8 @@ static void BuildShaderOutput(
 		int32 VulkanBindingIndex = Spirv.FindBinding(UAV.Name);
 		check(VulkanBindingIndex != -1);
 
-		ParameterMap.AddParameterAllocation(
-			*UAV.Name,
-			UAV.Offset,
-			VulkanBindingIndex,
-			UAV.Count,
-			EShaderParameterType::UAV
-			);
+		HandleReflectedShaderUAV(UAV.Name, UAV.Offset, VulkanBindingIndex, UAV.Count, ShaderOutput);
+
 		NEWEntryTypes.Add(UAV.Name, FVulkanShaderHeader::Global);
 
 		OLDHeader.SerializedBindings.NumUAVs = FMath::Max<uint8>(
@@ -1279,13 +1267,8 @@ static void BuildShaderOutput(
 		int32 VulkanBindingIndex = Spirv.FindBinding(AccelerationStructure.Name);
 		check(VulkanBindingIndex != -1);
 
-		ParameterMap.AddParameterAllocation(
-			*AccelerationStructure.Name,
-			AccelerationStructure.Offset,
-			VulkanBindingIndex,
-			1,
-			EShaderParameterType::SRV
-		);
+		HandleReflectedShaderResource(AccelerationStructure.Name, AccelerationStructure.Offset, VulkanBindingIndex, 1, ShaderOutput);
+
 		NEWEntryTypes.Add(AccelerationStructure.Name, FVulkanShaderHeader::Global);
 
 		OLDHeader.SerializedBindings.NumAccelerationStructures = FMath::Max<uint8>(
diff --git a/Engine/Source/Developer/Windows/ShaderFormatD3D/Private/D3DShaderCompiler.inl b/Engine/Source/Developer/Windows/ShaderFormatD3D/Private/D3DShaderCompiler.inl
--- a/Engine/Source/Developer/Windows/ShaderFormatD3D/Private/D3DShaderCompiler.inl
+++ b/Engine/Source/Developer/Windows/ShaderFormatD3D/Private/D3DShaderCompiler.inl
@@ -342,13 +342,14 @@ template <typename ID3D1xShaderReflection, typename D3D1x_SHADER_DESC, typename
 						{
 							bGlobalUniformBufferUsed = true;
 
-							Output.ParameterMap.AddParameterAllocation(
-								ANSI_TO_TCHAR(VariableDesc.Name),
+							HandleReflectedGlobalConstantBufferMember(
+								FString(VariableDesc.Name),
 								CBIndex,
 								VariableDesc.StartOffset,
 								VariableDesc.Size,
-								EShaderParameterType::LooseData
+								Output
 							);
+
 							UsedUniformBufferSlots[CBIndex] = true;
 						}
 					}
@@ -368,24 +369,22 @@ template <typename ID3D1xShaderReflection, typename D3D1x_SHADER_DESC, typename
 						Variable->GetDesc(&VariableDesc);
 						if (VariableDesc.uFlags & D3D_SVF_USED)
 						{
-							FString MemberName(ANSI_TO_TCHAR(VariableDesc.Name));
-							int32 ReflectionSize = VariableDesc.Size;
-							int32 ReflectionOffset = VariableDesc.StartOffset;
-
-							ShaderParameterParser.ValidateShaderParameterType(Input, MemberName, ReflectionOffset, ReflectionSize, Output);
+							HandleReflectedRootConstantBufferMember(
+								Input,
+								ShaderParameterParser,
+								FString(VariableDesc.Name),
+								VariableDesc.StartOffset,
+								VariableDesc.Size,
+								Output
+							);
 
-							ConstantBufferSize = FMath::Max(ConstantBufferSize, ReflectionOffset + ReflectionSize);
+							ConstantBufferSize = FMath::Max<int32>(ConstantBufferSize, VariableDesc.StartOffset + VariableDesc.Size);
 						}
 					}
 
 					if (ConstantBufferSize > 0)
 					{
-						Output.ParameterMap.AddParameterAllocation(
-							FShaderParametersMetadata::kRootUniformBufferBindingName,
-							FShaderParametersMetadata::kRootCBufferBindingIndex,
-							/* Offset = */ uint16(0),
-							/* Size = */ uint16(ConstantBufferSize),
-							EShaderParameterType::LooseData);
+						HandleReflectedRootConstantBuffer(ConstantBufferSize, Output);
 
 						bGlobalUniformBufferUsed = true;
 						UsedUniformBufferSlots[CBIndex] = true;
@@ -405,20 +404,16 @@ template <typename ID3D1xShaderReflection, typename D3D1x_SHADER_DESC, typename
 			else
 			{
 				// Track just the constant buffer itself.
-				Output.ParameterMap.AddParameterAllocation(
-					ANSI_TO_TCHAR(CBDesc.Name),
-					CBIndex,
-					0,
-					0,
-					EShaderParameterType::UniformBuffer
-				);
-				UsedUniformBufferSlots[CBIndex] = true;
+				const FString UniformBufferName(CBDesc.Name);
+
+				HandleReflectedUniformBuffer(UniformBufferName, CBIndex, Output);
 
 				if (UniformBufferNames.Num() <= (int32)CBIndex)
 				{
 					UniformBufferNames.AddDefaulted(CBIndex - UniformBufferNames.Num() + 1);
 				}
-				UniformBufferNames[CBIndex] = CBDesc.Name;
+				UniformBufferNames[CBIndex] = UniformBufferName;
+				UsedUniformBufferSlots[CBIndex] = true;
 			}
 
 			NumCBs = FMath::Max(NumCBs, BindDesc.BindPoint + BindDesc.BindCount);
@@ -432,9 +427,6 @@ template <typename ID3D1xShaderReflection, typename D3D1x_SHADER_DESC, typename
 			const bool bIsAMDSmpExtension = (FCStringAnsi::Strcmp(BindDesc.Name, "AmdDxExtShaderIntrinsicsSamplerState") == 0);
 			const bool bIsVendorParameter = bIsAMDTexExtension || bIsAMDSmpExtension;
 
-			TCHAR OfficialName[1024];
-			FCString::Strcpy(OfficialName, ANSI_TO_TCHAR(BindDesc.Name));
-
 			const uint32 BindCount = 1;
 			EShaderParameterType ParameterType = EShaderParameterType::Num;
 			if (BindDesc.Type == D3D_SIT_SAMPLER)
@@ -452,16 +444,13 @@ template <typename ID3D1xShaderReflection, typename D3D1x_SHADER_DESC, typename
 			{
 				VendorExtensions.Emplace(0x1002 /*AMD*/, 0, BindDesc.BindPoint, BindCount, ParameterType);
 			}
+			else if (ParameterType == EShaderParameterType::Sampler)
+			{
+				HandleReflectedShaderSampler(FString(BindDesc.Name), BindDesc.BindPoint, Output);
+			}
 			else
 			{
-				// Add a parameter for the texture only, the sampler index will be invalid
-				Output.ParameterMap.AddParameterAllocation(
-					OfficialName,
-					0,
-					BindDesc.BindPoint,
-					BindCount,
-					ParameterType
-				);
+				HandleReflectedShaderResource(FString(BindDesc.Name), BindDesc.BindPoint, Output);
 			}
 		}
 		else if (BindDesc.Type == D3D_SIT_UAV_RWTYPED || BindDesc.Type == D3D_SIT_UAV_RWSTRUCTURED ||
@@ -487,9 +476,6 @@ template <typename ID3D1xShaderReflection, typename D3D1x_SHADER_DESC, typename
 			// See D3DCommon.ush
 			const bool bIsDiagnosticBufferParameter = (FCStringAnsi::Strcmp(BindDesc.Name, "UEDiagnosticBuffer") == 0);
 
-			TCHAR OfficialName[1024];
-			FCString::Strcpy(OfficialName, ANSI_TO_TCHAR(BindDesc.Name));
-
 			const uint32 BindCount = 1;
 			if (bIsVendorParameter)
 			{
@@ -506,13 +492,7 @@ template <typename ID3D1xShaderReflection, typename D3D1x_SHADER_DESC, typename
 			}
 			else
 			{
-				Output.ParameterMap.AddParameterAllocation(
-					OfficialName,
-					0,
-					BindDesc.BindPoint,
-					BindCount,
-					EShaderParameterType::UAV
-				);
+				HandleReflectedShaderUAV(FString(BindDesc.Name), BindDesc.BindPoint, Output);
 			}
 
 			NumUAVs = FMath::Max(NumUAVs, BindDesc.BindPoint + BindCount);
@@ -520,38 +500,15 @@ template <typename ID3D1xShaderReflection, typename D3D1x_SHADER_DESC, typename
 		else if (BindDesc.Type == D3D_SIT_STRUCTURED || BindDesc.Type == D3D_SIT_BYTEADDRESS)
 		{
 			check(BindDesc.BindCount == 1);
-			TCHAR OfficialName[1024];
-			FCString::Strcpy(OfficialName, ANSI_TO_TCHAR(BindDesc.Name));
-
-			const uint32 BindCount = 1;
-			Output.ParameterMap.AddParameterAllocation(
-				OfficialName,
-				0,
-				BindDesc.BindPoint,
-				BindCount,
-				EShaderParameterType::SRV
-			);
-
-			NumSRVs = FMath::Max(NumSRVs, BindDesc.BindPoint + BindCount);
+			HandleReflectedShaderResource(FString(BindDesc.Name), BindDesc.BindPoint, Output);
+			NumSRVs = FMath::Max(NumSRVs, BindDesc.BindPoint + 1);
 		}
 		else if (BindDesc.Type == (D3D_SHADER_INPUT_TYPE)(D3D_SIT_UAV_RWSTRUCTURED_WITH_COUNTER + 1)) // D3D_SIT_RTACCELERATIONSTRUCTURE (12)
 		{
 			// Acceleration structure resources are treated as SRVs.
 			check(BindDesc.BindCount == 1);
-
-			TCHAR OfficialName[1024];
-			FCString::Strcpy(OfficialName, ANSI_TO_TCHAR(BindDesc.Name));
-
-			const uint32 BindCount = 1;
-			Output.ParameterMap.AddParameterAllocation(
-				OfficialName,
-				0,
-				BindDesc.BindPoint,
-				BindCount,
-				EShaderParameterType::SRV
-			);
-
-			NumSRVs = FMath::Max(NumSRVs, BindDesc.BindPoint + BindCount);
+			HandleReflectedShaderResource(FString(BindDesc.Name), BindDesc.BindPoint, Output);
+			NumSRVs = FMath::Max(NumSRVs, BindDesc.BindPoint + 1);
 		}
 	}
 }
