Allow environment variable FXC to be used with shader debugging batch scripts to avoid hardcoded system path.

#rnx

diff --git a/Engine/Source/Developer/Windows/ShaderFormatD3D/Private/D3DShaderCompiler.cpp b/Engine/Source/Developer/Windows/ShaderFormatD3D/Private/D3DShaderCompiler.cpp
--- a/Engine/Source/Developer/Windows/ShaderFormatD3D/Private/D3DShaderCompiler.cpp
+++ b/Engine/Source/Developer/Windows/ShaderFormatD3D/Private/D3DShaderCompiler.cpp
@@ -282,10 +282,18 @@ static FString D3D11CreateShaderCompileCommandLine(
 					)
 			)
 	*/
-	FString BatchFileHeader = TEXT("@ECHO OFF\nSET FXC=\"C:\\Program Files (x86)\\Windows Kits\\10\\bin\\x64\\fxc.exe\"\n"\
-		"IF EXIST %FXC% (\nREM\n) ELSE (\nECHO Couldn't find Windows 10 SDK, falling back to DXSDK...\n"\
-		"SET FXC=\"%DXSDK_DIR%\\Utilities\\bin\\x86\\fxc.exe\"\nIF EXIST %FXC% (\nREM\n) ELSE (\nECHO Couldn't find DXSDK! Exiting...\n"\
-		"GOTO END\n)\n)\n");
+	const FString BatchFileHeader = TEXT(
+		"@ECHO OFF\n"\
+		"IF \"%FXC%\" == \"\" SET FXC=\"C:\\Program Files (x86)\\Windows Kits\\10\\bin\\x64\\fxc.exe\"\n"\
+		"IF NOT EXIST %FXC% (\n"\
+		"\tECHO Couldn't find Windows 10 SDK, falling back to DXSDK...\n"\
+		"\tSET FXC=\"%DXSDK_DIR%\\Utilities\\bin\\x86\\fxc.exe\"\n"\
+		"\tIF NOT EXIST % FXC % (\n"\
+		"\t\tECHO Couldn't find DXSDK! Exiting...\n"\
+		"\t\tGOTO END\n"\
+		"\t)\n"\
+		")\n"
+	);
 	return BatchFileHeader + FXCCommandline + TEXT("\n:END\nREM\n");
 }
 
