- Fix for Streamdeck devices showing up under WinGDK platform

diff --git a/Engine/Source/Developer/Windows/WindowsTargetPlatform/Private/GenericWindowsTargetPlatform.h b/Engine/Source/Developer/Windows/WindowsTargetPlatform/Private/GenericWindowsTargetPlatform.h
--- a/Engine/Source/Developer/Windows/WindowsTargetPlatform/Private/GenericWindowsTargetPlatform.h
+++ b/Engine/Source/Developer/Windows/WindowsTargetPlatform/Private/GenericWindowsTargetPlatform.h
@@ -47,8 +47,12 @@ public:
 		// only add local device if actually running on Windows
 		LocalDevice = MakeShareable(new TTargetDevice(*this));
 
-		// Check if we have any SteamDeck devices around
-		SteamDevices = TSteamDeckDevice<TLocalPcTargetDevice<true>>::DiscoverDevices(*this, TEXT("Proton"));
+		// quick solution to not having WinGDK steamdeck devices
+		if (this->PlatformName().StartsWith(TEXT("Windows")))
+		{
+			// Check if we have any SteamDeck devices around
+			SteamDevices = TSteamDeckDevice<TLocalPcTargetDevice<true>>::DiscoverDevices(*this, TEXT("Proton"));
+		}
 #endif
 
 	#if WITH_ENGINE
