Force DXC precompilation when FXC contains an "internal error: out of memory" report.

#rnx

diff --git a/Engine/Source/Developer/Windows/ShaderFormatD3D/Private/D3DShaderCompiler.cpp b/Engine/Source/Developer/Windows/ShaderFormatD3D/Private/D3DShaderCompiler.cpp
--- a/Engine/Source/Developer/Windows/ShaderFormatD3D/Private/D3DShaderCompiler.cpp
+++ b/Engine/Source/Developer/Windows/ShaderFormatD3D/Private/D3DShaderCompiler.cpp
@@ -636,6 +636,21 @@ static void PatchHlslForPrecompilation(TArray<ANSICHAR>& HlslSource)
 	HlslSource[HlslSourceString.Len()] = '\0';
 }
 
+// Returns whether the specified D3D compiler error buffer contains any internal error messages, e.g. "internal error: out of memory"
+static bool CompileErrorsContainInternalError(ID3DBlob* Errors)
+{
+	if (Errors)
+	{
+		void* ErrorBuffer = Errors->GetBufferPointer();
+		if (ErrorBuffer)
+		{
+			const FStringView ErrorString = ANSI_TO_TCHAR(ErrorBuffer);
+			return ErrorString.Contains(TEXT("internal error: "));
+		}
+	}
+	return false;
+}
+
 // Generate the dumped usf file; call the D3D compiler, gather reflection information and generate the output data
 bool CompileAndProcessD3DShaderFXC(FString& PreprocessedShaderSource, const FString& CompilerPath,
 	uint32 CompileFlags,
@@ -713,7 +728,7 @@ bool CompileAndProcessD3DShaderFXC(FString& PreprocessedShaderSource, const FStr
 		}
 
 		// Some materials give FXC a hard time to optimize and the compiler fails with an internal error.
-		if (bPrecompileWithDXC || Result == HRESULT_FROM_WIN32(ERROR_ARITHMETIC_OVERFLOW) || Result == E_OUTOFMEMORY || bException)
+		if (bPrecompileWithDXC || Result == HRESULT_FROM_WIN32(ERROR_ARITHMETIC_OVERFLOW) || Result == E_OUTOFMEMORY || bException || (Result != S_OK && CompileErrorsContainInternalError(Errors.GetReference())))
 		{
 			if (bPrecompileWithDXC)
 			{
