remove unused bSupportDX11TextureFormats in WindowsTargetPlatform

diff --git a/Engine/Source/Developer/Windows/WindowsTargetPlatform/Private/GenericWindowsTargetPlatform.h b/Engine/Source/Developer/Windows/WindowsTargetPlatform/Private/GenericWindowsTargetPlatform.h
--- a/Engine/Source/Developer/Windows/WindowsTargetPlatform/Private/GenericWindowsTargetPlatform.h
+++ b/Engine/Source/Developer/Windows/WindowsTargetPlatform/Private/GenericWindowsTargetPlatform.h
@@ -71,31 +71,6 @@ public:
 		static FName NAME_SF_VULKAN_ES31(TEXT("SF_VULKAN_ES31"));
 		static FName NAME_OPENGL_150_ES3_1(TEXT("GLSL_150_ES31"));
 
-		bool bSupportDX11TextureFormats = true;
-		bSupportCompressedVolumeTexture = true;
-
-		for (FName TargetedShaderFormat : TargetedShaderFormats)
-		{
-			if (TargetedShaderFormat == NAME_PCD3D_SM6 ||
-				TargetedShaderFormat == NAME_PCD3D_SM5 ||
-				TargetedShaderFormat == NAME_VULKAN_SM5 ||
-				TargetedShaderFormat == NAME_PCD3D_ES3_1)
-			{
-				// DX11 formats okay
-
-				// technically PCD3D_ES3_1 might not support DX11 formats, but in UE5 we require them
-				//   PCD3D_ES3_1 is used as a low-spec version of DX11
-			}
-			else
-			{
-				// some TargetedShaderFormat doesn't support DX11 formats
-				// must turn it off altogether
-				bSupportDX11TextureFormats = false;
-				break;
-			}
-		}
-		check(bSupportDX11TextureFormats);
-
 		// If we are targeting ES3.1, we also must cook encoded HDR reflection captures
 		bRequiresEncodedHDRReflectionCaptures =	TargetedShaderFormats.Contains(NAME_SF_VULKAN_ES31)
 												|| TargetedShaderFormats.Contains(NAME_OPENGL_150_ES3_1)
@@ -322,7 +297,7 @@ public:
 	{
 		if (!TProperties::IsServerOnly())
 		{
-			GetDefaultTextureFormatNamePerLayer(OutFormats.AddDefaulted_GetRef(), this, InTexture, bSupportCompressedVolumeTexture, 4, true);
+			GetDefaultTextureFormatNamePerLayer(OutFormats.AddDefaulted_GetRef(), this, InTexture, true, 4, true);
 		}
 	}
 
@@ -428,8 +403,6 @@ private:
 	// True if the project requires encoded HDR reflection captures
 	bool bRequiresEncodedHDRReflectionCaptures;
 
-	// True if the project supports only compressed volume texture formats.
-	bool bSupportCompressedVolumeTexture;
 #endif // WITH_ENGINE
 
 };
