Expand handling of D3D internal compiler errors.

#rnx

diff --git a/Engine/Source/Developer/Windows/ShaderFormatD3D/Private/D3DShaderCompiler.cpp b/Engine/Source/Developer/Windows/ShaderFormatD3D/Private/D3DShaderCompiler.cpp
--- a/Engine/Source/Developer/Windows/ShaderFormatD3D/Private/D3DShaderCompiler.cpp
+++ b/Engine/Source/Developer/Windows/ShaderFormatD3D/Private/D3DShaderCompiler.cpp
@@ -645,7 +645,7 @@
 		if (ErrorBuffer)
 		{
 			const FStringView ErrorString = ANSI_TO_TCHAR(ErrorBuffer);
-			return ErrorString.Contains(TEXT("internal error: "));
+			return ErrorString.Contains(TEXT("internal error:")) || ErrorString.Contains(TEXT("Internal Compiler Error:"));
 		}
 	}
 	return false;
