Include Reflection data with SM6 embedded PDBs. This makes them way more useful in PIX.

diff --git a/Engine/Source/Developer/Windows/ShaderFormatD3D/Private/D3DShaderCompilerDXC.cpp b/Engine/Source/Developer/Windows/ShaderFormatD3D/Private/D3DShaderCompilerDXC.cpp
--- a/Engine/Source/Developer/Windows/ShaderFormatD3D/Private/D3DShaderCompilerDXC.cpp
+++ b/Engine/Source/Developer/Windows/ShaderFormatD3D/Private/D3DShaderCompilerDXC.cpp
@@ -264,7 +264,7 @@ static bool RemoveContainerReflection(dxc::DxcDllSupport& DxcDllHelper, TRefCoun
 	
 	// Try and remove both the PDB & Reflection Data
 	bool bPDBRemoved = bRemovePDB && SUCCEEDED(Builder->RemovePart(DXC_PART_PDB));
-	bool bReflectionDataRemoved = SUCCEEDED(Builder->RemovePart(DXC_PART_REFLECTION_DATA));
+	bool bReflectionDataRemoved = bRemovePDB && SUCCEEDED(Builder->RemovePart(DXC_PART_REFLECTION_DATA));
 	if (bPDBRemoved || bReflectionDataRemoved)
 	{
 		VERIFYHRESULT(Builder->SerializeContainer(Result.GetInitReference()));
diff --git a/Engine/Source/Developer/Windows/ShaderFormatD3D/Private/ShaderFormatD3D.cpp b/Engine/Source/Developer/Windows/ShaderFormatD3D/Private/ShaderFormatD3D.cpp
--- a/Engine/Source/Developer/Windows/ShaderFormatD3D/Private/ShaderFormatD3D.cpp
+++ b/Engine/Source/Developer/Windows/ShaderFormatD3D/Private/ShaderFormatD3D.cpp
@@ -20,7 +20,7 @@ class FShaderFormatD3D : public IShaderFormat
 		UE_SHADER_PCD3D_SHARED_VER = 2,
 
 		/** Version for shader format, this becomes part of the DDC key. */
-		UE_SHADER_PCD3D_SM6_VER = UE_SHADER_PCD3D_SHARED_VER + 3,
+		UE_SHADER_PCD3D_SM6_VER = UE_SHADER_PCD3D_SHARED_VER + 4,
 		UE_SHADER_PCD3D_SM5_VER = UE_SHADER_PCD3D_SHARED_VER + 9,
 		UE_SHADER_PCD3D_ES3_1_VER = UE_SHADER_PCD3D_SHARED_VER + 8,
 		UE_SHADER_D3D_ES3_1_HOLOLENS_VER = UE_SHADER_PCD3D_ES3_1_VER,
