Fixed a reliability issue with LiveCoding where it could crash on startup.

diff --git a/Engine/Source/Developer/Windows/LiveCoding/Private/External/LC_Process.cpp b/Engine/Source/Developer/Windows/LiveCoding/Private/External/LC_Process.cpp
--- a/Engine/Source/Developer/Windows/LiveCoding/Private/External/LC_Process.cpp
+++ b/Engine/Source/Developer/Windows/LiveCoding/Private/External/LC_Process.cpp
@@ -500,6 +500,9 @@ void Process::FlushInstructionCache(Handle handle, void* address, size_t size)
 Process::Environment Process::CreateEnvironment(Handle handle)
 {
 	const void* processEnvironment = nullptr;
+	// EPIC BEGIN MOD
+	SIZE_T processEnvironmentSize = 0;
+	// EPIC END MOD
 
 	const bool isWow64 = IsWoW64(handle);
 	if (!isWow64)
@@ -513,6 +516,9 @@ Process::Environment Process::CreateEnvironment(Handle handle)
 		const WindowsInternals::RTL_USER_PROCESS_PARAMETERS parameters = ReadProcessMemory<WindowsInternals::RTL_USER_PROCESS_PARAMETERS>(handle, peb.ProcessParameters);
 
 		processEnvironment = parameters.Environment;
+		// EPIC BEGIN MOD
+		processEnvironmentSize = parameters.EnvironmentSize;
+		// EPIC END MOD
 	}
 	else
 	{
@@ -525,6 +531,9 @@ Process::Environment Process::CreateEnvironment(Handle handle)
 		const WindowsInternals::RTL_USER_PROCESS_PARAMETERS32 parameters32 = ReadProcessMemory<WindowsInternals::RTL_USER_PROCESS_PARAMETERS32>(handle, pointer::FromInteger<const void*>(peb32.ProcessParameters32));
 
 		processEnvironment = pointer::FromInteger<const void*>(parameters32.Environment);
+		// EPIC BEGIN MOD
+		processEnvironmentSize = parameters32.EnvironmentSize;
+		// EPIC END MOD
 	}
 
 	if (!processEnvironment)
@@ -532,19 +541,9 @@ Process::Environment Process::CreateEnvironment(Handle handle)
 		return Environment { 0u, nullptr };
 	}
 
-	// query the size of the page(s) the environment is stored in
-	::MEMORY_BASIC_INFORMATION memoryInfo = {};
-	const SIZE_T bytesInBuffer = ::VirtualQueryEx(+handle, processEnvironment, &memoryInfo, sizeof(::MEMORY_BASIC_INFORMATION));
-	if (bytesInBuffer == 0u)
-	{
-		// operation failed, bail out
-		return Environment { 0u, nullptr };
-	}
-
-	Environment environment = {};
-	environment.size = memoryInfo.RegionSize - (reinterpret_cast<uintptr_t>(processEnvironment) - reinterpret_cast<uintptr_t>(memoryInfo.BaseAddress));
-	environment.data = ::malloc(environment.size);
-
+	// EPIC BEGIN MOD
+	Environment environment = { processEnvironmentSize, ::malloc(processEnvironmentSize) };
+	// EPIC END MOD
 	ReadProcessMemory(handle, processEnvironment, environment.data, environment.size);
 
 	return environment;
diff --git a/Engine/Source/Developer/Windows/LiveCoding/Private/External/LC_WindowsInternals.h b/Engine/Source/Developer/Windows/LiveCoding/Private/External/LC_WindowsInternals.h
--- a/Engine/Source/Developer/Windows/LiveCoding/Private/External/LC_WindowsInternals.h
+++ b/Engine/Source/Developer/Windows/LiveCoding/Private/External/LC_WindowsInternals.h
@@ -63,6 +63,29 @@ namespace WindowsInternals
 		PWSTR Buffer;
 	};
 
+	// BEGIN EPIC MOD
+	struct NT_UNICODE_STRING32
+	{
+		USHORT Length;
+		USHORT MaximumLength;
+		ULONG Buffer;
+	};
+
+	struct NT_ANSI_STRING
+	{
+		USHORT Length;
+		USHORT MaximumLength;
+		PSTR Buffer;
+	};
+
+	struct NT_ANSI_STRING32
+	{
+		USHORT Length;
+		USHORT MaximumLength;
+		ULONG Buffer;
+	};
+	// END EPIC MOD
+
 	// https://www.geoffchappell.com/studies/windows/km/ntoskrnl/api/ex/sysinfo/process.htm
 	// found in <winternl.h>
 	struct NT_SYSTEM_PROCESS_INFORMATION
@@ -111,6 +134,24 @@ namespace WindowsInternals
 		ProcessWow64Information = 26
 	};
 
+	// BEGIN EPIC MOD
+	struct RLT_DRIVE_LETTER_CURDIR
+	{
+		WORD Flags;
+		WORD Length;
+		ULONG TimeStemp;
+		NT_ANSI_STRING DosPath;
+	};
+
+	struct RLT_DRIVE_LETTER_CURDIR32
+	{
+		WORD Flags;
+		WORD Length;
+		ULONG TimeStemp;
+		NT_ANSI_STRING32 DosPath;
+	};
+	// END EPIC MOD
+
 	// https://docs.microsoft.com/en-us/windows/desktop/api/winternl/ns-winternl-_rtl_user_process_parameters
 	// https://www.nirsoft.net/kernel_struct/vista/RTL_USER_PROCESS_PARAMETERS.html
 	// found in <winternl.h>
@@ -121,6 +162,12 @@ namespace WindowsInternals
 		NT_UNICODE_STRING ImagePathName;
 		NT_UNICODE_STRING CommandLine;
 		PWSTR Environment;
+		// EPIC BEGIN MOD
+		ULONG Reserved3[9];
+		NT_UNICODE_STRING Reserved4[4];
+		RLT_DRIVE_LETTER_CURDIR CurrentDirectories[32];
+		ULONG EnvironmentSize;
+		// EPIC END MOD
 	};
 
 	// similar to RTL_USER_PROCESS_PARAMETERS, altered to behave as a struct containing 32-bit pointers in a 64-bit environment
@@ -128,6 +175,12 @@ namespace WindowsInternals
 	{
 		char Reserved[72];
 		ULONG Environment;
+		// EPIC BEGIN MOD
+		ULONG Reserved3[9];
+		NT_UNICODE_STRING32 Reserved4[4];
+		RLT_DRIVE_LETTER_CURDIR32 CurrentDirectories[32];
+		ULONG EnvironmentSize;
+		// EPIC END MOD
 	};
 
 	// https://www.geoffchappell.com/studies/windows/win32/ntdll/structs/ldr_data_table_entry.htm
