Removing the "equality comparison with extraneous parentheses" warning from DXC compiles.

diff --git a/Engine/Source/Developer/Windows/ShaderFormatD3D/Private/D3DShaderCompiler.inl b/Engine/Source/Developer/Windows/ShaderFormatD3D/Private/D3DShaderCompiler.inl
--- a/Engine/Source/Developer/Windows/ShaderFormatD3D/Private/D3DShaderCompiler.inl
+++ b/Engine/Source/Developer/Windows/ShaderFormatD3D/Private/D3DShaderCompiler.inl
@@ -190,6 +190,9 @@ public:
 
 		// Reflection will be removed later, otherwise the disassembly won't contain variables
 		//ExtraArguments.Add(TEXT("/Qstrip_reflect"));
+
+		// disable undesired warnings
+		ExtraArguments.Add(TEXT("-Wno-parentheses-equality"));
 	}
 
 	inline FString GetDumpDebugInfoPath() const
diff --git a/Engine/Source/Developer/Windows/ShaderFormatD3D/Private/ShaderFormatD3D.cpp b/Engine/Source/Developer/Windows/ShaderFormatD3D/Private/ShaderFormatD3D.cpp
--- a/Engine/Source/Developer/Windows/ShaderFormatD3D/Private/ShaderFormatD3D.cpp
+++ b/Engine/Source/Developer/Windows/ShaderFormatD3D/Private/ShaderFormatD3D.cpp
@@ -20,8 +20,8 @@ class FShaderFormatD3D : public IShaderFormat
 		UE_SHADER_PCD3D_SHARED_VER = 2,
 
 		/** Version for shader format, this becomes part of the DDC key. */
-		UE_SHADER_PCD3D_SM6_VER = UE_SHADER_PCD3D_SHARED_VER + 4,
-		UE_SHADER_PCD3D_SM5_VER = UE_SHADER_PCD3D_SHARED_VER + 9,
+		UE_SHADER_PCD3D_SM6_VER = UE_SHADER_PCD3D_SHARED_VER + 5,
+		UE_SHADER_PCD3D_SM5_VER = UE_SHADER_PCD3D_SHARED_VER + 10,
 		UE_SHADER_PCD3D_ES3_1_VER = UE_SHADER_PCD3D_SHARED_VER + 8,
 		UE_SHADER_D3D_ES3_1_HOLOLENS_VER = UE_SHADER_PCD3D_ES3_1_VER,
 	};
