Add "TargetedRHIs" to the list of settings evaluated in GetRayTracingShaderFormats()

* This is necessary because some (pre UE 5.1) projects use this instead of D3D12TargetedShaderFormats, etc.

diff --git a/Engine/Source/Developer/Windows/WindowsTargetPlatform/Private/GenericWindowsTargetPlatform.h b/Engine/Source/Developer/Windows/WindowsTargetPlatform/Private/GenericWindowsTargetPlatform.h
--- a/Engine/Source/Developer/Windows/WindowsTargetPlatform/Private/GenericWindowsTargetPlatform.h
+++ b/Engine/Source/Developer/Windows/WindowsTargetPlatform/Private/GenericWindowsTargetPlatform.h
@@ -286,6 +286,7 @@ public:
 			// Only D3D12 and Vulkan can possibly support ray tracing
 			TCHAR const* RelevantSettings[] = 
 			{
+				TEXT("TargetedRHIs"),
 				TEXT("D3D12TargetedShaderFormats"),
 				TEXT("VulkanTargetedShaderFormats")
 			};
