Always support ray tracing shaders when cooking for SM6

* RT shaders are still excluded from SM5 if project is configured to not cook SM5 shaders for D3D12

diff --git a/Engine/Source/Developer/Windows/WindowsTargetPlatform/Private/GenericWindowsTargetPlatform.h b/Engine/Source/Developer/Windows/WindowsTargetPlatform/Private/GenericWindowsTargetPlatform.h
--- a/Engine/Source/Developer/Windows/WindowsTargetPlatform/Private/GenericWindowsTargetPlatform.h
+++ b/Engine/Source/Developer/Windows/WindowsTargetPlatform/Private/GenericWindowsTargetPlatform.h
@@ -292,6 +292,9 @@ public:
 			};
 
 			GetAllTargetedShaderFormatsInternal(RelevantSettings, OutFormats);
+
+			// We always support ray tracing shaders when cooking for SM6, however we may skip them for SM5 based on project settings.
+			OutFormats.AddUnique(FName(TEXT("PCD3D_SM6")));
 		}
 	}
 
