Fixed issue where live coding would always complain that a plugin wasn't enabled for live coding.  This would happen if the plugin wasn't explicitly enabled by the project but was loaded by another plugin.

diff --git a/Engine/Source/Developer/Windows/LiveCoding/Private/LiveCodingModule.cpp b/Engine/Source/Developer/Windows/LiveCoding/Private/LiveCodingModule.cpp
--- a/Engine/Source/Developer/Windows/LiveCoding/Private/LiveCodingModule.cpp
+++ b/Engine/Source/Developer/Windows/LiveCoding/Private/LiveCodingModule.cpp
@@ -34,6 +34,13 @@
 	#include "UObject/StrongObjectPtr.h"
 #endif
 
+#include "Windows/AllowWindowsPlatformTypes.h"
+#include <psapi.h> // for GetModuleFileNameEx
+#ifdef NT_SUCCESS
+#undef NT_SUCCESS
+#endif
+#include <subauth.h> // for UNICODE_STRING
+
 IMPLEMENT_MODULE(FLiveCodingModule, LiveCoding)
 
 #define LOCTEXT_NAMESPACE "LiveCodingModule"
@@ -136,8 +143,29 @@ private:
 };
 #endif
 
+// Helper structure to load the NTDLL library and work with an API
+struct FNtDllFunction
+{
+	FARPROC Addr;
+
+	FNtDllFunction(const char* Name)
+	{
+		HMODULE NtDll = LoadLibraryW(L"ntdll.dll");
+		check(NtDll);
+		Addr = GetProcAddress(NtDll, Name);
+	}
+
+	template <typename... ArgTypes>
+	unsigned int operator () (ArgTypes... Args)
+	{
+		typedef unsigned int (NTAPI* Prototype)(ArgTypes...);
+		return (Prototype((void*)Addr))(Args...);
+	}
+};
+
 FLiveCodingModule::FLiveCodingModule()
-	: FullEnginePluginsDir(FPaths::ConvertRelativePathToFull(FPaths::EnginePluginsDir()))
+	: FullEngineDir(FPaths::ConvertRelativePathToFull(FPaths::EngineDir()))
+	, FullEnginePluginsDir(FPaths::ConvertRelativePathToFull(FPaths::EnginePluginsDir()))
 	, FullProjectDir(FPaths::ConvertRelativePathToFull(FPaths::ProjectDir()))
 	, FullProjectPluginsDir(FPaths::ConvertRelativePathToFull(FPaths::ProjectPluginsDir()))
 {
@@ -151,6 +179,76 @@ FLiveCodingModule::~FLiveCodingModule()
 
 void FLiveCodingModule::StartupModule()
 {
+	// Register with NT to get dll nitrifications
+	FNtDllFunction RegisterFunc("LdrRegisterDllNotification");
+	RegisterFunc(0, OnDllNotification, this, &CallbackCookie);
+
+	// Get the creating process main executable name.  We want to skip this file when iterating through the dll list
+	FString MainModuleName;
+	{
+		TCHAR Scratch[MAX_PATH];
+		int Length = GetModuleFileNameEx(::GetCurrentProcess(), NULL, Scratch, MAX_PATH);
+		MainModuleName = FString(Length, Scratch);
+	}
+
+	// https://docs.microsoft.com/en-us/windows/win32/api/winternl/ns-winternl-peb_ldr_data
+	typedef struct _PEB_LDR_DATA {
+		BYTE       Reserved1[8];
+		PVOID      Reserved2[3];
+		LIST_ENTRY InMemoryOrderModuleList;
+	} PEB_LDR_DATA, * PPEB_LDR_DATA;
+
+	// https://docs.microsoft.com/en-us/windows/win32/api/winternl/ns-winternl-peb
+	typedef struct _PEB {
+		BYTE                          Reserved1[2];
+		BYTE                          BeingDebugged;
+		BYTE                          Reserved2[1];
+		PVOID                         Reserved3[2];
+		PPEB_LDR_DATA                 Ldr;
+		//...
+	} PEB, * PPEB;
+
+	// https://docs.microsoft.com/en-us/windows/win32/api/winternl/ns-winternl-teb
+	typedef struct _TEB {
+		PVOID Reserved1[12];
+		PPEB  ProcessEnvironmentBlock;
+		//...
+	} TEB, * PTEB;
+
+	// https://docs.microsoft.com/en-us/windows/win32/api/winternl/ns-winternl-peb_ldr_data
+	typedef struct _LDR_DATA_TABLE_ENTRY {
+		PVOID Reserved1[2];
+		LIST_ENTRY InMemoryOrderLinks;
+		PVOID Reserved2[2];
+		PVOID DllBase;
+		PVOID EntryPoint;
+		PVOID Reserved3;
+		UNICODE_STRING FullDllName;
+		//...
+	} LDR_DATA_TABLE_ENTRY, * PLDR_DATA_TABLE_ENTRY;
+
+	// Enumerate already loaded modules.
+	const TEB* ThreadEnvBlock = reinterpret_cast<TEB*>(NtCurrentTeb());
+	const PEB* ProcessEnvBlock = ThreadEnvBlock->ProcessEnvironmentBlock;
+	const LIST_ENTRY* ModuleIter = ProcessEnvBlock->Ldr->InMemoryOrderModuleList.Flink;
+	const LIST_ENTRY* ModuleIterEnd = ModuleIter->Blink;
+	do
+	{
+		const auto& ModuleData = *(LDR_DATA_TABLE_ENTRY*)(ModuleIter - 1);
+		if (ModuleData.DllBase == 0)
+		{
+			break;
+		}
+
+		FString FullPath(ModuleData.FullDllName.Length / sizeof(ModuleData.FullDllName.Buffer[0]), ModuleData.FullDllName.Buffer);
+		if (!FullPath.Equals(MainModuleName, ESearchCase::IgnoreCase))
+		{
+			FPaths::NormalizeFilename(FullPath);
+			OnDllLoaded(FullPath);
+		}
+		ModuleIter = ModuleIter->Flink;
+	} while (ModuleIter != ModuleIterEnd);
+
 	Settings = GetMutableDefault<ULiveCodingSettings>();
 
 	IConsoleManager& ConsoleManager = IConsoleManager::Get();
@@ -246,6 +344,13 @@ void FLiveCodingModule::ShutdownModule()
 	ConsoleManager.UnregisterConsoleObject(ConsolePathVariable);
 	ConsoleManager.UnregisterConsoleObject(CompileCommand);
 	ConsoleManager.UnregisterConsoleObject(EnableCommand);
+
+	// Unregister from the dll notifications
+	if (CallbackCookie)
+	{
+		FNtDllFunction UnregisterFunc("LdrUnregisterDllNotification");
+		UnregisterFunc(CallbackCookie);
+	}
 }
 
 void FLiveCodingModule::EnableByDefault(bool bEnable)
@@ -746,7 +851,7 @@ bool FLiveCodingModule::StartLiveCoding()
 			}
 		}
 
-		ensure(CreateMutex(NULL, Windows::FALSE, *MutexName));
+		ensure(CreateMutex(NULL, FALSE, *MutexName));
 
 		// Configure all the current modules. For non-commandlets, schedule it to be done in the first Tick() so we can batch everything together.
 		if (IsRunningCommandlet())
@@ -758,9 +863,6 @@ bool FLiveCodingModule::StartLiveCoding()
 			bUpdateModulesInTick = true;
 		}
 
-		// Register a delegate to listen for new modules loaded from this point onwards
-		ModulesChangedDelegateHandle = FModuleManager::Get().OnModulesChanged().AddRaw(this, &FLiveCodingModule::OnModulesChanged);
-
 		// Mark it as started
 		bStarted = true;
 		bEnabledForSession = true;
@@ -770,6 +872,12 @@ bool FLiveCodingModule::StartLiveCoding()
 
 void FLiveCodingModule::UpdateModules()
 {
+	TArray<ModuleChange> Changes;
+	{
+		FScopeLock lock(&ModuleChangeCs);
+		Swap(Changes, ModuleChanges);
+	}
+
 	if (bEnabledForSession)
 	{
 #if IS_MONOLITHIC
@@ -777,31 +885,28 @@ void FLiveCodingModule::UpdateModules()
 		verify(GetModuleFileName(hInstance, FullFilePath, UE_ARRAY_COUNT(FullFilePath)));
 		LppEnableModule(FullFilePath);
 #else
-		TArray<FModuleStatus> ModuleStatuses;
-		FModuleManager::Get().QueryModules(ModuleStatuses);
-
 		TArray<FString> EnableModules;
-		for (const FModuleStatus& ModuleStatus : ModuleStatuses)
+		for (const ModuleChange& Change : Changes)
 		{
-			if (ModuleStatus.bIsLoaded)
+			if (Change.bLoaded)
 			{
-				FName ModuleName(*ModuleStatus.Name);
-				if (!ConfiguredModules.Contains(ModuleName))
+				ConfiguredModules.Add(Change.FullName);
+				FString FullFilePath(Change.FullName.ToString());
+				if (ShouldPreloadModule(FullFilePath))
 				{
-					FString FullFilePath = FPaths::ConvertRelativePathToFull(ModuleStatus.FilePath);
-					if (ShouldPreloadModule(ModuleName, FullFilePath))
-					{
-						EnableModules.Add(FullFilePath);
-					}
-					else
-					{
-						TRACE_CPUPROFILER_EVENT_SCOPE(LppEnableLazyLoadedModule);
-						void* LppEnableLazyLoadedModuleToken = LppEnableLazyLoadedModule(*FullFilePath);
-						LppPendingTokens.Add(LppEnableLazyLoadedModuleToken);
-					}
-					ConfiguredModules.Add(ModuleName);
+					EnableModules.Add(FullFilePath);
+				}
+				else
+				{
+					TRACE_CPUPROFILER_EVENT_SCOPE(LppEnableLazyLoadedModule);
+					void* LppEnableLazyLoadedModuleToken = LppEnableLazyLoadedModule(*FullFilePath);
+					LppPendingTokens.Add(LppEnableLazyLoadedModuleToken);
 				}
 			}
+			else
+			{
+				ConfiguredModules.Remove(Change.FullName);
+			}
 		}
 
 		if (EnableModules.Num() > 0)
@@ -840,27 +945,27 @@ void FLiveCodingModule::OnModulesChanged(FName ModuleName, EModuleChangeReason R
 #endif
 }
 
-bool FLiveCodingModule::ShouldPreloadModule(const FName& Name, const FString& FullFilePath) const
+bool FLiveCodingModule::ShouldPreloadModule(const FString& FullFilePath) const
 {
-	// For the hooks to work properly, we always have to load the live coding module
-	if (Name == TEXT(LIVE_CODING_MODULE_NAME))
-	{
-		return true;
-	}
 
-	if (Settings->PreloadNamedModules.Contains(Name))
+	// Perform some name based checks
 	{
-		return true;
+		FString NameStr = FPaths::GetBaseFilename(FullFilePath, true);
+		FName Name(NameStr, FNAME_Find);
+		if (Name != NAME_None && Settings->PreloadNamedModules.Contains(Name))
+		{
+			return true;
+		}
 	}
 
-	if (FullFilePath.StartsWith(FullProjectDir))
+	if (FullFilePath.StartsWith(FullProjectDir, ESearchCase::IgnoreCase))
 	{
 		if (Settings->bPreloadProjectModules == Settings->bPreloadProjectPluginModules)
 		{
 			return Settings->bPreloadProjectModules;
 		}
 
-		if(FullFilePath.StartsWith(FullProjectPluginsDir))
+		if(FullFilePath.StartsWith(FullProjectPluginsDir, ESearchCase::IgnoreCase))
 		{
 			return Settings->bPreloadProjectPluginModules;
 		}
@@ -881,7 +986,7 @@ bool FLiveCodingModule::ShouldPreloadModule(const FName& Name, const FString& Fu
 			return Settings->bPreloadEngineModules;
 		}
 
-		if(FullFilePath.StartsWith(FullEnginePluginsDir))
+		if(FullFilePath.StartsWith(FullEnginePluginsDir, ESearchCase::IgnoreCase))
 		{
 			return Settings->bPreloadEnginePluginModules;
 		}
@@ -926,6 +1031,97 @@ bool FLiveCodingModule::AutomaticallyCompileNewClasses() const
 	return Settings->bAutomaticallyCompileNewClasses;
 }
 
+void FLiveCodingModule::OnDllNotification(unsigned int Reason, const void* DataPtr, void* Context)
+{
+	FLiveCodingModule* This = reinterpret_cast<FLiveCodingModule*>(Context);
+
+	enum
+	{
+		LDR_DLL_NOTIFICATION_REASON_LOADED = 1,
+		LDR_DLL_NOTIFICATION_REASON_UNLOADED = 2,
+	};
+
+	struct FNotificationData
+	{
+		uint32 Flags;
+		const UNICODE_STRING& FullPath;
+		const UNICODE_STRING& BaseName;
+		UPTRINT	Base;
+	};
+	const auto& Data = *(FNotificationData*)DataPtr;
+	FString FullPath(Data.FullPath.Length / sizeof(Data.FullPath.Buffer[0]), Data.FullPath.Buffer);
+	FPaths::NormalizeFilename(FullPath);
+
+	switch (Reason)
+	{
+	case LDR_DLL_NOTIFICATION_REASON_LOADED:
+		This->OnDllLoaded(FullPath);
+		break;
+
+	case LDR_DLL_NOTIFICATION_REASON_UNLOADED:
+		This->OnDllUnloaded(FullPath);
+		break;
+	}
+}
+
+void FLiveCodingModule::OnDllLoaded(const FString& FullPath)
+{
+	if (IsUEDll(FullPath))
+	{
+		FScopeLock lock(&ModuleChangeCs);
+		ModuleChanges.Emplace(ModuleChange{ FName(FullPath), true });
+	}
+}
+
+void FLiveCodingModule::OnDllUnloaded(const FString& FullPath)
+{
+	if (IsUEDll(FullPath))
+	{
+		FScopeLock lock(&ModuleChangeCs);
+		ModuleChanges.Emplace(ModuleChange{ FName(FullPath), false });
+	}
+}
+
+
+bool FLiveCodingModule::IsUEDll(const FString& FullPath)
+{
+	// Ignore patches.
+	if (IsPatchDll(FullPath))
+	{
+		return false;
+	}
+
+	// Dll must be in the engine or project dir
+	if (!FullPath.StartsWith(FullEngineDir, ESearchCase::IgnoreCase) &&
+		!FullPath.StartsWith(FullEnginePluginsDir, ESearchCase::IgnoreCase) &&
+		!FullPath.StartsWith(FullProjectDir, ESearchCase::IgnoreCase) &&
+		!FullPath.StartsWith(FullProjectPluginsDir, ESearchCase::IgnoreCase))
+	{
+		return false;
+	}
+	return true;
+}
+
+bool FLiveCodingModule::IsPatchDll(const FString& FullPath)
+{
+	// If the Dll ends with ".patch_#.dll", then ignore it
+	FString Name = FPaths::GetBaseFilename(FullPath, true);
+	FString Extension = FPaths::GetExtension(Name);
+	if (!Extension.StartsWith("patch_", ESearchCase::IgnoreCase))
+	{
+		return false;
+	}
+
+	for (int Index = 6; Index < Extension.Len(); ++Index)
+	{
+		if (!FChar::IsDigit(Extension[Index]))
+		{
+			return false;
+		}
+	}
+	return true;
+}
+
 // Invoked from LC_ClientCommandActions
 void LiveCodingBeginPatch()
 {
diff --git a/Engine/Source/Developer/Windows/LiveCoding/Private/LiveCodingModule.h b/Engine/Source/Developer/Windows/LiveCoding/Private/LiveCodingModule.h
--- a/Engine/Source/Developer/Windows/LiveCoding/Private/LiveCodingModule.h
+++ b/Engine/Source/Developer/Windows/LiveCoding/Private/LiveCodingModule.h
@@ -49,6 +49,11 @@ public:
 
 private:
 	void AttemptSyncLivePatching();
+	static void OnDllNotification(unsigned int Reason, const void* DataPtr, void* Context);
+	void OnDllLoaded(const FString& FullPath);
+	void OnDllUnloaded(const FString& FullPath);
+	bool IsUEDll(const FString& FullPath);
+	bool IsPatchDll(const FString& FullPath);
 
 private:
 	ULiveCodingSettings* Settings;
@@ -63,9 +68,11 @@ private:
 	ELiveCodingCompileResult LastResults = ELiveCodingCompileResult::Success;
 	TSet<FName> ConfiguredModules;
 	TArray<void*> LppPendingTokens;
+	void* CallbackCookie = nullptr;
 
 	FText EnableErrorText;
 
+	const FString FullEngineDir;
 	const FString FullEnginePluginsDir;
 	const FString FullProjectDir;
 	const FString FullProjectPluginsDir;
@@ -78,6 +85,14 @@ private:
 	FDelegateHandle ModulesChangedDelegateHandle;
 	FOnPatchCompleteDelegate OnPatchCompleteDelegate;
 
+	struct ModuleChange
+	{
+		FName FullName;
+		bool bLoaded;
+	};
+	FCriticalSection ModuleChangeCs;
+	TArray<ModuleChange> ModuleChanges;
+
 #if WITH_EDITOR
 	TUniquePtr<FReload> Reload;
 #else
@@ -90,7 +105,7 @@ private:
 
 	void UpdateModules();
 
-	bool ShouldPreloadModule(const FName& Name, const FString& FullFilePath) const;
+	bool ShouldPreloadModule(const FString& FullFilePath) const;
 
 	bool IsReinstancingEnabled() const;
 
