Updating LiveCoding to Live++ 1.6.10

diff --git a/Engine/Source/Developer/Windows/LiveCoding/Private/External/LPP_API.h b/Engine/Source/Developer/Windows/LiveCoding/Private/External/LPP_API.h
--- a/Engine/Source/Developer/Windows/LiveCoding/Private/External/LPP_API.h
+++ b/Engine/Source/Developer/Windows/LiveCoding/Private/External/LPP_API.h
@@ -105,7 +105,7 @@
 /******************************************************************************/
 
 // version string
-#define LPP_VERSION "1.6.4"
+#define LPP_VERSION "1.6.10"
 
 // macros to temporarily enable/disable optimizations
 #ifdef __clang__
diff --git a/Engine/Source/Developer/Windows/LiveCodingServer/Private/External/LC_Coff.cpp b/Engine/Source/Developer/Windows/LiveCodingServer/Private/External/LC_Coff.cpp
--- a/Engine/Source/Developer/Windows/LiveCodingServer/Private/External/LC_Coff.cpp
+++ b/Engine/Source/Developer/Windows/LiveCodingServer/Private/External/LC_Coff.cpp
@@ -13,6 +13,7 @@
 #include "LC_MemoryStream.h"
 #include "LC_UniqueId.h"
 #include "LC_MemoryMappedFile.h"
+#include "LC_Assembler.h"
 // BEGIN EPIC MOD
 #include "LC_Allocators.h"
 #include "LC_Assert.h"
@@ -963,6 +964,30 @@ namespace coff
 	}
 
 
+	template <typename RawCoffType>
+	static void PatchFunctionSymbol(RawCoffType* rawCoff, size_t symbolIndex, uint32_t sectionIndex)
+	{
+		LC_LOG_DEV("Patching function symbol %d (%s)", symbolIndex, rawCoff->stringTable[symbolIndex].c_str());
+
+		coff::RawSection& section = rawCoff->sections[sectionIndex];
+		uint8_t* code = static_cast<uint8_t*>(section.data);
+		code[0] = Assembler::Opcode::RET[0];
+	}
+
+
+	void PatchFunctionSymbol(RawCoff* rawCoff, size_t symbolIndex, uint32_t sectionIndex)
+	{
+		if (rawCoff->type == RawCoffRegular::TYPE)
+		{
+			PatchFunctionSymbol(static_cast<RawCoffRegular*>(rawCoff), symbolIndex, sectionIndex);
+		}
+		else if (rawCoff->type == RawCoffBigObj::TYPE)
+		{
+			PatchFunctionSymbol(static_cast<RawCoffBigObj*>(rawCoff), symbolIndex, sectionIndex);
+		}
+	}
+
+
 	template <typename RawCoffType>
 	static void RemoveSymbol(RawCoffType* rawCoff, size_t symbolIndex, SymbolRemovalStrategy::Enum removalStrategy)
 	{
diff --git a/Engine/Source/Developer/Windows/LiveCodingServer/Private/External/LC_Coff.h b/Engine/Source/Developer/Windows/LiveCodingServer/Private/External/LC_Coff.h
--- a/Engine/Source/Developer/Windows/LiveCodingServer/Private/External/LC_Coff.h
+++ b/Engine/Source/Developer/Windows/LiveCodingServer/Private/External/LC_Coff.h
@@ -392,6 +392,9 @@ namespace coff
 	void ReplaceLinkerDirectives(RawCoff* rawCoff, const types::vector<std::string>& directives);
 
 
+	// patches a function symbol to return immediately
+	void PatchFunctionSymbol(RawCoff* rawCoff, size_t symbolIndex, uint32_t sectionIndex);
+
 	// removes a symbol
 	void RemoveSymbol(RawCoff* rawCoff, size_t symbolIndex, SymbolRemovalStrategy::Enum removalStrategy);
 
diff --git a/Engine/Source/Developer/Windows/LiveCodingServer/Private/External/LC_DirectoryCache.cpp b/Engine/Source/Developer/Windows/LiveCodingServer/Private/External/LC_DirectoryCache.cpp
--- a/Engine/Source/Developer/Windows/LiveCodingServer/Private/External/LC_DirectoryCache.cpp
+++ b/Engine/Source/Developer/Windows/LiveCodingServer/Private/External/LC_DirectoryCache.cpp
@@ -59,7 +59,6 @@ void DirectoryCache::RestartNotifications(void)
 	for (auto it : m_directories)
 	{
 		Directory* dir = it.second;
-		dir->changeNotification.Check(0u);
 		dir->hadChange = false;
 	}
 }
diff --git a/Engine/Source/Developer/Windows/LiveCodingServer/Private/External/LC_LiveModule.cpp b/Engine/Source/Developer/Windows/LiveCodingServer/Private/External/LC_LiveModule.cpp
--- a/Engine/Source/Developer/Windows/LiveCodingServer/Private/External/LC_LiveModule.cpp
+++ b/Engine/Source/Developer/Windows/LiveCodingServer/Private/External/LC_LiveModule.cpp
@@ -72,16 +72,19 @@ namespace
 
 	static const void* GetUpperBoundIn4GBRange(const void* moduleBase)
 	{
+		const uint64_t HighestPossibleAddress = 0x00007FFFFFFF0000ull;
+		const uint64_t HighestPossibleAddressThreshold = HighestPossibleAddress - 2ull * 1024ull * 1024ull * 1024ull;
+
 		const uint64_t base = pointer::AsInteger<uint64_t>(moduleBase);
 
 		// make sure we don't overflow
-		if (base <= 0xFFFFFFFF7FFFFFFFull)
+		if (base <= HighestPossibleAddressThreshold)
 		{
-			return pointer::FromInteger<const void*>(base + 0x80000000ull);
+			return pointer::FromInteger<const void*>(base + 2ull * 1024ull * 1024ull * 1024ull);
 		}
 
 		// operation would overflow
-		return pointer::FromInteger<const void*>(0xFFFFFFFFFFFFFFFFull);
+		return pointer::FromInteger<const void*>(HighestPossibleAddress);
 	}
 #endif
 
@@ -2891,6 +2894,19 @@ LiveModule::ErrorType::Enum LiveModule::Update(FileAttributeCache* fileCache, Di
 						strippedSymbols.insert(symbolName);
 					}
 				}
+
+				// Clang workaround: in certain cases, we don't find the _GLOBAL__sub_I_ dynamic initializer symbol, which can lead to
+				// a crash in user code, unfortunately.
+				// as a workaround, we patch the code of the function itself to return immediately, so that even when it is called, it doesn't do anything.
+				if (usesClang)
+				{
+					const ImmutableString& symbolName = coff::GetSymbolName(rawCoff, i);
+					const bool isClangDynamicInitializer = (string::StartsWith(symbolName.c_str(), "_GLOBAL__sub_I_"));
+					if (isClangDynamicInitializer)
+					{
+						coff::PatchFunctionSymbol(rawCoff, i, symbolSectionIndex);
+					}
+				}
 			}
 
 			coff::WriteRaw(wideObjPath.c_str(), rawCoff, removalStrategy);
diff --git a/Engine/Source/Developer/Windows/LiveCodingServer/Private/External/LC_LiveProcess.cpp b/Engine/Source/Developer/Windows/LiveCodingServer/Private/External/LC_LiveProcess.cpp
--- a/Engine/Source/Developer/Windows/LiveCodingServer/Private/External/LC_LiveProcess.cpp
+++ b/Engine/Source/Developer/Windows/LiveCodingServer/Private/External/LC_LiveProcess.cpp
@@ -38,16 +38,19 @@ namespace liveProcess // Duplicate definitions
 
 	static const void* GetUpperBoundIn4GBRange(const void* moduleBase)
 	{
+		const uint64_t HighestPossibleAddress = 0x00007FFFFFFF0000ull;
+		const uint64_t HighestPossibleAddressThreshold = HighestPossibleAddress - 2ull * 1024ull * 1024ull * 1024ull;
+
 		const uint64_t base = pointer::AsInteger<uint64_t>(moduleBase);
 
 		// make sure we don't overflow
-		if (base <= 0xFFFFFFFF7FFFFFFFull)
+		if (base <= HighestPossibleAddressThreshold)
 		{
-			return pointer::FromInteger<const void*>(base + 0x80000000ull);
+			return pointer::FromInteger<const void*>(base + 2ull * 1024ull * 1024ull * 1024ull);
 		}
 
 		// operation would overflow
-		return pointer::FromInteger<const void*>(0xFFFFFFFFFFFFFFFFull);
+		return pointer::FromInteger<const void*>(HighestPossibleAddress);
 	}
 }
 #endif
diff --git a/Engine/Source/Developer/Windows/LiveCodingServer/Private/External/LC_Symbols.cpp b/Engine/Source/Developer/Windows/LiveCodingServer/Private/External/LC_Symbols.cpp
--- a/Engine/Source/Developer/Windows/LiveCodingServer/Private/External/LC_Symbols.cpp
+++ b/Engine/Source/Developer/Windows/LiveCodingServer/Private/External/LC_Symbols.cpp
@@ -1755,9 +1755,9 @@ namespace symbols
 			auto it = existingDb->dependencies.find(filePath);
 			if (it != existingDb->dependencies.end())
 			{
-				// merge and update dependent .obj paths and modification time
+				// merge and update dependent .obj paths.
+				// do not update modification time, because the file might have been changed while we were still compiling it.
 				symbols::Dependency* existingDependency = it->second;
-				existingDependency->lastModification = newDependency->lastModification;
 
 				types::StringSet paths;
 				paths.insert(existingDependency->objPaths.begin(), existingDependency->objPaths.end());
@@ -2272,7 +2272,7 @@ namespace symbols
 
 	bool IsStringLiteral(const ImmutableString& symbolName)
 	{
-		return ContainsPatterns(symbolName.c_str(), symbolPatterns::STRING_LITERAL_PATTERNS);
+		return StartsWithPatterns(symbolName.c_str(), symbolPatterns::STRING_LITERAL_PATTERNS);
 	}
 
 
diff --git a/Engine/Source/Developer/Windows/LiveCodingServer/Private/External/LC_VirtualMemoryRange.cpp b/Engine/Source/Developer/Windows/LiveCodingServer/Private/External/LC_VirtualMemoryRange.cpp
--- a/Engine/Source/Developer/Windows/LiveCodingServer/Private/External/LC_VirtualMemoryRange.cpp
+++ b/Engine/Source/Developer/Windows/LiveCodingServer/Private/External/LC_VirtualMemoryRange.cpp
@@ -48,6 +48,12 @@ void VirtualMemoryRange::ReservePages(const void* addressStart, const void* addr
 		// align address to be scanned
 		address = pointer::AlignTop<const void*>(address, alignment);
 
+		if (address < addressStart)
+		{
+			// overflow happened because we scanned too far
+			break;
+		}
+
 		::MEMORY_BASIC_INFORMATION memoryInfo = {};
 		const size_t bytesReturned = ::VirtualQueryEx(+m_processHandle, address, &memoryInfo, sizeof(::MEMORY_BASIC_INFORMATION));
 		// BEGIN EPIC MOD
