Make sure every uniform buffer name in shader optional data is a valid name, i.e. no more "None"-s. Usually it's $Global being removed

diff --git a/Engine/Source/Developer/Windows/ShaderFormatD3D/Private/D3DShaderCompiler.inl b/Engine/Source/Developer/Windows/ShaderFormatD3D/Private/D3DShaderCompiler.inl
--- a/Engine/Source/Developer/Windows/ShaderFormatD3D/Private/D3DShaderCompiler.inl
+++ b/Engine/Source/Developer/Windows/ShaderFormatD3D/Private/D3DShaderCompiler.inl
@@ -410,15 +410,16 @@ template <typename ID3D1xShaderReflection, typename D3D1x_SHADER_DESC, typename
 				const FString UniformBufferName(CBDesc.Name);
 
 				HandleReflectedUniformBuffer(UniformBufferName, CBIndex, Output);
-
-				if (UniformBufferNames.Num() <= (int32)CBIndex)
-				{
-					UniformBufferNames.AddDefaulted(CBIndex - UniformBufferNames.Num() + 1);
-				}
-				UniformBufferNames[CBIndex] = UE::ShaderCompilerCommon::RemoveConstantBufferPrefix(UniformBufferName);
+				
 				UsedUniformBufferSlots[CBIndex] = true;
 			}
 
+			if (UniformBufferNames.Num() <= (int32)CBIndex)
+			{
+				UniformBufferNames.AddDefaulted(CBIndex - UniformBufferNames.Num() + 1);
+			}
+			UniformBufferNames[CBIndex] = UE::ShaderCompilerCommon::RemoveConstantBufferPrefix(FString(CBDesc.Name));
+
 			NumCBs = FMath::Max(NumCBs, BindDesc.BindPoint + BindDesc.BindCount);
 		}
 		else if (BindDesc.Type == D3D_SIT_TEXTURE || BindDesc.Type == D3D_SIT_SAMPLER)
@@ -544,7 +545,7 @@ inline void GenerateFinalOutput(TRefCountPtr<TBlob>& CompressedData,
 
 		if (UniformBufferNames.Num() < GenericSRT.ResourceTableLayoutHashes.Num())
 		{
-			UniformBufferNames.AddDefaulted(GenericSRT.ResourceTableLayoutHashes.Num() - UniformBufferNames.Num() + 1);
+			UniformBufferNames.AddDefaulted(GenericSRT.ResourceTableLayoutHashes.Num() - UniformBufferNames.Num());
 		}
 
 		for (int32 Index = 0; Index < GenericSRT.ResourceTableLayoutHashes.Num(); ++Index)
diff --git a/Engine/Source/Developer/Windows/ShaderFormatD3D/Private/ShaderFormatD3D.cpp b/Engine/Source/Developer/Windows/ShaderFormatD3D/Private/ShaderFormatD3D.cpp
--- a/Engine/Source/Developer/Windows/ShaderFormatD3D/Private/ShaderFormatD3D.cpp
+++ b/Engine/Source/Developer/Windows/ShaderFormatD3D/Private/ShaderFormatD3D.cpp
@@ -17,7 +17,7 @@ class FShaderFormatD3D : public IShaderFormat
 {
 	enum
 	{
-		UE_SHADER_PCD3D_SHARED_VER = 2,
+		UE_SHADER_PCD3D_SHARED_VER = 3,
 
 		/** Version for shader format, this becomes part of the DDC key. */
 		UE_SHADER_PCD3D_SM6_VER = UE_SHADER_PCD3D_SHARED_VER + 5,
