Fixed issue where modules listed to be preloaded were not being preloaded.

diff --git a/Engine/Source/Developer/Windows/LiveCoding/Private/LiveCodingModule.cpp b/Engine/Source/Developer/Windows/LiveCoding/Private/LiveCodingModule.cpp
--- a/Engine/Source/Developer/Windows/LiveCoding/Private/LiveCodingModule.cpp
+++ b/Engine/Source/Developer/Windows/LiveCoding/Private/LiveCodingModule.cpp
@@ -885,6 +885,20 @@ void FLiveCodingModule::UpdateModules()
 		verify(GetModuleFileName(hInstance, FullFilePath, UE_ARRAY_COUNT(FullFilePath)));
 		LppEnableModule(FullFilePath);
 #else
+		// Collect the list of preloaded modules
+		TSet<FName> PreloadedFileNames;
+		{
+			FModuleManager& Manager = FModuleManager::Get();
+			for (FName moduleName : Settings->PreloadNamedModules)
+			{
+				FString FileName = Manager.GetModuleFilename(moduleName);
+				if (!FileName.IsEmpty())
+				{
+					PreloadedFileNames.Add(FName(FPaths::GetBaseFilename(FileName, true)));
+				}
+			}
+		}
+
 		TArray<FString> EnableModules;
 		for (const ModuleChange& Change : Changes)
 		{
@@ -892,7 +906,7 @@ void FLiveCodingModule::UpdateModules()
 			{
 				ConfiguredModules.Add(Change.FullName);
 				FString FullFilePath(Change.FullName.ToString());
-				if (ShouldPreloadModule(FullFilePath))
+				if (ShouldPreloadModule(PreloadedFileNames, FullFilePath))
 				{
 					EnableModules.Add(FullFilePath);
 				}
@@ -945,14 +959,14 @@ void FLiveCodingModule::OnModulesChanged(FName ModuleName, EModuleChangeReason R
 #endif
 }
 
-bool FLiveCodingModule::ShouldPreloadModule(const FString& FullFilePath) const
+bool FLiveCodingModule::ShouldPreloadModule(const TSet<FName>& PreloadedFileNames, const FString& FullFilePath) const
 {
 
 	// Perform some name based checks
 	{
-		FString NameStr = FPaths::GetBaseFilename(FullFilePath, true);
-		FName Name(NameStr, FNAME_Find);
-		if (Name != NAME_None && Settings->PreloadNamedModules.Contains(Name))
+		FString FileName = FPaths::GetBaseFilename(FullFilePath, true);
+		FName Name(FileName, FNAME_Find);
+		if (Name != NAME_None && PreloadedFileNames.Contains(Name))
 		{
 			return true;
 		}
diff --git a/Engine/Source/Developer/Windows/LiveCoding/Private/LiveCodingModule.h b/Engine/Source/Developer/Windows/LiveCoding/Private/LiveCodingModule.h
--- a/Engine/Source/Developer/Windows/LiveCoding/Private/LiveCodingModule.h
+++ b/Engine/Source/Developer/Windows/LiveCoding/Private/LiveCodingModule.h
@@ -105,7 +105,7 @@ private:
 
 	void UpdateModules();
 
-	bool ShouldPreloadModule(const FString& FullFilePath) const;
+	bool ShouldPreloadModule(const TSet<FName>& PreloadedFileNames, const FString& FullFilePath) const;
 
 	bool IsReinstancingEnabled() const;
 
