Provide console feedback when starting live coding fails packaged builds

diff --git a/Engine/Source/Developer/Windows/LiveCoding/Private/LiveCodingModule.cpp b/Engine/Source/Developer/Windows/LiveCoding/Private/LiveCodingModule.cpp
--- a/Engine/Source/Developer/Windows/LiveCoding/Private/LiveCodingModule.cpp
+++ b/Engine/Source/Developer/Windows/LiveCoding/Private/LiveCodingModule.cpp
@@ -261,7 +261,7 @@ void FLiveCodingModule::StartupModule()
 	EnableCommand = ConsoleManager.RegisterConsoleCommand(
 		TEXT("LiveCoding"),
 		TEXT("Enables live coding support"),
-		FConsoleCommandDelegate::CreateRaw(this, &FLiveCodingModule::EnableForSession, true),
+		FConsoleCommandWithOutputDeviceDelegate::CreateRaw(this, &FLiveCodingModule::EnableConsoleCommand),
 		ECVF_Cheat
 	);
 
@@ -376,6 +376,18 @@ bool FLiveCodingModule::IsEnabledByDefault() const
 	return Settings->bEnabled;
 }
 
+void FLiveCodingModule::EnableConsoleCommand(FOutputDevice& out)
+{
+	EnableForSession(true);
+
+	// For packaged builds, by default it is unlikely that UE_LOG messages will be seen in the console.
+	// So log any error directly.
+	if (!EnableErrorText.IsEmpty())
+	{
+		out.Log(EnableErrorText);
+	}
+}
+
 void FLiveCodingModule::EnableForSession(bool bEnable)
 {
 	if (bEnable)
diff --git a/Engine/Source/Developer/Windows/LiveCoding/Private/LiveCodingModule.h b/Engine/Source/Developer/Windows/LiveCoding/Private/LiveCodingModule.h
--- a/Engine/Source/Developer/Windows/LiveCoding/Private/LiveCodingModule.h
+++ b/Engine/Source/Developer/Windows/LiveCoding/Private/LiveCodingModule.h
@@ -12,6 +12,7 @@ struct IConsoleCommand;
 class IConsoleVariable;
 class ISettingsSection;
 class ULiveCodingSettings;
+class FOutputDevice;
 
 #if WITH_EDITOR
 class FReload;
@@ -99,6 +100,8 @@ private:
 	TUniquePtr<FNullReload> Reload;
 #endif
 
+	void EnableConsoleCommand(FOutputDevice& out);
+
 	bool StartLiveCoding();
 
 	void OnModulesChanged(FName ModuleName, EModuleChangeReason Reason);
