Fixed issue where packaged blueprint builds would not be able to start live coding without first having to clear the live coding source project variable.

diff --git a/Engine/Source/Developer/Windows/LiveCoding/Private/LiveCodingModule.cpp b/Engine/Source/Developer/Windows/LiveCoding/Private/LiveCodingModule.cpp
--- a/Engine/Source/Developer/Windows/LiveCoding/Private/LiveCodingModule.cpp
+++ b/Engine/Source/Developer/Windows/LiveCoding/Private/LiveCodingModule.cpp
@@ -294,9 +294,13 @@ void FLiveCodingModule::StartupModule()
 #else
 	FString SourceProject = FPaths::IsProjectFilePathSet() ? FPaths::GetProjectFilePath() : TEXT("");
 #endif
+	if (SourceProject.Len() > 0)
+	{
+		SourceProject = FPaths::ConvertRelativePathToFull(SourceProject);
+	}
 	SourceProjectVariable = ConsoleManager.RegisterConsoleVariable(
 		TEXT("LiveCoding.SourceProject"),
-		FPaths::ConvertRelativePathToFull(SourceProject),
+		SourceProject,
 		TEXT("Path to the project that this target was built from"),
 		ECVF_Cheat
 	);
