The Windows Feature Level checkboxes should prompt the user to restart.

diff --git a/Engine/Source/Developer/Windows/WindowsTargetPlatform/Classes/WindowsTargetSettings.h b/Engine/Source/Developer/Windows/WindowsTargetPlatform/Classes/WindowsTargetSettings.h
--- a/Engine/Source/Developer/Windows/WindowsTargetPlatform/Classes/WindowsTargetSettings.h
+++ b/Engine/Source/Developer/Windows/WindowsTargetPlatform/Classes/WindowsTargetSettings.h
@@ -51,13 +51,13 @@ public:
 	UPROPERTY(config, meta = (DeprecatedProperty, DeprecationMessage = "Use one of the RHI specific lists."))
 	TArray<FString> TargetedRHIs_DEPRECATED;
 
-	UPROPERTY(EditAnywhere, config, Category = "Rendering")
+	UPROPERTY(EditAnywhere, config, Category = "Rendering", Meta = (ConfigRestartRequired = true))
 	TArray<FString> D3D12TargetedShaderFormats;
 
-	UPROPERTY(EditAnywhere, config, Category = "Rendering")
+	UPROPERTY(EditAnywhere, config, Category = "Rendering", Meta = (ConfigRestartRequired = true))
 	TArray<FString> D3D11TargetedShaderFormats;
 
-	UPROPERTY(EditAnywhere, config, Category = "Rendering")
+	UPROPERTY(EditAnywhere, config, Category = "Rendering", Meta = (ConfigRestartRequired = true))
 	TArray<FString> VulkanTargetedShaderFormats;
 
 	/** The compiler version to use for this project. May be different to the chosen IDE. */
