Moving bindless resource count validation into the new UE::ShaderCompilerCommon::ValidatePackedResourceCounts

diff --git a/Engine/Source/Developer/ShaderCompilerCommon/Private/ShaderCompilerCommon.cpp b/Engine/Source/Developer/ShaderCompilerCommon/Private/ShaderCompilerCommon.cpp
--- a/Engine/Source/Developer/ShaderCompilerCommon/Private/ShaderCompilerCommon.cpp
+++ b/Engine/Source/Developer/ShaderCompilerCommon/Private/ShaderCompilerCommon.cpp
@@ -586,6 +586,45 @@ bool UE::ShaderCompilerCommon::RemoveBindlessParameterPrefix(FString& InName)
 		|| InName.RemoveFromStart(UE::ShaderCompilerCommon::kBindlessSamplerPrefix);
 }
 
+bool UE::ShaderCompilerCommon::ValidatePackedResourceCounts(FShaderCompilerOutput& Output, const FShaderCodePackedResourceCounts& PackedResourceCounts)
+{
+	if (Output.bSucceeded)
+	{
+		auto GetAllResourcesOfType = [&](EShaderParameterType InType)
+		{
+			const TArray<FString> AllNames = Output.ParameterMap.GetAllParameterNamesOfType(InType);
+			if (AllNames.IsEmpty())
+			{
+				return FString();
+			}
+			return FString::Join(AllNames, TEXT(", "));
+		};
+
+		if (EnumHasAnyFlags(PackedResourceCounts.UsageFlags, EShaderResourceUsageFlags::BindlessResources) && PackedResourceCounts.NumSRVs > 0)
+		{
+			const FString Names = GetAllResourcesOfType(EShaderParameterType::SRV);
+			Output.Errors.Add(FString::Printf(TEXT("Shader is mixing bindless resources with non-bindless resources. %d SRV slots were detected: %s"), PackedResourceCounts.NumSRVs, *Names));
+			Output.bSucceeded = false;
+		}
+
+		if (EnumHasAnyFlags(PackedResourceCounts.UsageFlags, EShaderResourceUsageFlags::BindlessResources) && PackedResourceCounts.NumUAVs > 0)
+		{
+			const FString Names = GetAllResourcesOfType(EShaderParameterType::UAV);
+			Output.Errors.Add(FString::Printf(TEXT("Shader is mixing bindless resources with non-bindless resources. %d UAV slots were detected: %s"), PackedResourceCounts.NumUAVs, *Names));
+			Output.bSucceeded = false;
+		}
+
+		if (EnumHasAnyFlags(PackedResourceCounts.UsageFlags, EShaderResourceUsageFlags::BindlessSamplers) && PackedResourceCounts.NumSamplers > 0)
+		{
+			const FString Names = GetAllResourcesOfType(EShaderParameterType::Sampler);
+			Output.Errors.Add(FString::Printf(TEXT("Shader is mixing bindless samplers with non-bindless samplers. %d sampler slots were detected: %s"), PackedResourceCounts.NumSamplers, *Names));
+			Output.bSucceeded = false;
+		}
+	}
+
+	return Output.bSucceeded;
+}
+
 void HandleReflectedGlobalConstantBufferMember(
 	const FString& InMemberName,
 	uint32 ConstantBufferIndex,
diff --git a/Engine/Source/Developer/ShaderCompilerCommon/Public/ShaderCompilerCommon.h b/Engine/Source/Developer/ShaderCompilerCommon/Public/ShaderCompilerCommon.h
--- a/Engine/Source/Developer/ShaderCompilerCommon/Public/ShaderCompilerCommon.h
+++ b/Engine/Source/Developer/ShaderCompilerCommon/Public/ShaderCompilerCommon.h
@@ -65,6 +65,8 @@ namespace UE::ShaderCompilerCommon
 	extern SHADERCOMPILERCOMMON_API EShaderParameterType ParseAndRemoveBindlessParameterPrefix(FString& InName);
 
 	extern SHADERCOMPILERCOMMON_API bool                 RemoveBindlessParameterPrefix(FString& InName);
+
+	extern SHADERCOMPILERCOMMON_API bool                 ValidatePackedResourceCounts(FShaderCompilerOutput& Output, const FShaderCodePackedResourceCounts& PackedResourceCounts);
 }
 
 extern SHADERCOMPILERCOMMON_API void HandleReflectedGlobalConstantBufferMember(
diff --git a/Engine/Source/Developer/Windows/ShaderFormatD3D/Private/D3DShaderCompilerDXC.cpp b/Engine/Source/Developer/Windows/ShaderFormatD3D/Private/D3DShaderCompilerDXC.cpp
--- a/Engine/Source/Developer/Windows/ShaderFormatD3D/Private/D3DShaderCompilerDXC.cpp
+++ b/Engine/Source/Developer/Windows/ShaderFormatD3D/Private/D3DShaderCompilerDXC.cpp
@@ -886,41 +886,31 @@ bool CompileAndProcessD3DShaderDXC(FString& PreprocessedShaderSource,
 			Output.bSucceeded = false;
 		}
 
-		const bool bBindlessResourcesUsed = ((ShaderRequiresFlags & D3D_SHADER_REQUIRES_RESOURCE_DESCRIPTOR_HEAP_INDEXING) != 0);
-		const bool bBindlessSamplersUsed = ((ShaderRequiresFlags & D3D_SHADER_REQUIRES_SAMPLER_DESCRIPTOR_HEAP_INDEXING) != 0);
+		FShaderCodePackedResourceCounts PackedResourceCounts{};
 
 		if (Output.bSucceeded)
 		{
-			auto GetAllResourcesOfType = [&](EShaderParameterType InType)
-			{
-				const TArray<FString> AllNames = Output.ParameterMap.GetAllParameterNamesOfType(InType);
-				if (AllNames.IsEmpty())
-				{
-					return FString();
-				}
-				return FString::Join(AllNames, TEXT(", "));
-			};
-
-			if (bBindlessResourcesUsed && NumSRVs > 0)
+			if (bGlobalUniformBufferUsed)
 			{
-				const FString Names = GetAllResourcesOfType(EShaderParameterType::SRV);
-				FilteredErrors.Add(FString::Printf(TEXT("Shader is mixing bindless resources with non-bindless resources. %d SRV slots were detected: %s"), NumSRVs, *Names));
-				Output.bSucceeded = false;
+				PackedResourceCounts.UsageFlags |= EShaderResourceUsageFlags::GlobalUniformBuffer;
 			}
 
-			if (bBindlessResourcesUsed && NumUAVs > 0)
+			if (ShaderRequiresFlags & D3D_SHADER_REQUIRES_RESOURCE_DESCRIPTOR_HEAP_INDEXING)
 			{
-				const FString Names = GetAllResourcesOfType(EShaderParameterType::UAV);
-				FilteredErrors.Add(FString::Printf(TEXT("Shader is mixing bindless resources with non-bindless resources. %d UAV slots were detected: %s"), NumUAVs, *Names));
-				Output.bSucceeded = false;
+				PackedResourceCounts.UsageFlags |= EShaderResourceUsageFlags::BindlessResources;
 			}
 
-			if (bBindlessSamplersUsed && NumSamplers > 0)
+			if (ShaderRequiresFlags & D3D_SHADER_REQUIRES_SAMPLER_DESCRIPTOR_HEAP_INDEXING)
 			{
-				const FString Names = GetAllResourcesOfType(EShaderParameterType::Sampler);
-				FilteredErrors.Add(FString::Printf(TEXT("Shader is mixing bindless samplers with non-bindless samplers. %d sampler slots were detected: %s"), NumSamplers, *Names));
-				Output.bSucceeded = false;
+				PackedResourceCounts.UsageFlags |= EShaderResourceUsageFlags::BindlessSamplers;
 			}
+
+			PackedResourceCounts.NumSamplers = static_cast<uint8>(NumSamplers);
+			PackedResourceCounts.NumSRVs = static_cast<uint8>(NumSRVs);
+			PackedResourceCounts.NumCBs = static_cast<uint8>(NumCBs);
+			PackedResourceCounts.NumUAVs = static_cast<uint8>(NumUAVs);
+
+			Output.bSucceeded = UE::ShaderCompilerCommon::ValidatePackedResourceCounts(Output, PackedResourceCounts);
 		}
 
 		// Save results if compilation and reflection succeeded
@@ -987,25 +977,6 @@ bool CompileAndProcessD3DShaderDXC(FString& PreprocessedShaderSource,
 
 			//#todo-rco: Should compress ShaderCode?
 
-			FShaderCodePackedResourceCounts PackedResourceCounts{};
-			if (bGlobalUniformBufferUsed)
-			{
-				PackedResourceCounts.UsageFlags |= EShaderResourceUsageFlags::GlobalUniformBuffer;
-			}
-			if (bBindlessResourcesUsed)
-			{
-				PackedResourceCounts.UsageFlags |= EShaderResourceUsageFlags::BindlessResources;
-			}
-			if (bBindlessSamplersUsed)
-			{
-				PackedResourceCounts.UsageFlags |= EShaderResourceUsageFlags::BindlessSamplers;
-			}
-
-			PackedResourceCounts.NumSamplers = static_cast<uint8>(NumSamplers);
-			PackedResourceCounts.NumSRVs = static_cast<uint8>(NumSRVs);
-			PackedResourceCounts.NumCBs = static_cast<uint8>(NumCBs);
-			PackedResourceCounts.NumUAVs = static_cast<uint8>(NumUAVs);
-
 			GenerateFinalOutput(ShaderBlob,
 				Input, VendorExtensions,
 				UsedUniformBufferSlots, UniformBufferNames,
