Fixing debug DXC shader compiles.
When D3DCOMPILE_SKIP_OPTIMIZATION is set, none of the other D3DCOMPILE_OPTIMIZATION_LEVEL flags are set. Given that D3DCOMPILE_OPTIMIZATION_LEVEL1 is "0", this always forced debug shader builds to be compiled with "/Od /O1".
The fix is to ignore the other flags when D3DCOMPILE_SKIP_OPTIMIZATION is set.

diff --git a/Engine/Source/Developer/Windows/ShaderFormatD3D/Private/D3DShaderCompiler.inl b/Engine/Source/Developer/Windows/ShaderFormatD3D/Private/D3DShaderCompiler.inl
--- a/Engine/Source/Developer/Windows/ShaderFormatD3D/Private/D3DShaderCompiler.inl
+++ b/Engine/Source/Developer/Windows/ShaderFormatD3D/Private/D3DShaderCompiler.inl
@@ -92,12 +92,6 @@ public:
 			ExtraArguments.Add(TEXT("/Gfp"));
 		}
 
-		if (D3DCompileFlags & D3DCOMPILE_SKIP_OPTIMIZATION)
-		{
-			D3DCompileFlags &= ~D3DCOMPILE_SKIP_OPTIMIZATION;
-			ExtraArguments.Add(TEXT("/Od"));
-		}
-
 		if (D3DCompileFlags & D3DCOMPILE_SKIP_VALIDATION)
 		{
 			D3DCompileFlags &= ~D3DCOMPILE_SKIP_VALIDATION;
@@ -128,30 +122,38 @@ public:
 			ExtraArguments.Add(TEXT("/WX"));
 		}
 
-		switch (D3DCompileFlags & SHADER_OPTIMIZATION_LEVEL_MASK)
+		if (D3DCompileFlags & D3DCOMPILE_SKIP_OPTIMIZATION)
 		{
-		case D3DCOMPILE_OPTIMIZATION_LEVEL0:
-			D3DCompileFlags &= ~D3DCOMPILE_OPTIMIZATION_LEVEL0;
-			ExtraArguments.Add(TEXT("/O0"));
-			break;
-
-		case D3DCOMPILE_OPTIMIZATION_LEVEL1:
-			D3DCompileFlags &= ~D3DCOMPILE_OPTIMIZATION_LEVEL1;
-			ExtraArguments.Add(TEXT("/O1"));
-			break;
-
-		case D3DCOMPILE_OPTIMIZATION_LEVEL2:
-			D3DCompileFlags &= ~D3DCOMPILE_OPTIMIZATION_LEVEL2;
-			ExtraArguments.Add(TEXT("/O2"));
-			break;
-
-		case D3DCOMPILE_OPTIMIZATION_LEVEL3:
-			D3DCompileFlags &= ~D3DCOMPILE_OPTIMIZATION_LEVEL3;
-			ExtraArguments.Add(TEXT("/O3"));
-			break;
-
-		default:
-			break;
+			D3DCompileFlags &= ~D3DCOMPILE_SKIP_OPTIMIZATION;
+			ExtraArguments.Add(TEXT("/Od"));
+		}
+		else
+		{
+			switch (D3DCompileFlags & SHADER_OPTIMIZATION_LEVEL_MASK)
+			{
+			case D3DCOMPILE_OPTIMIZATION_LEVEL0:
+				D3DCompileFlags &= ~D3DCOMPILE_OPTIMIZATION_LEVEL0;
+				ExtraArguments.Add(TEXT("/O0"));
+				break;
+
+			case D3DCOMPILE_OPTIMIZATION_LEVEL1:
+				D3DCompileFlags &= ~D3DCOMPILE_OPTIMIZATION_LEVEL1;
+				ExtraArguments.Add(TEXT("/O1"));
+				break;
+
+			case D3DCOMPILE_OPTIMIZATION_LEVEL2:
+				D3DCompileFlags &= ~D3DCOMPILE_OPTIMIZATION_LEVEL2;
+				ExtraArguments.Add(TEXT("/O2"));
+				break;
+
+			case D3DCOMPILE_OPTIMIZATION_LEVEL3:
+				D3DCompileFlags &= ~D3DCOMPILE_OPTIMIZATION_LEVEL3;
+				ExtraArguments.Add(TEXT("/O3"));
+				break;
+
+			default:
+				break;
+			}
 		}
 
 		if (D3DCompileFlags & D3DCOMPILE_DEBUG)
diff --git a/Engine/Source/Developer/Windows/ShaderFormatD3D/Private/ShaderFormatD3D.cpp b/Engine/Source/Developer/Windows/ShaderFormatD3D/Private/ShaderFormatD3D.cpp
--- a/Engine/Source/Developer/Windows/ShaderFormatD3D/Private/ShaderFormatD3D.cpp
+++ b/Engine/Source/Developer/Windows/ShaderFormatD3D/Private/ShaderFormatD3D.cpp
@@ -17,7 +17,7 @@ class FShaderFormatD3D : public IShaderFormat
 {
 	enum
 	{
-		UE_SHADER_PCD3D_SHARED_VER = 3,
+		UE_SHADER_PCD3D_SHARED_VER = 4,
 
 		/** Version for shader format, this becomes part of the DDC key. */
 		UE_SHADER_PCD3D_SM6_VER = UE_SHADER_PCD3D_SHARED_VER + 5,
