Fix crash issue when LiveCodingConsole is run independently.

diff --git a/Engine/Source/Developer/Windows/LiveCodingServer/Private/LiveCodingServer.cpp b/Engine/Source/Developer/Windows/LiveCodingServer/Private/LiveCodingServer.cpp
--- a/Engine/Source/Developer/Windows/LiveCodingServer/Private/LiveCodingServer.cpp
+++ b/Engine/Source/Developer/Windows/LiveCodingServer/Private/LiveCodingServer.cpp
@@ -117,10 +117,10 @@ ILiveCodingServer::FSetVisibleDelegate& FLiveCodingServer::GetSetVisibleDelegate
 
 bool FLiveCodingServer::HasReinstancingProcess()
 {
-	return CommandThread->HasReinstancingProcess();
+	return CommandThread != nullptr ? CommandThread->HasReinstancingProcess() : false;
 }
 
 bool FLiveCodingServer::ShowCompileFinishNotification()
 {
-	return CommandThread->ShowCompileFinishNotification();
+	return CommandThread != nullptr ? CommandThread->ShowCompileFinishNotification() : false;
 }
