Add the Frequency name to the "Unrecognized shader frequency" Shader Compiler error

diff --git a/Engine/Source/Developer/Windows/ShaderFormatD3D/Private/D3DShaderCompiler.cpp b/Engine/Source/Developer/Windows/ShaderFormatD3D/Private/D3DShaderCompiler.cpp
--- a/Engine/Source/Developer/Windows/ShaderFormatD3D/Private/D3DShaderCompiler.cpp
+++ b/Engine/Source/Developer/Windows/ShaderFormatD3D/Private/D3DShaderCompiler.cpp
@@ -1198,7 +1198,7 @@ void CompileD3DShader(const FShaderCompilerInput& Input, FShaderCompilerOutput&
 
 	if(!ShaderProfile)
 	{
-		Output.Errors.Add(FShaderCompilerError(TEXT("Unrecognized shader frequency")));
+		Output.Errors.Add(FShaderCompilerError(*FString::Printf(TEXT("Unrecognized shader frequency %s"), GetShaderFrequencyString((EShaderFrequency)Input.Target.Frequency))));
 		return;
 	}
 
