ShaderMinifier - Add experimental CFLAG_RemoveDeadCode to run shader rewriting step that removes unused functions and types from shader source before compilation to improve compile times

diff --git a/Engine/Source/Developer/Windows/ShaderFormatD3D/Private/D3DShaderCompiler.cpp b/Engine/Source/Developer/Windows/ShaderFormatD3D/Private/D3DShaderCompiler.cpp
--- a/Engine/Source/Developer/Windows/ShaderFormatD3D/Private/D3DShaderCompiler.cpp
+++ b/Engine/Source/Developer/Windows/ShaderFormatD3D/Private/D3DShaderCompiler.cpp
@@ -3,6 +3,7 @@
 #include "ShaderFormatD3D.h"
 #include "ShaderPreprocessor.h"
 #include "ShaderCompilerCommon.h"
+#include "ShaderMinifier.h"
 #include "ShaderParameterParser.h"
 #include "D3D11ShaderResources.h"
 #include "D3D12RHI.h"
@@ -1316,6 +1317,42 @@ void CompileD3DShader(const FShaderCompilerInput& Input, FShaderCompilerOutput&
 	// Process TEXT macro.
 	TransformStringIntoCharacterArray(PreprocessedShaderSource);
 
+	TArray<FString> FilteredErrors;
+
+	// Run the experimental shader minifier
+
+	if (Input.Environment.CompilerFlags.Contains(CFLAG_RemoveDeadCode))
+	{
+		FString EntryMain;
+		FString EntryAnyHit;
+		FString EntryIntersection;
+		UE::ShaderCompilerCommon::ParseRayTracingEntryPoint(EntryPointName, EntryMain, EntryAnyHit, EntryIntersection);
+
+		if (!EntryAnyHit.IsEmpty())
+		{
+			EntryMain += TEXT(";");
+			EntryMain += EntryAnyHit;
+		}
+
+		if (!EntryIntersection.IsEmpty())
+		{
+			EntryMain += TEXT(";");
+			EntryMain += EntryIntersection;
+		}
+
+		UE::ShaderMinifier::FMinifiedShader Minified  = UE::ShaderMinifier::Minify(PreprocessedShaderSource, EntryMain, 
+			UE::ShaderMinifier::EMinifyShaderFlags::OutputReasons | UE::ShaderMinifier::EMinifyShaderFlags::OutputStats);
+
+		if (Minified.Success())
+		{
+			Swap(PreprocessedShaderSource, Minified.Code);
+		}
+		else
+		{
+			FilteredErrors.Add(TEXT("Shader minification failed."));
+		}
+	}
+
 	// @TODO - implement different material path to allow us to remove backwards compat flag on sm5 shaders
 	uint32 CompileFlags = D3DCOMPILE_ENABLE_BACKWARDS_COMPATIBILITY
 		// Unpack uniform matrices as row-major to match the CPU layout.
@@ -1347,7 +1384,6 @@ void CompileD3DShader(const FShaderCompilerInput& Input, FShaderCompilerOutput&
 			CompileFlags |= TranslateCompilerFlagD3D11((ECompilerFlags)Flag);
 		});
 
-	TArray<FString> FilteredErrors;
 	if (bUseDXC)
 	{
 		if (!CompileAndProcessD3DShaderDXC(PreprocessedShaderSource, CompileFlags, Input, ShaderParameterParser, EntryPointName, ShaderProfile, Language, false, FilteredErrors, Output))
diff --git a/Engine/Source/Runtime/Engine/Private/ShaderCompiler/ShaderCompiler.cpp b/Engine/Source/Runtime/Engine/Private/ShaderCompiler/ShaderCompiler.cpp
--- a/Engine/Source/Runtime/Engine/Private/ShaderCompiler/ShaderCompiler.cpp
+++ b/Engine/Source/Runtime/Engine/Private/ShaderCompiler/ShaderCompiler.cpp
@@ -1184,6 +1184,15 @@ static TAutoConsoleVariable<int32> CVarShadersValidation(
 	TEXT("Enabled shader compiler validation warnings and errors."),
 	ECVF_ReadOnly);
 
+static TAutoConsoleVariable<int32> CVarShadersRemoveDeadCode(
+	TEXT("r.Shaders.RemoveDeadCode"),
+	0,
+	TEXT("EXPERIMENTAL: Run a preprocessing step that removes unreferenced code before compiling shaders.\n")
+	TEXT("This can improve the compilation speed for shaders which include many large utility headers.\n")
+	TEXT("\t0: Keep all input source code (Default).\n")
+	TEXT("\t1: Remove unreferenced code before compilation\n"),
+	ECVF_ReadOnly);
+
 extern bool CompileShaderPipeline(const IShaderFormat* Compiler, FName Format, FShaderPipelineCompileJob* PipelineJob, const FString& Dir);
 
 #if ENABLE_COOK_STATS
@@ -6165,6 +6174,11 @@ void GlobalBeginCompileShader(
 		}
 	}
 
+	if (CVarShadersRemoveDeadCode.GetValueOnAnyThread())
+	{
+		Input.Environment.CompilerFlags.Add(CFLAG_RemoveDeadCode);
+	}
+
 	{
 		static const auto CVar = IConsoleManager::Get().FindConsoleVariable(TEXT("r.VT.AnisotropicFiltering"));
 		Input.Environment.SetDefine(TEXT("VIRTUAL_TEXTURE_ANISOTROPIC_FILTERING"), CVar ? (CVar->GetInt() != 0) : 0);
diff --git a/Engine/Source/Runtime/RenderCore/Private/Shader.cpp b/Engine/Source/Runtime/RenderCore/Private/Shader.cpp
--- a/Engine/Source/Runtime/RenderCore/Private/Shader.cpp
+++ b/Engine/Source/Runtime/RenderCore/Private/Shader.cpp
@@ -1933,6 +1933,14 @@ void ShaderMapAppendKeyString(EShaderPlatform Platform, FString& KeyString)
  		KeyString += tt;
 	}
 
+	{
+		static const auto CVar = IConsoleManager::Get().FindTConsoleVariableDataInt(TEXT("r.Shaders.RemoveDeadCode"));
+		if (CVar && CVar->GetValueOnAnyThread() != 0)
+		{
+			KeyString += TEXT("_MIN");
+		}
+	}
+
 	if (RHISupportsRenderTargetWriteMask(Platform))
 	{
 		KeyString += TEXT("_RTWM");
diff --git a/Engine/Source/Runtime/RenderCore/Public/ShaderCompilerCore.h b/Engine/Source/Runtime/RenderCore/Public/ShaderCompilerCore.h
--- a/Engine/Source/Runtime/RenderCore/Public/ShaderCompilerCore.h
+++ b/Engine/Source/Runtime/RenderCore/Public/ShaderCompilerCore.h
@@ -123,6 +123,8 @@ enum ECompilerFlags
 	CFLAG_BindlessResources,
 	// Enabled if bindless samplers are enabled for the platform
 	CFLAG_BindlessSamplers,
+	// EXPERIMENTAL: Run the shader re-writer that removes any unused functions/resources/types from source code before compilation.
+	CFLAG_RemoveDeadCode,
 
 	CFLAG_Max,
 };
