Don't remove PDB from DXIL container if it wasn't generated in the first place to prevent hlsl::Exception to be thrown by DXC if PDB part is not present.

#rnx

diff --git a/Engine/Source/Developer/Windows/ShaderFormatD3D/Private/D3DShaderCompilerDXC.cpp b/Engine/Source/Developer/Windows/ShaderFormatD3D/Private/D3DShaderCompilerDXC.cpp
--- a/Engine/Source/Developer/Windows/ShaderFormatD3D/Private/D3DShaderCompilerDXC.cpp
+++ b/Engine/Source/Developer/Windows/ShaderFormatD3D/Private/D3DShaderCompilerDXC.cpp
@@ -314,6 +314,8 @@ static HRESULT D3DCompileToDxil(const char* SourceText, FDxcArguments& Arguments
 		checkf(CompileResult->HasOutput(DXC_OUT_REFLECTION), TEXT("No reflection found!"));
 		VERIFYHRESULT(CompileResult->GetOutput(DXC_OUT_REFLECTION, IID_PPV_ARGS(OutReflectionBlob.GetInitReference()), ReflectionNameBlob.GetInitReference()));
 
+		const bool bHasOutputPDB = CompileResult->HasOutput(DXC_OUT_PDB);
+
 		if (Arguments.ShouldDump())
 		{
 			// Dump disassembly before we strip reflection out
@@ -325,7 +327,7 @@ static HRESULT D3DCompileToDxil(const char* SourceText, FDxcArguments& Arguments
 			FString DxilFile = Arguments.GetDumpDisassemblyFilename().LeftChop(7) + TEXT("_refl.dxil");
 			SaveDxcBlobToFile(OutDxilBlob, DxilFile);
 
-			if (CompileResult->HasOutput(DXC_OUT_PDB))
+			if (bHasOutputPDB)
 			{
 				TRefCountPtr<IDxcBlob> PdbBlob;
 				TRefCountPtr<IDxcBlobUtf16> PdbNameBlob;
@@ -340,7 +342,7 @@ static HRESULT D3DCompileToDxil(const char* SourceText, FDxcArguments& Arguments
 		}
 
 		DumpFourCCParts(DxcDllHelper, OutDxilBlob);
-		if (RemoveContainerReflection(DxcDllHelper, OutDxilBlob, !Arguments.ShouldKeepEmbeddedPDB()))
+		if (RemoveContainerReflection(DxcDllHelper, OutDxilBlob, bHasOutputPDB && !Arguments.ShouldKeepEmbeddedPDB()))
 		{
 			DumpFourCCParts(DxcDllHelper, OutDxilBlob);
 		}
