When changing SM6 raytracing shaders to use lib_6_6 from lib_6_5, the __SHADER_TARGET_MINOR value wasn't updated to 6.
When changing SM5 raytracing shaders to use lib_6_5 from lib_6_3, the __SHADER_TARGET_MINOR value wasn't updated to 5.
Bumping SM6 and SM5 shader versions.

diff --git a/Engine/Source/Developer/Windows/ShaderFormatD3D/Private/ShaderFormatD3D.cpp b/Engine/Source/Developer/Windows/ShaderFormatD3D/Private/ShaderFormatD3D.cpp
--- a/Engine/Source/Developer/Windows/ShaderFormatD3D/Private/ShaderFormatD3D.cpp
+++ b/Engine/Source/Developer/Windows/ShaderFormatD3D/Private/ShaderFormatD3D.cpp
@@ -20,8 +20,8 @@ class FShaderFormatD3D : public IShaderFormat
 		UE_SHADER_PCD3D_SHARED_VER = 4,
 
 		/** Version for shader format, this becomes part of the DDC key. */
-		UE_SHADER_PCD3D_SM6_VER = UE_SHADER_PCD3D_SHARED_VER + 5,
-		UE_SHADER_PCD3D_SM5_VER = UE_SHADER_PCD3D_SHARED_VER + 10,
+		UE_SHADER_PCD3D_SM6_VER = UE_SHADER_PCD3D_SHARED_VER + 6,
+		UE_SHADER_PCD3D_SM5_VER = UE_SHADER_PCD3D_SHARED_VER + 11,
 		UE_SHADER_PCD3D_ES3_1_VER = UE_SHADER_PCD3D_SHARED_VER + 8,
 		UE_SHADER_D3D_ES3_1_HOLOLENS_VER = UE_SHADER_PCD3D_ES3_1_VER,
 	};
@@ -122,7 +122,7 @@ public:
 			Input.Environment.SetDefine(TEXT("COMPILER_DXC"), 1);
 			Input.Environment.SetDefine(TEXT("PLATFORM_SUPPORTS_UB_STRUCT"), 1);
 
-			if (USE_SHADER_MODEL_6_6 && !Input.IsRayTracingShader())
+			if (USE_SHADER_MODEL_6_6)
 			{
 				AddShaderTargetDefines(Input, 6, 6);
 			}
@@ -142,7 +142,7 @@ public:
 
 			if (bUseDXC)
 			{
-				AddShaderTargetDefines(Input, 6, Input.IsRayTracingShader() ? 3 : 0);
+				AddShaderTargetDefines(Input, 6, Input.IsRayTracingShader() ? 5 : 0);
 			}
 			else
 			{
