Instead of adding the different D3D shader format versions together, combine their hashes. This will prevent version overlaps when someone changes the shared version.

diff --git a/Engine/Source/Developer/Windows/ShaderFormatD3D/Private/ShaderFormatD3D.cpp b/Engine/Source/Developer/Windows/ShaderFormatD3D/Private/ShaderFormatD3D.cpp
--- a/Engine/Source/Developer/Windows/ShaderFormatD3D/Private/ShaderFormatD3D.cpp
+++ b/Engine/Source/Developer/Windows/ShaderFormatD3D/Private/ShaderFormatD3D.cpp
@@ -15,14 +15,14 @@ static FName NAME_D3D_ES3_1_HOLOLENS(TEXT("D3D_ES3_1_HOLOLENS"));
 
 class FShaderFormatD3D : public IShaderFormat
 {
-	enum
+	enum EVersion
 	{
 		UE_SHADER_PCD3D_SHARED_VER = 4,
 
 		/** Version for shader format, this becomes part of the DDC key. */
-		UE_SHADER_PCD3D_SM6_VER = UE_SHADER_PCD3D_SHARED_VER + 6,
-		UE_SHADER_PCD3D_SM5_VER = UE_SHADER_PCD3D_SHARED_VER + 11,
-		UE_SHADER_PCD3D_ES3_1_VER = UE_SHADER_PCD3D_SHARED_VER + 8,
+		UE_SHADER_PCD3D_SM6_VER = 6,
+		UE_SHADER_PCD3D_SM5_VER = 11,
+		UE_SHADER_PCD3D_ES3_1_VER = 8,
 		UE_SHADER_D3D_ES3_1_HOLOLENS_VER = UE_SHADER_PCD3D_ES3_1_VER,
 	};
 
@@ -40,12 +40,19 @@ public:
 	{
 	}
 
+	inline uint32 GetVersionHash(EVersion InVersion) const
+	{
+		const uint32 BaseHash = GetTypeHash(UE_SHADER_PCD3D_SHARED_VER);
+		const uint32 VersionHash = GetTypeHash(InVersion);
+		return HashCombine(BaseHash, VersionHash);
+	}
+
 	virtual uint32 GetVersion(FName Format) const override
 	{
 		CheckFormat(Format);
 		if (Format == NAME_PCD3D_SM6)
 		{
-			uint32 ShaderModelHash = GetTypeHash(UE_SHADER_PCD3D_SM6_VER);
+			uint32 ShaderModelHash = GetVersionHash(UE_SHADER_PCD3D_SM6_VER);
 
 		#if USE_SHADER_MODEL_6_6
 			ShaderModelHash ^= 0x96ED7F56;
@@ -57,17 +64,17 @@ public:
 		{
 			// Technically not needed for regular SM5 compiled with legacy compiler,
 			// but PCD3D_SM5 currently includes ray tracing shaders that are compiled with new compiler stack.
-			return HashCombine(DxcVersionHash, GetTypeHash(UE_SHADER_PCD3D_SM5_VER));
+			return HashCombine(DxcVersionHash, GetVersionHash(UE_SHADER_PCD3D_SM5_VER));
 		}
 		else if (Format == NAME_PCD3D_ES3_1) 
 		{
 			// Shader DXC signature is intentionally not included, as ES3_1 target always uses legacy compiler.
-			return UE_SHADER_PCD3D_ES3_1_VER;
+			return GetVersionHash(UE_SHADER_PCD3D_ES3_1_VER);
 		}
 		else if (Format == NAME_D3D_ES3_1_HOLOLENS)
 		{
 			// Shader DXC signature is intentionally not included, as ES3_1 target always uses legacy compiler.
-			return UE_SHADER_D3D_ES3_1_HOLOLENS_VER;
+			return GetVersionHash(UE_SHADER_D3D_ES3_1_HOLOLENS_VER);
 		}
 		checkf(0, TEXT("Unknown Format %s"), *Format.ToString());
 		return 0;
