Don't count shader parameters that we don't bind through resource tables. Resources like UEDiagnosticBuffer and vendor extensions are handled through Root Signatures.

diff --git a/Engine/Source/Developer/Windows/ShaderFormatD3D/Private/D3DShaderCompiler.inl b/Engine/Source/Developer/Windows/ShaderFormatD3D/Private/D3DShaderCompiler.inl
--- a/Engine/Source/Developer/Windows/ShaderFormatD3D/Private/D3DShaderCompiler.inl
+++ b/Engine/Source/Developer/Windows/ShaderFormatD3D/Private/D3DShaderCompiler.inl
@@ -434,17 +434,7 @@ template <typename ID3D1xShaderReflection, typename D3D1x_SHADER_DESC, typename
 			const bool bIsVendorParameter = bIsAMDTexExtension || bIsAMDSmpExtension;
 
 			const uint32 BindCount = 1;
-			EShaderParameterType ParameterType = EShaderParameterType::Num;
-			if (BindDesc.Type == D3D_SIT_SAMPLER)
-			{
-				ParameterType = EShaderParameterType::Sampler;
-				NumSamplers = FMath::Max(NumSamplers, BindDesc.BindPoint + BindCount);
-			}
-			else // D3D_SIT_TEXTURE
-			{
-				ParameterType = EShaderParameterType::SRV;
-				NumSRVs = FMath::Max(NumSRVs, BindDesc.BindPoint + BindCount);
-			}
+			const EShaderParameterType ParameterType = (BindDesc.Type == D3D_SIT_SAMPLER) ? EShaderParameterType::Sampler : EShaderParameterType::SRV;
 
 			if (bIsVendorParameter)
 			{
@@ -453,10 +443,12 @@ template <typename ID3D1xShaderReflection, typename D3D1x_SHADER_DESC, typename
 			else if (ParameterType == EShaderParameterType::Sampler)
 			{
 				HandleReflectedShaderSampler(FString(BindDesc.Name), BindDesc.BindPoint, Output);
+				NumSamplers = FMath::Max(NumSamplers, BindDesc.BindPoint + BindCount);
 			}
 			else
 			{
 				HandleReflectedShaderResource(FString(BindDesc.Name), BindDesc.BindPoint, Output);
+				NumSRVs = FMath::Max(NumSRVs, BindDesc.BindPoint + BindCount);
 			}
 		}
 		else if (BindDesc.Type == D3D_SIT_UAV_RWTYPED || BindDesc.Type == D3D_SIT_UAV_RWSTRUCTURED ||
@@ -499,9 +491,8 @@ template <typename ID3D1xShaderReflection, typename D3D1x_SHADER_DESC, typename
 			else
 			{
 				HandleReflectedShaderUAV(FString(BindDesc.Name), BindDesc.BindPoint, Output);
+				NumUAVs = FMath::Max(NumUAVs, BindDesc.BindPoint + BindCount);
 			}
-
-			NumUAVs = FMath::Max(NumUAVs, BindDesc.BindPoint + BindCount);
 		}
 		else if (BindDesc.Type == D3D_SIT_STRUCTURED || BindDesc.Type == D3D_SIT_BYTEADDRESS)
 		{
diff --git a/Engine/Source/Developer/Windows/ShaderFormatD3D/Private/ShaderFormatD3D.cpp b/Engine/Source/Developer/Windows/ShaderFormatD3D/Private/ShaderFormatD3D.cpp
--- a/Engine/Source/Developer/Windows/ShaderFormatD3D/Private/ShaderFormatD3D.cpp
+++ b/Engine/Source/Developer/Windows/ShaderFormatD3D/Private/ShaderFormatD3D.cpp
@@ -17,7 +17,7 @@ class FShaderFormatD3D : public IShaderFormat
 {
 	enum EVersion
 	{
-		UE_SHADER_PCD3D_SHARED_VER = 4,
+		UE_SHADER_PCD3D_SHARED_VER = 5,
 
 		/** Version for shader format, this becomes part of the DDC key. */
 		UE_SHADER_PCD3D_SM6_VER = 6,
