Working around DXC issue [REDACTED_PROFILE_URL]/DirectXShaderCompiler/issues/4715
When filling out D3D12_FUNCTION_DESC, DXC does not populate RequiredFeatureFlags. Given we rely on these flags to validate bindless support at runtime, we need to emulate these flags being set.

diff --git a/Engine/Source/Developer/Windows/ShaderFormatD3D/Private/D3DShaderCompilerDXC.cpp b/Engine/Source/Developer/Windows/ShaderFormatD3D/Private/D3DShaderCompilerDXC.cpp
--- a/Engine/Source/Developer/Windows/ShaderFormatD3D/Private/D3DShaderCompilerDXC.cpp
+++ b/Engine/Source/Developer/Windows/ShaderFormatD3D/Private/D3DShaderCompilerDXC.cpp
@@ -776,6 +776,19 @@ bool CompileAndProcessD3DShaderDXC(FString& PreprocessedShaderSource,
 				}
 			}
 
+			// @todo - working around DXC issue [REDACTED_PROFILE_URL]/DirectXShaderCompiler/issues/4715
+			if (LibraryDesc.FunctionCount > 0)
+			{
+				if (Input.Environment.CompilerFlags.Contains(CFLAG_BindlessResources))
+				{
+					ShaderRequiresFlags |= D3D_SHADER_REQUIRES_RESOURCE_DESCRIPTOR_HEAP_INDEXING;
+				}
+				if (Input.Environment.CompilerFlags.Contains(CFLAG_BindlessSamplers))
+				{
+					ShaderRequiresFlags |= D3D_SHADER_REQUIRES_SAMPLER_DESCRIPTOR_HEAP_INDEXING;
+				}
+			}
+
 			if (NumFoundEntryPoints == MangledEntryPoints.Num())
 			{
 				Output.bSucceeded = true;
