Improved Live Coding startup time to make starting the console an asyncronous task.

diff --git a/Engine/Source/Developer/Windows/LiveCoding/Private/LiveCodingModule.cpp b/Engine/Source/Developer/Windows/LiveCoding/Private/LiveCodingModule.cpp
--- a/Engine/Source/Developer/Windows/LiveCoding/Private/LiveCodingModule.cpp
+++ b/Engine/Source/Developer/Windows/LiveCoding/Private/LiveCodingModule.cpp
@@ -14,13 +14,13 @@
 #include "Misc/ConfigCacheIni.h"
 #include "Misc/MessageDialog.h"
 #include "Misc/ScopedSlowTask.h"
-#include "LiveCodingSettings.h"
 #include "ISettingsModule.h"
 #include "ISettingsSection.h"
 #include "Windows/WindowsHWrapper.h"
 #include "Algo/Sort.h"
 #include "Algo/BinarySearch.h"
 #include "HAL/LowLevelMemTracker.h"
+#include "Async/TaskGraphInterfaces.h"
 #if WITH_EDITOR
 	#include "Editor.h"
 	#include "Kismet2/ReloadUtilities.h"
@@ -319,26 +319,15 @@ void FLiveCodingModule::StartupModule()
 
 	LppStartup();
 
-	if (Settings->bEnabled && !FApp::IsUnattended())
+	if (Settings->bEnabled && Settings->Startup != ELiveCodingStartupMode::Manual && !FApp::IsUnattended())
 	{
-		if(Settings->Startup == ELiveCodingStartupMode::Automatic)
-		{
-			StartLiveCoding();
-			ShowConsole();
-		}
-		else if(Settings->Startup == ELiveCodingStartupMode::AutomaticButHidden)
-		{
-			GLiveCodingConsoleArguments = L"-Hidden";
-			StartLiveCoding();
-		}
-	}
-
-	if(FParse::Param(FCommandLine::Get(), TEXT("LiveCoding")))
+		StartLiveCodingAsync(Settings->Startup);
+	} 
+	else if (FParse::Param(FCommandLine::Get(), TEXT("LiveCoding")))
 	{
-		StartLiveCoding();
+		StartLiveCodingAsync(ELiveCodingStartupMode::Manual);
 	}
 
-	bEnabledLastTick = Settings->bEnabled;
 	bEnableReinstancingLastTick = IsReinstancingEnabled();
 }
 
@@ -364,10 +353,10 @@ void FLiveCodingModule::ShutdownModule()
 
 void FLiveCodingModule::EnableByDefault(bool bEnable)
 {
-	if(Settings->bEnabled != bEnable)
+	if (Settings->bEnabled != bEnable)
 	{
 		Settings->bEnabled = bEnable;
-		if(SettingsSection.IsValid())
+		if (SettingsSection.IsValid())
 		{
 			SettingsSection->Save();
 		}
@@ -397,32 +386,65 @@ void FLiveCodingModule::EnableForSession(bool bEnable)
 	if (bEnable)
 	{
 		EnableErrorText = FText::GetEmpty();
-		if(!bStarted)
+		switch (State)
 		{
-			StartLiveCoding();
+		case EState::NotRunning:
+			StartLiveCoding(ELiveCodingStartupMode::Manual); // State set in this method
 			ShowConsole();
-		}
-		else
-		{
-			bEnabledForSession = true;
+			break;
+
+		case EState::Starting:
+			WaitForStartup();
+			break;
+
+		case EState::Running:
+		case EState::RunningAndEnabled:
+			State = EState::RunningAndEnabled;
 			ShowConsole();
+			break;
 		}
 	}
 	else 
 	{
-		if(bStarted)
+		switch (State)
 		{
-			UE_LOG(LogLiveCoding, Display, TEXT("Console will be hidden but remain running in the background. Restart to disable completely."));
-			LppSetActive(false);
-			LppSetVisible(false);
-			bEnabledForSession = false;
+		case EState::NotRunning:
+		case EState::Starting:
+			break;
+
+		case EState::Running:
+		case EState::RunningAndEnabled:
+			HideConsole();
+			State = EState::Running;
+			break;
 		}
 	}
 }
 
+void FLiveCodingModule::WaitForStartup()
+{
+	if (State == EState::Starting)
+	{
+		UE_LOG(LogLiveCoding, Display, TEXT("Waiting for console to start."));
+		FPlatformProcess::ConditionalSleep([this]() { return State != EState::Starting; });
+	}
+}
+
 bool FLiveCodingModule::IsEnabledForSession() const
 {
-	return bEnabledForSession;
+	switch (State)
+	{
+	case EState::NotRunning:
+	case EState::Starting:
+	case EState::Running:
+		return false;
+
+	case EState::RunningAndEnabled:
+		return true;
+
+	default:
+		return false;
+	}
 }
 
 const FText& FLiveCodingModule::GetEnableErrorText() const
@@ -444,12 +466,36 @@ bool FLiveCodingModule::CanEnableForSession() const
 
 bool FLiveCodingModule::HasStarted() const
 {
-	return bStarted;
+	return HasStarted(false);
+}
+
+bool FLiveCodingModule::HasStarted(bool bAllowStarting) const
+{
+	switch (State)
+	{
+	case EState::NotRunning:
+		return false;
+
+	case EState::Starting:
+		return bAllowStarting;
+
+	case EState::Running:
+	case EState::RunningAndEnabled:
+		return true;
+
+	default:
+		return false;
+	}
 }
 
 void FLiveCodingModule::ShowConsole()
 {
-	if (bStarted)
+	ShowConsole(false);
+}
+
+void FLiveCodingModule::ShowConsole(bool bAllowStarting)
+{
+	if (HasStarted(bAllowStarting))
 	{
 		LppSetVisible(true);
 		LppSetActive(true);
@@ -457,6 +503,16 @@ void FLiveCodingModule::ShowConsole()
 	}
 }
 
+void FLiveCodingModule::HideConsole()
+{
+	if (HasStarted())
+	{
+		UE_LOG(LogLiveCoding, Display, TEXT("Console will be hidden but remain running in the background. Restart to disable completely."));
+		LppSetActive(false);
+		LppSetVisible(false);
+	}
+}
+
 void FLiveCodingModule::Compile()
 {
 	Compile(ELiveCodingCompileFlags::None, nullptr);
@@ -479,13 +535,13 @@ bool FLiveCodingModule::Compile(ELiveCodingCompileFlags CompileFlags, ELiveCodin
 	}
 
 	EnableForSession(true);
-	if (!bStarted)
+	if (!HasStarted())
 	{
 		return ReturnResults(ELiveCodingCompileResult::NotStarted, Result);
 	}
 
 	// Need to do this immediately rather than waiting until next tick
-	UpdateModules(); 
+	UpdateModules(false); 
 
 	// Trigger the recompile
 	GIsCompileActive = true;
@@ -525,33 +581,57 @@ bool FLiveCodingModule::IsCompiling() const
 
 void FLiveCodingModule::Tick()
 {
-	if (LppWantsRestart())
-	{
-		LppRestart(lpp::LPP_RESTART_BEHAVIOR_REQUEST_EXIT, 0);
-	}
 
-	if (Settings->bEnabled != bEnabledLastTick && Settings->Startup != ELiveCodingStartupMode::Manual)
+	// Check for a change in the last requested enable state if we are in automatic mode
+	if (Settings->Startup != ELiveCodingStartupMode::Manual)
 	{
-		EnableForSession(Settings->bEnabled);
-		bEnabledLastTick = Settings->bEnabled;
-		if (IsEnabledByDefault() && !IsEnabledForSession())
+		switch (State)
 		{
-			FMessageDialog::Open(EAppMsgType::Ok, LOCTEXT("NoEnableLiveCodingAfterHotReload", "Live Coding cannot be enabled while hot-reloaded modules are active. Please close the editor and build from your IDE before restarting."));
+		case EState::NotRunning:
+			if (Settings->bEnabled)
+			{
+				EnableForSession(true);
+				if (!IsEnabledForSession())
+				{
+					FMessageDialog::Open(EAppMsgType::Ok, EnableErrorText);
+				}
+			}
+			break;
+
+		case EState::Starting:
+		case EState::Running:
+			break;
+
+		case EState::RunningAndEnabled:
+			if (!Settings->bEnabled)
+			{
+				EnableForSession(false);
+			}
+			break;
 		}
 	}
-	else if (IsEnabledForSession() && IsReinstancingEnabled() != bEnableReinstancingLastTick)
-	{
-		bEnableReinstancingLastTick = IsReinstancingEnabled();
-		LppSetReinstancingFlow(bEnableReinstancingLastTick);
-	}
 
-	if (bUpdateModulesInTick)
+	if (HasStarted())
 	{
-		UpdateModules();
-		bUpdateModulesInTick = false;
-	}
+		if (LppWantsRestart())
+		{
+			LppRestart(lpp::LPP_RESTART_BEHAVIOR_REQUEST_EXIT, 0);
+		}
 
-	AttemptSyncLivePatching();
+		if (HasStarted() && IsReinstancingEnabled() != bEnableReinstancingLastTick)
+		{
+			bEnableReinstancingLastTick = IsReinstancingEnabled();
+			LppSetReinstancingFlow(bEnableReinstancingLastTick);
+		}
+
+		if (bUpdateModulesInTick)
+		{
+			UpdateModules(false);
+			bUpdateModulesInTick = false;
+		}
+
+		AttemptSyncLivePatching();
+	}
 }
 
 void FLiveCodingModule::AttemptSyncLivePatching()
@@ -781,16 +861,44 @@ ILiveCodingModule::FOnPatchCompleteDelegate& FLiveCodingModule::GetOnPatchComple
 	return OnPatchCompleteDelegate;
 }
 
-bool FLiveCodingModule::StartLiveCoding()
+void FLiveCodingModule::StartLiveCodingAsync(ELiveCodingStartupMode StartupMode)
+{
+	if (IsRunningCommandlet())
+	{
+		StartLiveCoding(StartupMode);
+	}
+	else
+	{
+		auto Task = [this, StartupMode]()
+		{
+			StartLiveCoding(StartupMode);
+		};
+
+		FFunctionGraphTask::CreateAndDispatchWhenReady(MoveTemp(Task), TStatId());
+	}
+}
+
+bool FLiveCodingModule::StartLiveCoding(ELiveCodingStartupMode StartupMode)
 {
 	EnableErrorText = FText::GetEmpty();
-	if(!bStarted)
+	if (!HasStarted())
 	{
+		State = EState::Starting;
+		if (StartupMode == ELiveCodingStartupMode::AutomaticButHidden)
+		{
+			GLiveCodingConsoleArguments = L"-Hidden";
+		}
+		else
+		{
+			GLiveCodingConsoleArguments = L"";
+		}
+
 		// Make sure there aren't any hot reload modules already active
 		if (!CanEnableForSession())
 		{
 			EnableErrorText = LOCTEXT("NoLiveCodingCompileAfterHotReload", "Live Coding cannot be enabled while hot-reloaded modules are active. Please close the editor and build from your IDE before restarting.");
 			UE_LOG(LogLiveCoding, Error, TEXT("Unable to start live coding session. Some modules have already been hot reloaded."));
+			State = EState::NotRunning;
 			return false;
 		}
 
@@ -803,6 +911,7 @@ bool FLiveCodingModule::StartLiveCoding()
 			const static FText FormatString = LOCTEXT("LiveCodingMissingExecutable", "Unable to start live coding session. Missing executable '{Executable}'. Use the LiveCoding.ConsolePath console variable to modify.");
 			EnableErrorText = FText::Format(FormatString, Args);
 			UE_LOG(LogLiveCoding, Error, TEXT("Unable to start live coding session. Missing executable '%s'. Use the LiveCoding.ConsolePath console variable to modify."), *GLiveCodingConsolePath);
+			State = EState::NotRunning;
 			return false;
 		}
 
@@ -815,6 +924,7 @@ bool FLiveCodingModule::StartLiveCoding()
 			const static FText FormatString = LOCTEXT("LiveCodingMissingProjectFile", "Unable to start live coding session. Unable to find source project file '{ProjectFile}'.");
 			EnableErrorText = FText::Format(FormatString, Args);
 			UE_LOG(LogLiveCoding, Error, TEXT("Unable to start live coding session. Unable to find source project file '%s'."), *SourceProject);
+			State = EState::NotRunning;
 			return false;
 		}
 
@@ -877,21 +987,25 @@ bool FLiveCodingModule::StartLiveCoding()
 		// Configure all the current modules. For non-commandlets, schedule it to be done in the first Tick() so we can batch everything together.
 		if (IsRunningCommandlet())
 		{
-			UpdateModules();
+			UpdateModules(true);
 		}
 		else
 		{
 			bUpdateModulesInTick = true;
 		}
 
+		if (StartupMode == ELiveCodingStartupMode::Automatic)
+		{
+			ShowConsole(true);
+		}
+
 		// Mark it as started
-		bStarted = true;
-		bEnabledForSession = true;
+		State = EState::RunningAndEnabled;
 	}
 	return true;
 }
 
-void FLiveCodingModule::UpdateModules()
+void FLiveCodingModule::UpdateModules(bool bAllowStarting)
 {
 	TArray<ModuleChange> Changes;
 	{
@@ -899,7 +1013,7 @@ void FLiveCodingModule::UpdateModules()
 		Swap(Changes, ModuleChanges);
 	}
 
-	if (bEnabledForSession)
+	if (HasStarted(bAllowStarting))
 	{
 #if IS_MONOLITHIC
 		wchar_t FullFilePath[WINDOWS_MAX_PATH];
@@ -970,7 +1084,7 @@ void FLiveCodingModule::OnModulesChanged(FName ModuleName, EModuleChangeReason R
 		// Assume that Tick() won't be called if we're running a commandlet
 		if (IsRunningCommandlet())
 		{
-			UpdateModules();
+			UpdateModules(false);
 		}
 		else
 		{
diff --git a/Engine/Source/Developer/Windows/LiveCoding/Private/LiveCodingModule.h b/Engine/Source/Developer/Windows/LiveCoding/Private/LiveCodingModule.h
--- a/Engine/Source/Developer/Windows/LiveCoding/Private/LiveCodingModule.h
+++ b/Engine/Source/Developer/Windows/LiveCoding/Private/LiveCodingModule.h
@@ -3,10 +3,12 @@
 #pragma once
 
 #include "ILiveCodingModule.h"
+#include "LiveCodingSettings.h"
 #include "Delegates/Delegate.h"
 #include "Modules/ModuleManager.h"
 #include "Templates/SharedPointer.h"
 #include "Internationalization/Text.h"
+#include <atomic>
 
 struct IConsoleCommand;
 class IConsoleVariable;
@@ -55,14 +57,35 @@ private:
 	void OnDllUnloaded(const FString& FullPath);
 	bool IsUEDll(const FString& FullPath);
 	bool IsPatchDll(const FString& FullPath);
+	void HideConsole();
+	void EnableConsoleCommand(FOutputDevice& out);
+	void StartLiveCodingAsync(ELiveCodingStartupMode StartupMode);
+	bool StartLiveCoding(ELiveCodingStartupMode StartupMode);
+	void OnModulesChanged(FName ModuleName, EModuleChangeReason Reason);
+	void UpdateModules(bool bAllowStarting);
+	bool ShouldPreloadModule(const TSet<FName>& PreloadedFileNames, const FString& FullFilePath) const;
+	bool IsReinstancingEnabled() const;
+	void WaitForStartup();
+	bool HasStarted(bool bAllowStarting) const;
+	void ShowConsole(bool bAllowStarting);
+
+#if WITH_EDITOR
+	void ShowNotification(bool Success, const FText& Title, const FText* SubText);
+#endif
 
 private:
+	enum class EState
+	{
+		NotRunning,
+		Starting,
+		Running,
+		RunningAndEnabled,
+	};
+
 	ULiveCodingSettings* Settings;
 	TSharedPtr<ISettingsSection> SettingsSection;
-	bool bEnabledLastTick = false;
 	bool bEnableReinstancingLastTick = false;
-	bool bEnabledForSession = false;
-	bool bStarted = false;
+	std::atomic<EState> State = EState::NotRunning;
 	bool bUpdateModulesInTick = false;
 	bool bHasReinstancingOccurred = false;
 	bool bHasPatchBeenLoaded = false;
@@ -99,21 +122,5 @@ private:
 #else
 	TUniquePtr<FNullReload> Reload;
 #endif
-
-	void EnableConsoleCommand(FOutputDevice& out);
-
-	bool StartLiveCoding();
-
-	void OnModulesChanged(FName ModuleName, EModuleChangeReason Reason);
-
-	void UpdateModules();
-
-	bool ShouldPreloadModule(const TSet<FName>& PreloadedFileNames, const FString& FullFilePath) const;
-
-	bool IsReinstancingEnabled() const;
-
-#if WITH_EDITOR
-	void ShowNotification(bool Success, const FText& Title, const FText* SubText);
-#endif
 };
 
diff --git a/Engine/Source/Developer/Windows/LiveCoding/Public/ILiveCodingModule.h b/Engine/Source/Developer/Windows/LiveCoding/Public/ILiveCodingModule.h
--- a/Engine/Source/Developer/Windows/LiveCoding/Public/ILiveCodingModule.h
+++ b/Engine/Source/Developer/Windows/LiveCoding/Public/ILiveCodingModule.h
@@ -31,17 +31,32 @@ enum class ELiveCodingCompileResult : uint8
 class ILiveCodingModule : public IModuleInterface
 {
 public:
+
+	// Enable by default tracks the value in the live coding settings.
+	// If enabled by default and the startup mode is one of the automatic settings, then
+	// live coding console will be started when the module is started.
 	virtual void EnableByDefault(bool bEnabled) = 0;
 	virtual bool IsEnabledByDefault() const = 0;
 
+	// The enabled for session state and started state are NOT independent states.  
+	// If a request is made to enabled live coding (automatically or manually) and live coding 
+	// hasn't already been started, then it is started.  If and only if the console starts 
+	// will both the started state and the enabled for session state will be true.
+	// If enable for session is set to false, the console isn't stopped so the started state
+	// will remain true.
+	//
+	// Started = false, Enabled = false => no request to enable has ever been made or it failed to start
+	// Started = true, Enabled = true => a request to enable was made and it started properly
+	// Started = true, Enabled = false => has been started but later disabled
+	// Started = false, Enabled = true => impossible state.
+	virtual bool HasStarted() const = 0;
 	virtual void EnableForSession(bool bEnabled) = 0;
 	virtual bool IsEnabledForSession() const = 0;
 	virtual bool CanEnableForSession() const = 0;
+
 	virtual const FText& GetEnableErrorText() const = 0;
 	virtual bool AutomaticallyCompileNewClasses() const = 0;
 
-	virtual bool HasStarted() const = 0;
-
 	virtual void ShowConsole() = 0;
 	virtual void Compile() = 0;
 	virtual bool Compile(ELiveCodingCompileFlags CompileFlags, ELiveCodingCompileResult* Result) = 0;
