* Fixed so preprocessor conforms to C99 and C++11  (/Zc:preprocessor)

Not proud over the LC_* files but don't dare touching the actual log macros since that is hard to get right and I don't know how to test this code

diff --git a/Engine/Source/Developer/Windows/LiveCoding/Private/External/LC_API.cpp b/Engine/Source/Developer/Windows/LiveCoding/Private/External/LC_API.cpp
--- a/Engine/Source/Developer/Windows/LiveCoding/Private/External/LC_API.cpp
+++ b/Engine/Source/Developer/Windows/LiveCoding/Private/External/LC_API.cpp
@@ -31,7 +31,7 @@ namespace
 	{
 		if (!g_startupThread)
 		{
-			LC_ERROR_USER("Live++ was not started properly. Call lppStartup() first.");
+			LC_ERROR_USER("%s", "Live++ was not started properly. Call lppStartup() first.");
 			return false;
 		}
 
diff --git a/Engine/Source/Developer/Windows/LiveCoding/Private/External/LC_ClientStartupThread.cpp b/Engine/Source/Developer/Windows/LiveCoding/Private/External/LC_ClientStartupThread.cpp
--- a/Engine/Source/Developer/Windows/LiveCoding/Private/External/LC_ClientStartupThread.cpp
+++ b/Engine/Source/Developer/Windows/LiveCoding/Private/External/LC_ClientStartupThread.cpp
@@ -529,7 +529,7 @@ Thread::ReturnValue ClientStartupThread::ThreadFunction(const std::wstring& proc
 	{
 		// we were unable to open the process, bail out
 		// BEGIN EPIC MOD
-		LC_ERROR_USER("Unable to attach to console process");
+		LC_ERROR_USER("%s", "Unable to attach to console process");
 		// END EPIC MOD
 		Process::DestroyNamedSharedMemory(m_sharedMemory);
 
@@ -538,7 +538,7 @@ Thread::ReturnValue ClientStartupThread::ThreadFunction(const std::wstring& proc
 
 	// wait for server to become ready
 	{
-		LC_LOG_USER("Waiting for server");
+		LC_LOG_USER("%s", "Waiting for server");
 
 		Event serverReadyEvent(primitiveNames::ServerReadyEvent(processGroupName).c_str(), Event::Type::AUTO_RESET);
 		serverReadyEvent.Wait();
@@ -549,7 +549,7 @@ Thread::ReturnValue ClientStartupThread::ThreadFunction(const std::wstring& proc
 	{
 		// could not connect to Live++ process
 		// BEGIN EPIC MOD
-		LC_ERROR_USER("Could not connect named pipe to console process");
+		LC_ERROR_USER("%s", "Could not connect named pipe to console process");
 		// END EPIC MOD
 
 		return Thread::ReturnValue(2u);
@@ -560,7 +560,7 @@ Thread::ReturnValue ClientStartupThread::ThreadFunction(const std::wstring& proc
 	{
 		// could not connect to Live++ process
 		// BEGIN EPIC MOD
-		LC_ERROR_USER("Could not connect exception pipe to console process");
+		LC_ERROR_USER("%s", "Could not connect exception pipe to console process");
 		// END EPIC MOD
 
 		return Thread::ReturnValue(3u);
@@ -634,7 +634,7 @@ Thread::ReturnValue ClientStartupThread::ThreadFunction(const std::wstring& proc
 	{
 		// process could not be registered, bail out
 		// BEGIN EPIC MOD
-		LC_ERROR_USER("Could not register live coding process");
+		LC_ERROR_USER("%s", "Could not register live coding process");
 		// END EPIC MOD
 
 		// close the pipe and then wait for the helper threads to finish.
@@ -666,7 +666,7 @@ Thread::ReturnValue ClientStartupThread::ThreadFunction(const std::wstring& proc
 		return Thread::ReturnValue(3u);
 	}
 
-	LC_LOG_USER("Successfully initialized, removing startup thread");
+	LC_LOG_USER("%s", "Successfully initialized, removing startup thread");
 
 	// helper threads are now allowed to run, we're finished with the pipe
 	m_startEvent->Signal();
diff --git a/Engine/Source/Developer/Windows/LiveCoding/Private/External/LC_Event.cpp b/Engine/Source/Developer/Windows/LiveCoding/Private/External/LC_Event.cpp
--- a/Engine/Source/Developer/Windows/LiveCoding/Private/External/LC_Event.cpp
+++ b/Engine/Source/Developer/Windows/LiveCoding/Private/External/LC_Event.cpp
@@ -64,16 +64,16 @@ bool Event::WaitTimeout(unsigned int milliSeconds)
 			// the operation timed out, which should never happen with a timeout of INFINITE
 			if (milliSeconds == INFINITE)
 			{
-				LC_ERROR_DEV("Event timed out.");
+				LC_ERROR_DEV("%s", "Event timed out.");
 			}
 			return false;
 
 		case WAIT_ABANDONED:
-			LC_ERROR_DEV("Wait() was called on a stale event which was not released by the owning thread.");
+			LC_ERROR_DEV("%s", "Wait() was called on a stale event which was not released by the owning thread.");
 			return false;
 
 		case WAIT_FAILED:
-			LC_ERROR_DEV("Failed to Wait() on an event.");
+			LC_ERROR_DEV("%s", "Failed to Wait() on an event.");
 			return false;
 
 		default:
diff --git a/Engine/Source/Developer/Windows/LiveCoding/Private/External/LC_Executable.cpp b/Engine/Source/Developer/Windows/LiveCoding/Private/External/LC_Executable.cpp
--- a/Engine/Source/Developer/Windows/LiveCoding/Private/External/LC_Executable.cpp
+++ b/Engine/Source/Developer/Windows/LiveCoding/Private/External/LC_Executable.cpp
@@ -35,14 +35,14 @@ namespace detail
 		const IMAGE_DOS_HEADER* dosHeader = static_cast<const IMAGE_DOS_HEADER*>(base);
 		if (dosHeader->e_magic != IMAGE_DOS_SIGNATURE)
 		{
-			LC_ERROR_USER("Image has unknown file format");
+			LC_ERROR_USER("%s", "Image has unknown file format");
 			return nullptr;
 		}
 
 		const IMAGE_NT_HEADERS* ntHeader = pointer::Offset<const IMAGE_NT_HEADERS*>(dosHeader, dosHeader->e_lfanew);
 		if (ntHeader->Signature != IMAGE_NT_SIGNATURE)
 		{
-			LC_ERROR_USER("Invalid .exe file");
+			LC_ERROR_USER("%s", "Invalid .exe file");
 			return nullptr;
 		}
 
diff --git a/Engine/Source/Developer/Windows/LiveCoding/Private/External/LC_InterprocessMutex.cpp b/Engine/Source/Developer/Windows/LiveCoding/Private/External/LC_InterprocessMutex.cpp
--- a/Engine/Source/Developer/Windows/LiveCoding/Private/External/LC_InterprocessMutex.cpp
+++ b/Engine/Source/Developer/Windows/LiveCoding/Private/External/LC_InterprocessMutex.cpp
@@ -34,15 +34,15 @@ void InterprocessMutex::Lock(void)
 
 		case WAIT_TIMEOUT:
 			// the operation timed out, which should never happen with a timeout of INFINITE
-			LC_ERROR_DEV("Mutex timed out.");
+			LC_ERROR_DEV("%s", "Mutex timed out.");
 			break;
 
 		case WAIT_ABANDONED:
-			LC_ERROR_DEV("Wait() was called on a stale mutex which was not released by the owning thread.");
+			LC_ERROR_DEV("%s", "Wait() was called on a stale mutex which was not released by the owning thread.");
 			break;
 
 		case WAIT_FAILED:
-			LC_ERROR_DEV("Failed to Wait() on a mutex.");
+			LC_ERROR_DEV("%s", "Failed to Wait() on a mutex.");
 			break;
 
 		default:
diff --git a/Engine/Source/Developer/Windows/LiveCoding/Private/External/LC_Process.cpp b/Engine/Source/Developer/Windows/LiveCoding/Private/External/LC_Process.cpp
--- a/Engine/Source/Developer/Windows/LiveCoding/Private/External/LC_Process.cpp
+++ b/Engine/Source/Developer/Windows/LiveCoding/Private/External/LC_Process.cpp
@@ -204,7 +204,7 @@ Process::Context* Process::Spawn(const wchar_t* exePath, const wchar_t* workingD
 		::wcscpy_s(commandLineBuffer, 32768u, commandLine);
 	}
 
-	LC_LOG_DEV("Spawning process:");
+	LC_LOG_DEV("%s", "Spawning process:");
 	{
 		LC_LOG_INDENT_DEV;
 		LC_LOG_DEV("Executable: %S", exePath);
@@ -742,7 +742,7 @@ void Process::DumpMemory(Handle handle, const void* address, size_t size)
 	uint8_t* memory = new uint8_t[size];
 	ReadProcessMemory(handle, address, memory, size);
 
-	LC_LOG_DEV("Raw data:");
+	LC_LOG_DEV("%s", "Raw data:");
 	LC_LOG_INDENT_DEV;
 	for (size_t i = 0u; i < size; ++i)
 	{
diff --git a/Engine/Source/Developer/Windows/LiveCoding/Private/External/LC_Semaphore.cpp b/Engine/Source/Developer/Windows/LiveCoding/Private/External/LC_Semaphore.cpp
--- a/Engine/Source/Developer/Windows/LiveCoding/Private/External/LC_Semaphore.cpp
+++ b/Engine/Source/Developer/Windows/LiveCoding/Private/External/LC_Semaphore.cpp
@@ -46,16 +46,16 @@ bool Semaphore::WaitTimeout(unsigned int milliSeconds)
 			// the operation timed out, which should never happen with a timeout of INFINITE
 			if (milliSeconds == INFINITE)
 			{
-				LC_ERROR_DEV("Semaphore timed out.");
+				LC_ERROR_DEV("%s", "Semaphore timed out.");
 			}
 			return false;
 
 		case WAIT_ABANDONED:
-			LC_ERROR_DEV("Wait() was called on a stale semaphore which was not released by the owning thread.");
+			LC_ERROR_DEV("%s", "Wait() was called on a stale semaphore which was not released by the owning thread.");
 			return false;
 
 		case WAIT_FAILED:
-			LC_ERROR_DEV("Failed to Wait() on a semaphore.");
+			LC_ERROR_DEV("%s", "Failed to Wait() on a semaphore.");
 			return false;
 
 		default:
diff --git a/Engine/Source/Runtime/D3D12RHI/Private/D3D12CommandContext.h b/Engine/Source/Runtime/D3D12RHI/Private/D3D12CommandContext.h
--- a/Engine/Source/Runtime/D3D12RHI/Private/D3D12CommandContext.h
+++ b/Engine/Source/Runtime/D3D12RHI/Private/D3D12CommandContext.h
@@ -715,8 +715,8 @@ public:
 
 	virtual FD3D12CommandContextRedirector* AsRedirector() override { return this; }
 
-#define ContextRedirect(Call) { for (uint32 GPUIndex : GPUMask) PhysicalContexts[GPUIndex]->##Call; }
-#define ContextGPU0(Call) { PhysicalContexts[0]->##Call; }
+#define ContextRedirect(Call) { for (uint32 GPUIndex : GPUMask) PhysicalContexts[GPUIndex]-> Call; }
+#define ContextGPU0(Call) { PhysicalContexts[0]-> Call; }
 
 	FORCEINLINE virtual void RHISetComputePipelineState(FRHIComputePipelineState* ComputePipelineState) final override
 	{
diff --git a/Engine/Source/Runtime/D3D12RHI/Private/D3D12PipelineState.h b/Engine/Source/Runtime/D3D12RHI/Private/D3D12PipelineState.h
--- a/Engine/Source/Runtime/D3D12RHI/Private/D3D12PipelineState.h
+++ b/Engine/Source/Runtime/D3D12RHI/Private/D3D12PipelineState.h
@@ -158,11 +158,11 @@ struct FD3D12ComputePipelineStateDesc
 };
 
 
-#define PSO_IF_NOT_EQUAL_RETURN_FALSE( value ) if(lhs.##value != rhs.##value){ return false; }
-#define PSO_IF_MEMCMP_FAILS_RETURN_FALSE( value ) if(FMemory::Memcmp(&lhs.##value, &rhs.##value, sizeof(rhs.##value)) != 0){ return false; }
+#define PSO_IF_NOT_EQUAL_RETURN_FALSE( value ) if(lhs.value != rhs.value){ return false; }
+#define PSO_IF_MEMCMP_FAILS_RETURN_FALSE( value ) if(FMemory::Memcmp(&lhs.value, &rhs.value, sizeof(rhs.value)) != 0){ return false; }
 #define PSO_IF_STRING_COMPARE_FAILS_RETURN_FALSE( value ) \
-	const char* const lhString = lhs.##value##; \
-	const char* const rhString = rhs.##value##; \
+	const char* const lhString = lhs.value; \
+	const char* const rhString = rhs.value; \
 	if (lhString != rhString) \
 	{ \
 		if (strcmp(lhString, rhString) != 0) \
diff --git a/Engine/Source/Runtime/D3D12RHI/Private/D3D12Pipelinestate.cpp b/Engine/Source/Runtime/D3D12RHI/Private/D3D12Pipelinestate.cpp
--- a/Engine/Source/Runtime/D3D12RHI/Private/D3D12Pipelinestate.cpp
+++ b/Engine/Source/Runtime/D3D12RHI/Private/D3D12Pipelinestate.cpp
@@ -66,7 +66,7 @@ static FD3D12LowLevelGraphicsPipelineStateDesc GetLowLevelGraphicsPipelineStateD
 	{ \
 		if (Shader->VendorExtensions.Num() > 0) \
 		{ \
-			Desc.##Initial##SExtensions = &Shader->VendorExtensions; \
+			Desc.Initial##SExtensions = &Shader->VendorExtensions; \
 		} \
 	}
 	EXT_SHADER(V, Vertex);
diff --git a/Engine/Source/Runtime/D3D12RHI/Private/D3D12StateCache.cpp b/Engine/Source/Runtime/D3D12RHI/Private/D3D12StateCache.cpp
--- a/Engine/Source/Runtime/D3D12RHI/Private/D3D12StateCache.cpp
+++ b/Engine/Source/Runtime/D3D12RHI/Private/D3D12StateCache.cpp
@@ -693,9 +693,9 @@ void FD3D12StateCache::ApplyResources(const FD3D12RootSignature* const pRootSign
 		FD3D12ShaderResourceViewCache& SRVCache = PipelineState.Common.SRVCache;
 
 #define CONDITIONAL_SET_SRVS(Shader) \
-		if (CurrentShaderDirtySRVSlots[##Shader]) \
+		if (CurrentShaderDirtySRVSlots[Shader]) \
 		{ \
-			DescriptorCache.SetSRVs<##Shader>(pRootSignature, SRVCache, CurrentShaderDirtySRVSlots[##Shader], NumSRVs[##Shader], ViewHeapSlot); \
+			DescriptorCache.SetSRVs<Shader>(pRootSignature, SRVCache, CurrentShaderDirtySRVSlots[Shader], NumSRVs[Shader], ViewHeapSlot); \
 		}
 
 		if (StartStage == SF_Compute)
@@ -721,9 +721,9 @@ void FD3D12StateCache::ApplyResources(const FD3D12RootSignature* const pRootSign
 		FD3D12ConstantBufferCache& CBVCache = PipelineState.Common.CBVCache;
 
 	#define CONDITIONAL_SET_CBVS(Shader) \
-		if (CurrentShaderDirtyCBVSlots[##Shader]) \
+		if (CurrentShaderDirtyCBVSlots[Shader]) \
 		{ \
-			DescriptorCache.SetConstantBufferViews(##Shader, pRootSignature, CBVCache, CurrentShaderDirtyCBVSlots[##Shader], NumCBVs[##Shader], ViewHeapSlot); \
+			DescriptorCache.SetConstantBufferViews(Shader, pRootSignature, CBVCache, CurrentShaderDirtyCBVSlots[Shader], NumCBVs[Shader], ViewHeapSlot); \
 		}
 
 		if (StartStage == SF_Compute)
@@ -928,9 +928,9 @@ void FD3D12StateCache::ApplySamplers(const FD3D12RootSignature* const pRootSigna
 	uint32 SamplerHeapSlot = SamplerHeap->ReserveSlots(NumSamplers[SF_NumStandardFrequencies]);
 
 #define CONDITIONAL_SET_SAMPLERS(Shader) \
-	if (CurrentShaderDirtySamplerSlots[##Shader]) \
+	if (CurrentShaderDirtySamplerSlots[Shader]) \
 	{ \
-		DescriptorCache.SetSamplers<##Shader>(pRootSignature, Cache, CurrentShaderDirtySamplerSlots[##Shader], NumSamplers[##Shader], SamplerHeapSlot); \
+		DescriptorCache.SetSamplers<Shader>(pRootSignature, Cache, CurrentShaderDirtySamplerSlots[Shader], NumSamplers[Shader], SamplerHeapSlot); \
 	}
 
 	if (StartStage == SF_Compute)
diff --git a/Engine/Source/Runtime/D3D12RHI/Private/Windows/WindowsD3D12PipelineState.cpp b/Engine/Source/Runtime/D3D12RHI/Private/Windows/WindowsD3D12PipelineState.cpp
--- a/Engine/Source/Runtime/D3D12RHI/Private/Windows/WindowsD3D12PipelineState.cpp
+++ b/Engine/Source/Runtime/D3D12RHI/Private/Windows/WindowsD3D12PipelineState.cpp
@@ -940,9 +940,9 @@ static void CreateGraphicsPipelineState(ID3D12PipelineState** PSO, FD3D12Adapter
 		TArray<FShaderCodeVendorExtension, TInlineAllocator<2 /* VS + PS */>> MergedExtensions;
 
 	#define MERGE_EXT(Initial) \
-		if (CreationArgs->Desc.##Initial##SExtensions != nullptr) \
+		if (CreationArgs->Desc.Initial##SExtensions != nullptr) \
 		{ \
-			for (const FShaderCodeVendorExtension& Extension : *CreationArgs->Desc.##Initial##SExtensions) \
+			for (const FShaderCodeVendorExtension& Extension : *CreationArgs->Desc.Initial##SExtensions) \
 			{ \
 				MergedExtensions.AddUnique(Extension); \
 			} \
diff --git a/Engine/Source/Runtime/VulkanRHI/Private/Windows/VulkanWindowsPlatform.cpp b/Engine/Source/Runtime/VulkanRHI/Private/Windows/VulkanWindowsPlatform.cpp
--- a/Engine/Source/Runtime/VulkanRHI/Private/Windows/VulkanWindowsPlatform.cpp
+++ b/Engine/Source/Runtime/VulkanRHI/Private/Windows/VulkanWindowsPlatform.cpp
@@ -81,7 +81,7 @@ bool FVulkanWindowsPlatform::LoadVulkanLibrary()
 
 	if (GVulkanDLLModule)
 	{
-#define GET_VK_ENTRYPOINTS(Type,Func) VulkanDynamicAPI::Func = (Type)FPlatformProcess::GetDllExport(GVulkanDLLModule, L#Func);
+#define GET_VK_ENTRYPOINTS(Type,Func) VulkanDynamicAPI::Func = (Type)FPlatformProcess::GetDllExport(GVulkanDLLModule, L""#Func);
 		ENUM_VK_ENTRYPOINTS_BASE(GET_VK_ENTRYPOINTS);
 
 		bool bFoundAllEntryPoints = true;
