Abstract number of instructions for complexity visualization by estimated metric.

diff --git a/Engine/Source/Developer/Windows/ShaderFormatD3D/Private/D3DShaderCompilerDXC.cpp b/Engine/Source/Developer/Windows/ShaderFormatD3D/Private/D3DShaderCompilerDXC.cpp
--- a/Engine/Source/Developer/Windows/ShaderFormatD3D/Private/D3DShaderCompilerDXC.cpp
+++ b/Engine/Source/Developer/Windows/ShaderFormatD3D/Private/D3DShaderCompilerDXC.cpp
@@ -963,6 +963,11 @@ bool CompileAndProcessD3DShaderDXC(FString& PreprocessedShaderSource,
 				}
 			};
 
+			// Return a fraction of the number of instructions as DXIL is more verbose than DXBC.
+			// Ratio 119:307 was estimated by gathering average instruction count for D3D11 and D3D12 shaders in ShooterGame with result being ~ 357:921.
+			constexpr uint32 DxbcToDxilInstructionRatio[2] = { 119, 307 };
+			NumInstructions = NumInstructions * DxbcToDxilInstructionRatio[0] / DxbcToDxilInstructionRatio[1];
+
 			//#todo-rco: Should compress ShaderCode?
 
 			GenerateFinalOutput(ShaderBlob,
