Fixed issue where Live Coding Console would fail to properly enable when it is started on demand.

diff --git a/Engine/Source/Developer/Windows/LiveCoding/Private/LiveCodingModule.cpp b/Engine/Source/Developer/Windows/LiveCoding/Private/LiveCodingModule.cpp
--- a/Engine/Source/Developer/Windows/LiveCoding/Private/LiveCodingModule.cpp
+++ b/Engine/Source/Developer/Windows/LiveCoding/Private/LiveCodingModule.cpp
@@ -319,6 +319,7 @@ void FLiveCodingModule::StartupModule()
 
 	LppStartup();
 
+	bSettingsEnabledLastTick = Settings->bEnabled;
 	if (Settings->bEnabled && Settings->Startup != ELiveCodingStartupMode::Manual && !FApp::IsUnattended())
 	{
 		StartLiveCodingAsync(Settings->Startup);
@@ -583,11 +584,12 @@ void FLiveCodingModule::Tick()
 {
 
 	// Check for a change in the last requested enable state if we are in automatic mode
-	if (Settings->Startup != ELiveCodingStartupMode::Manual)
+	if (Settings->bEnabled != bSettingsEnabledLastTick && Settings->Startup != ELiveCodingStartupMode::Manual)
 	{
 		switch (State)
 		{
 		case EState::NotRunning:
+			bSettingsEnabledLastTick = Settings->bEnabled;
 			if (Settings->bEnabled)
 			{
 				EnableForSession(true);
@@ -603,6 +605,7 @@ void FLiveCodingModule::Tick()
 			break;
 
 		case EState::RunningAndEnabled:
+			bSettingsEnabledLastTick = Settings->bEnabled;
 			if (!Settings->bEnabled)
 			{
 				EnableForSession(false);
diff --git a/Engine/Source/Developer/Windows/LiveCoding/Private/LiveCodingModule.h b/Engine/Source/Developer/Windows/LiveCoding/Private/LiveCodingModule.h
--- a/Engine/Source/Developer/Windows/LiveCoding/Private/LiveCodingModule.h
+++ b/Engine/Source/Developer/Windows/LiveCoding/Private/LiveCodingModule.h
@@ -84,6 +84,7 @@ private:
 
 	ULiveCodingSettings* Settings;
 	TSharedPtr<ISettingsSection> SettingsSection;
+	bool bSettingsEnabledLastTick = false;
 	bool bEnableReinstancingLastTick = false;
 	std::atomic<EState> State = EState::NotRunning;
 	bool bUpdateModulesInTick = false;
