Log available system memory when DXC crashes on E_OUTOFMEMORY to verify system is in fact out of memory.

#rnx

diff --git a/Engine/Source/Developer/Windows/ShaderFormatD3D/Private/D3DShaderCompilerDXC.cpp b/Engine/Source/Developer/Windows/ShaderFormatD3D/Private/D3DShaderCompilerDXC.cpp
--- a/Engine/Source/Developer/Windows/ShaderFormatD3D/Private/D3DShaderCompilerDXC.cpp
+++ b/Engine/Source/Developer/Windows/ShaderFormatD3D/Private/D3DShaderCompilerDXC.cpp
@@ -82,7 +82,33 @@ static uint32 GetAutoBindingSpace(const FShaderTarget& Target)
 // Utility variable so we can place a breakpoint while debugging
 static int32 GBreakpointDXC = 0;
 
-#define VERIFYHRESULT(expr) { HRESULT HR##__LINE__ = expr; if (FAILED(HR##__LINE__)) { UE_LOG(LogD3D12ShaderCompiler, Fatal, TEXT(#expr " failed: Result=%08x"), HR##__LINE__); } }
+static void LogFailedHRESULT(const TCHAR* FailedExpressionStr, HRESULT Result)
+{
+	if (Result == E_OUTOFMEMORY)
+	{
+		// Provide currently availble system memory to verify out-of-memory error is legitimate
+		const FPlatformMemoryStats MemoryStats = FPlatformMemory::GetStats();
+		FString SystemMemoryStr = FString::Printf(TEXT("%d MB"), MemoryStats.AvailablePhysical / 1024 / 1024);
+		FString UsedMemoryStr = FString::Printf(TEXT("%d MB"), MemoryStats.UsedPhysical / 1024 / 1024);
+		UE_LOG(LogD3D12ShaderCompiler, Fatal, TEXT("%s failed: Result=0x%08x (E_OUTOFMEMORY); System memory %s, used %s"), FailedExpressionStr, Result, *SystemMemoryStr, *UsedMemoryStr);
+	}
+	else
+	{
+		// Turn HRESULT into human readable string for error report
+		TCHAR ResultStr[4096] = {};
+		FPlatformMisc::GetSystemErrorMessage(ResultStr, UE_ARRAY_COUNT(ResultStr), Result);
+		UE_LOG(LogD3D12ShaderCompiler, Fatal, TEXT("%s failed: Result=0x%08x (%s)"), FailedExpressionStr, Result, ResultStr);
+	}
+}
+
+#define VERIFYHRESULT(expr)									\
+	{														\
+		HRESULT HR##__LINE__ = expr;						\
+		if (FAILED(HR##__LINE__))							\
+		{													\
+			LogFailedHRESULT(TEXT(#expr), HR##__LINE__);	\
+		}													\
+	}
 
 static dxc::DxcDllSupport& GetDxcDllHelper()
 {
