Fix "Compile failed without errors" message and unify this error report for FXC and DXC.

#rnx

diff --git a/Engine/Source/Developer/Windows/ShaderFormatD3D/Private/D3DShaderCompiler.cpp b/Engine/Source/Developer/Windows/ShaderFormatD3D/Private/D3DShaderCompiler.cpp
--- a/Engine/Source/Developer/Windows/ShaderFormatD3D/Private/D3DShaderCompiler.cpp
+++ b/Engine/Source/Developer/Windows/ShaderFormatD3D/Private/D3DShaderCompiler.cpp
@@ -575,7 +575,8 @@ static bool CompileErrorsContainInternalError(ID3DBlob* Errors)
 }
 
 // Generate the dumped usf file; call the D3D compiler, gather reflection information and generate the output data
-bool CompileAndProcessD3DShaderFXC(FString& PreprocessedShaderSource, const FString& CompilerPath,
+static bool CompileAndProcessD3DShaderFXCExt(
+	FString& PreprocessedShaderSource, const FString& CompilerPath,
 	uint32 CompileFlags,
 	const FShaderCompilerInput& Input,
 	const FShaderParameterParser& ShaderParameterParser,
@@ -583,7 +584,7 @@ bool CompileAndProcessD3DShaderFXC(FString& PreprocessedShaderSource, const FStr
 	const TCHAR* ShaderProfile, bool bSecondPassAferUnusedInputRemoval,
 	TArray<FString>& FilteredErrors, FShaderCompilerOutput& Output)
 {
-	TRACE_CPUPROFILER_EVENT_SCOPE(CompileAndProcessD3DShaderFXC);
+	TRACE_CPUPROFILER_EVENT_SCOPE(CompileAndProcessD3DShaderFXCExt);
 
 	auto AnsiSourceFile = StringCast<ANSICHAR>(*PreprocessedShaderSource);
 
@@ -886,7 +887,7 @@ bool CompileAndProcessD3DShaderFXC(FString& PreprocessedShaderSource, const FStr
 						if (RemoveUnusedInputs(PreprocessedShaderSource, ShaderInputs, EntryPointName, RemoveErrors))
 						{
 							Output = OriginalOutput;
-							if (!CompileAndProcessD3DShaderFXC(PreprocessedShaderSource, CompilerPath, CompileFlags, Input, ShaderParameterParser, EntryPointName, ShaderProfile, true, FilteredErrors, Output))
+							if (!CompileAndProcessD3DShaderFXCExt(PreprocessedShaderSource, CompilerPath, CompileFlags, Input, ShaderParameterParser, EntryPointName, ShaderProfile, true, FilteredErrors, Output))
 							{
 								// if we failed to compile the shader, propagate the error up
 								return false;
@@ -1066,6 +1067,50 @@ bool CompileAndProcessD3DShaderFXC(FString& PreprocessedShaderSource, const FStr
 	return SUCCEEDED(Result);
 }
 
+bool CompileAndProcessD3DShaderFXC(FString& PreprocessedShaderSource,
+	uint32 CompileFlags,
+	const FShaderCompilerInput& Input,
+	const FShaderParameterParser& ShaderParameterParser,
+	FString& EntryPointName,
+	const TCHAR* ShaderProfile, bool bSecondPassAferUnusedInputRemoval,
+	FShaderCompilerOutput& Output)
+{
+	// Override default compiler path to newer dll
+	FString CompilerPath = FPaths::EngineDir();
+	CompilerPath.Append(TEXT("Binaries/ThirdParty/Windows/DirectX/x64/d3dcompiler_47.dll"));
+
+	TArray<FString> FilteredErrors;
+	const bool bSuccess = CompileAndProcessD3DShaderFXCExt(PreprocessedShaderSource, CompilerPath, CompileFlags, Input, ShaderParameterParser, EntryPointName, ShaderProfile, false, FilteredErrors, Output);
+
+	// Process errors
+	for (int32 ErrorIndex = 0; ErrorIndex < FilteredErrors.Num(); ErrorIndex++)
+	{
+		const FString& CurrentError = FilteredErrors[ErrorIndex];
+		FShaderCompilerError NewError;
+
+		// Extract filename and line number from FXC output with format:
+		// "d:\Project\Binaries\BasePassPixelShader(30,7): error X3000: invalid target or usage string"
+		int32 FirstParenIndex = CurrentError.Find(TEXT("("));
+		int32 LastParenIndex = CurrentError.Find(TEXT("):"));
+		if (FirstParenIndex != INDEX_NONE &&
+			LastParenIndex != INDEX_NONE &&
+			LastParenIndex > FirstParenIndex)
+		{
+			// Extract and store error message with source filename
+			NewError.ErrorVirtualFilePath = CurrentError.Left(FirstParenIndex);
+			NewError.ErrorLineString = CurrentError.Mid(FirstParenIndex + 1, LastParenIndex - FirstParenIndex - FCString::Strlen(TEXT("(")));
+			NewError.StrippedErrorMessage = CurrentError.Right(CurrentError.Len() - LastParenIndex - FCString::Strlen(TEXT("):")));
+		}
+		else
+		{
+			NewError.StrippedErrorMessage = CurrentError;
+		}
+		Output.Errors.Add(NewError);
+	}
+
+	return bSuccess;
+}
+
 void CompileD3DShader(const FShaderCompilerInput& Input, FShaderCompilerOutput& Output, FShaderCompilerDefinitions& AdditionalDefines, const FString& WorkingDirectory, ELanguage Language)
 {
 	TRACE_CPUPROFILER_EVENT_SCOPE(CompileD3DShader);
@@ -1241,56 +1286,13 @@ void CompileD3DShader(const FShaderCompilerInput& Input, FShaderCompilerOutput&
 			CompileFlags |= TranslateCompilerFlagD3D11((ECompilerFlags)Flag);
 		});
 
-	if (bUseDXC)
-	{
-		if (!CompileAndProcessD3DShaderDXC(PreprocessedShaderSource, CompileFlags, Input, ShaderParameterParser, EntryPointName, ShaderProfile, Language, false, FilteredErrors, Output))
-		{
-			if (!FilteredErrors.Num())
-			{
-				FilteredErrors.Add(TEXT("Compile Failed without errors!"));
-			}
-		}
-		CrossCompiler::FShaderConductorContext::ConvertCompileErrors(MoveTemp(FilteredErrors), Output.Errors);
-	}
-	else
-	{
-		// Override default compiler path to newer dll
-		FString CompilerPath = FPaths::EngineDir();
-		CompilerPath.Append(TEXT("Binaries/ThirdParty/Windows/DirectX/x64/d3dcompiler_47.dll"));
-
-		if (!CompileAndProcessD3DShaderFXC(PreprocessedShaderSource, CompilerPath, CompileFlags, Input, ShaderParameterParser, EntryPointName, ShaderProfile, false, FilteredErrors, Output))
-		{
-			if (!FilteredErrors.Num())
-			{
-				FilteredErrors.Add(TEXT("Compile Failed without errors!"));
-			}
-		}
+	const bool bSuccess = bUseDXC
+		? CompileAndProcessD3DShaderDXC(PreprocessedShaderSource, CompileFlags, Input, ShaderParameterParser, EntryPointName, ShaderProfile, Language, false, Output)
+		: CompileAndProcessD3DShaderFXC(PreprocessedShaderSource, CompileFlags, Input, ShaderParameterParser, EntryPointName, ShaderProfile, false, Output);
 
-		// Process errors
-		for (int32 ErrorIndex = 0; ErrorIndex < FilteredErrors.Num(); ErrorIndex++)
-		{
-			const FString& CurrentError = FilteredErrors[ErrorIndex];
-			FShaderCompilerError NewError;
-
-			// Extract filename and line number from FXC output with format:
-			// "d:\Project\Binaries\BasePassPixelShader(30,7): error X3000: invalid target or usage string"
-			int32 FirstParenIndex = CurrentError.Find(TEXT("("));
-			int32 LastParenIndex = CurrentError.Find(TEXT("):"));
-			if (FirstParenIndex != INDEX_NONE &&
-				LastParenIndex != INDEX_NONE &&
-				LastParenIndex > FirstParenIndex)
-			{
-				// Extract and store error message with source filename
-				NewError.ErrorVirtualFilePath = CurrentError.Left(FirstParenIndex);
-				NewError.ErrorLineString = CurrentError.Mid(FirstParenIndex + 1, LastParenIndex - FirstParenIndex - FCString::Strlen(TEXT("(")));
-				NewError.StrippedErrorMessage = CurrentError.Right(CurrentError.Len() - LastParenIndex - FCString::Strlen(TEXT("):")));
-			}
-			else
-			{
-				NewError.StrippedErrorMessage = CurrentError;
-			}
-			Output.Errors.Add(NewError);
-		}
+	if (!bSuccess && !Output.Errors.Num())
+	{
+		Output.Errors.Add(TEXT("Compile failed without errors!"));
 	}
 
 	const bool bDirectCompile = FParse::Param(FCommandLine::Get(), TEXT("directcompile"));
diff --git a/Engine/Source/Developer/Windows/ShaderFormatD3D/Private/D3DShaderCompilerDXC.cpp b/Engine/Source/Developer/Windows/ShaderFormatD3D/Private/D3DShaderCompilerDXC.cpp
--- a/Engine/Source/Developer/Windows/ShaderFormatD3D/Private/D3DShaderCompilerDXC.cpp
+++ b/Engine/Source/Developer/Windows/ShaderFormatD3D/Private/D3DShaderCompilerDXC.cpp
@@ -595,7 +595,7 @@ bool CompileAndProcessD3DShaderDXC(FString& PreprocessedShaderSource,
 	const FShaderParameterParser& ShaderParameterParser,
 	FString& EntryPointName,
 	const TCHAR* ShaderProfile, ELanguage Language, bool bProcessingSecondTime,
-	TArray<FString>& FilteredErrors, FShaderCompilerOutput& Output)
+	FShaderCompilerOutput& Output)
 {
 	TRACE_CPUPROFILER_EVENT_SCOPE(CompileAndProcessD3DShaderDXC);
 
@@ -701,6 +701,7 @@ bool CompileAndProcessD3DShaderDXC(FString& PreprocessedShaderSource,
 
 	const HRESULT D3DCompileToDxilResult = D3DCompileToDxil(AnsiSourceFile.Get(), Args, ShaderBlob, ReflectionBlob, DxcErrorBlob);
 
+	TArray<FString> FilteredErrors;
 	if (DxcErrorBlob && DxcErrorBlob->GetBufferSize())
 	{
 		FString ErrorString = DxcBlobEncodingToFString(DxcErrorBlob);
@@ -1013,6 +1014,9 @@ bool CompileAndProcessD3DShaderDXC(FString& PreprocessedShaderSource,
 		FilteredErrors.Add(ErrorString);
 	}
 
+	// Move intermediate filtered errors into compiler context for unification.
+	CrossCompiler::FShaderConductorContext::ConvertCompileErrors(MoveTemp(FilteredErrors), Output.Errors);
+
 	return Output.bSucceeded;
 }
 
diff --git a/Engine/Source/Developer/Windows/ShaderFormatD3D/Private/ShaderFormatD3D.h b/Engine/Source/Developer/Windows/ShaderFormatD3D/Private/ShaderFormatD3D.h
--- a/Engine/Source/Developer/Windows/ShaderFormatD3D/Private/ShaderFormatD3D.h
+++ b/Engine/Source/Developer/Windows/ShaderFormatD3D/Private/ShaderFormatD3D.h
@@ -22,13 +22,13 @@ void CompileShader_Windows(const struct FShaderCompilerInput& Input,struct FShad
 /**
  * @param bSecondPassAferUnusedInputRemoval whether we're compiling the shader second time, after having removed the unused inputs discovered in the first pass
  */
-bool CompileAndProcessD3DShaderFXC(FString& PreprocessedShaderSource, const FString& CompilerPath,
+bool CompileAndProcessD3DShaderFXC(FString& PreprocessedShaderSource,
 	uint32 CompileFlags,
 	const FShaderCompilerInput& Input,
 	const FShaderParameterParser& ShaderParameterParser,
 	FString& EntryPointName,
 	const TCHAR* ShaderProfile, bool bSecondPassAferUnusedInputRemoval,
-	TArray<FString>& FilteredErrors, FShaderCompilerOutput& Output);
+	FShaderCompilerOutput& Output);
 
 bool CompileAndProcessD3DShaderDXC(FString& PreprocessedShaderSource,
 	uint32 CompileFlags,
@@ -36,7 +36,7 @@ bool CompileAndProcessD3DShaderDXC(FString& PreprocessedShaderSource,
 	const FShaderParameterParser& ShaderParameterParser,
 	FString& EntryPointName,
 	const TCHAR* ShaderProfile, ELanguage Language, bool bProcessingSecondTime,
-	TArray<FString>& FilteredErrors, FShaderCompilerOutput& Output);
+	FShaderCompilerOutput& Output);
 
 bool ValidateResourceCounts(uint32 NumSRVs, uint32 NumSamplers, uint32 NumUAVs, uint32 NumCBs, TArray<FString>& OutFilteredErrors);
 
