Use custom memory allocator for DXC to report better callstack on OOM event.

#rnx

diff --git a/Engine/Source/Developer/Windows/ShaderFormatD3D/Private/D3DShaderCompilerDXC.cpp b/Engine/Source/Developer/Windows/ShaderFormatD3D/Private/D3DShaderCompilerDXC.cpp
--- a/Engine/Source/Developer/Windows/ShaderFormatD3D/Private/D3DShaderCompilerDXC.cpp
+++ b/Engine/Source/Developer/Windows/ShaderFormatD3D/Private/D3DShaderCompilerDXC.cpp
@@ -110,6 +110,69 @@ static void LogFailedHRESULT(const TCHAR* FailedExpressionStr, HRESULT Result)
 		}													\
 	}
 
+class FDxcMalloc final : public IMalloc
+{
+	int32 RefCount = 1;
+
+public:
+
+	// IMalloc
+
+	virtual void* STDMETHODCALLTYPE Alloc(SIZE_T cb) final override
+	{
+		return FMemory::Malloc(cb);
+	}
+
+	virtual void* STDMETHODCALLTYPE Realloc(void* pv, SIZE_T cb) final override
+	{
+		return FMemory::Realloc(pv, cb);
+	}
+
+	virtual void STDMETHODCALLTYPE Free(void* pv) final override
+	{
+		return FMemory::Free(pv);
+	}
+
+	virtual SIZE_T STDMETHODCALLTYPE GetSize(void* pv) final override
+	{
+		return FMemory::GetAllocSize(pv);
+	}
+
+	virtual int STDMETHODCALLTYPE DidAlloc(void* pv) final override
+	{
+		return 1; // assume that all allocation queries coming from DXC belong to our allocator
+	}
+
+	virtual void STDMETHODCALLTYPE HeapMinimize() final override
+	{
+		// nothing
+	}
+
+	// IUnknown
+
+	ULONG STDMETHODCALLTYPE AddRef()
+	{
+		return ++RefCount;
+	}
+
+	ULONG STDMETHODCALLTYPE Release() {
+		check(RefCount != 0);
+		return --RefCount;
+	}
+
+	STDMETHODIMP QueryInterface(REFIID iid, void** ppvObject)
+	{
+		checkNoEntry(); // We do not expect or support QI on DXC allocator replacement
+		return ERROR_NOINTERFACE;
+	}
+};
+
+static IMalloc* GetDxcMalloc()
+{
+	static FDxcMalloc Instance;
+	return &Instance;
+}
+
 static dxc::DxcDllSupport& GetDxcDllHelper()
 {
 	static dxc::DxcDllSupport DxcDllSupport;
@@ -260,7 +323,7 @@ static void DumpFourCCParts(dxc::DxcDllSupport& DxcDllHelper, TRefCountPtr<IDxcB
 {
 #if UE_BUILD_DEBUG && IS_PROGRAM
 	TRefCountPtr<IDxcContainerReflection> Refl;
-	VERIFYHRESULT(DxcDllHelper.CreateInstance(CLSID_DxcContainerReflection, Refl.GetInitReference()));
+	VERIFYHRESULT(DxcDllHelper.CreateInstance2(GetDxcMalloc(), CLSID_DxcContainerReflection, Refl.GetInitReference()));
 
 	VERIFYHRESULT(Refl->Load(Blob));
 
@@ -286,7 +349,7 @@ static bool RemoveContainerReflection(dxc::DxcDllSupport& DxcDllHelper, TRefCoun
 	TRefCountPtr<IDxcContainerBuilder> Builder;
 	TRefCountPtr<IDxcBlob> StrippedDxil;
 
-	VERIFYHRESULT(DxcDllHelper.CreateInstance(CLSID_DxcContainerBuilder, Builder.GetInitReference()));
+	VERIFYHRESULT(DxcDllHelper.CreateInstance2(GetDxcMalloc(), CLSID_DxcContainerBuilder, Builder.GetInitReference()));
 	VERIFYHRESULT(Builder->Load(Dxil));
 	
 	// Try and remove both the PDB & Reflection Data
@@ -312,10 +375,10 @@ static HRESULT D3DCompileToDxil(const char* SourceText, FDxcArguments& Arguments
 	dxc::DxcDllSupport& DxcDllHelper = GetDxcDllHelper();
 
 	TRefCountPtr<IDxcCompiler3> Compiler;
-	VERIFYHRESULT(DxcDllHelper.CreateInstance(CLSID_DxcCompiler, Compiler.GetInitReference()));
+	VERIFYHRESULT(DxcDllHelper.CreateInstance2(GetDxcMalloc(), CLSID_DxcCompiler, Compiler.GetInitReference()));
 
 	TRefCountPtr<IDxcLibrary> Library;
-	VERIFYHRESULT(DxcDllHelper.CreateInstance(CLSID_DxcLibrary, Library.GetInitReference()));
+	VERIFYHRESULT(DxcDllHelper.CreateInstance2(GetDxcMalloc(), CLSID_DxcLibrary, Library.GetInitReference()));
 
 	TRefCountPtr<IDxcBlobEncoding> TextBlob;
 	VERIFYHRESULT(Library->CreateBlobWithEncodingFromPinned((LPBYTE)SourceText, FCStringAnsi::Strlen(SourceText), CP_UTF8, TextBlob.GetInitReference()));
@@ -731,7 +794,7 @@ bool CompileAndProcessD3DShaderDXC(FString& PreprocessedShaderSource,
 
 		dxc::DxcDllSupport& DxcDllHelper = GetDxcDllHelper();
 		TRefCountPtr<IDxcUtils> Utils;
-		VERIFYHRESULT(DxcDllHelper.CreateInstance(CLSID_DxcUtils, Utils.GetInitReference()));
+		VERIFYHRESULT(DxcDllHelper.CreateInstance2(GetDxcMalloc(), CLSID_DxcUtils, Utils.GetInitReference()));
 		DxcBuffer ReflBuffer = { 0 };
 		ReflBuffer.Ptr = ReflectionBlob->GetBufferPointer();
 		ReflBuffer.Size = ReflectionBlob->GetBufferSize();
