Allow allocation of zero bytes in DXC custom allocator, use portable STDCALL declaration, and add error code translation for DXC specific HRESULT values.

#rnx

diff --git a/Engine/Source/Developer/Windows/ShaderFormatD3D/Private/D3DShaderCompilerDXC.cpp b/Engine/Source/Developer/Windows/ShaderFormatD3D/Private/D3DShaderCompilerDXC.cpp
--- a/Engine/Source/Developer/Windows/ShaderFormatD3D/Private/D3DShaderCompilerDXC.cpp
+++ b/Engine/Source/Developer/Windows/ShaderFormatD3D/Private/D3DShaderCompilerDXC.cpp
@@ -35,6 +35,7 @@ MSVC_PRAGMA(warning(push))
 MSVC_PRAGMA(warning(disable : 4191)) // warning C4191: 'type cast': unsafe conversion from 'FARPROC' to 'DxcCreateInstanceProc'
 #include <dxc/dxcapi.h>
 #include <dxc/Support/dxcapi.use.h>
+#include <dxc/Support/ErrorCodes.h>
 #include <d3d12shader.h>
 MSVC_PRAGMA(warning(pop))
 
@@ -79,6 +80,47 @@ static uint32 GetAutoBindingSpace(const FShaderTarget& Target)
 	}
 }
 
+// DXC specific error codes cannot be translated by FPlatformMisc::GetSystemErrorMessage, so do it manually.
+// Codes defines in <DXC>/include/dxc/Support/ErrorCodes.h
+static const TCHAR* DxcErrorCodeToString(HRESULT Code)
+{
+#define SWITCHCASE_TO_STRING(VALUE) case VALUE: return TEXT(#VALUE)
+	switch (Code)
+	{
+		SWITCHCASE_TO_STRING( DXC_E_OVERLAPPING_SEMANTICS );
+		SWITCHCASE_TO_STRING( DXC_E_MULTIPLE_DEPTH_SEMANTICS );
+		SWITCHCASE_TO_STRING( DXC_E_INPUT_FILE_TOO_LARGE );
+		SWITCHCASE_TO_STRING( DXC_E_INCORRECT_DXBC );
+		SWITCHCASE_TO_STRING( DXC_E_ERROR_PARSING_DXBC_BYTECODE );
+		SWITCHCASE_TO_STRING( DXC_E_DATA_TOO_LARGE );
+		SWITCHCASE_TO_STRING( DXC_E_INCOMPATIBLE_CONVERTER_OPTIONS);
+		SWITCHCASE_TO_STRING( DXC_E_IRREDUCIBLE_CFG );
+		SWITCHCASE_TO_STRING( DXC_E_IR_VERIFICATION_FAILED );
+		SWITCHCASE_TO_STRING( DXC_E_SCOPE_NESTED_FAILED );
+		SWITCHCASE_TO_STRING( DXC_E_NOT_SUPPORTED );
+		SWITCHCASE_TO_STRING( DXC_E_STRING_ENCODING_FAILED );
+		SWITCHCASE_TO_STRING( DXC_E_CONTAINER_INVALID );
+		SWITCHCASE_TO_STRING( DXC_E_CONTAINER_MISSING_DXIL );
+		SWITCHCASE_TO_STRING( DXC_E_INCORRECT_DXIL_METADATA );
+		SWITCHCASE_TO_STRING( DXC_E_INCORRECT_DDI_SIGNATURE );
+		SWITCHCASE_TO_STRING( DXC_E_DUPLICATE_PART );
+		SWITCHCASE_TO_STRING( DXC_E_MISSING_PART );
+		SWITCHCASE_TO_STRING( DXC_E_MALFORMED_CONTAINER );
+		SWITCHCASE_TO_STRING( DXC_E_INCORRECT_ROOT_SIGNATURE );
+		SWITCHCASE_TO_STRING( DXC_E_CONTAINER_MISSING_DEBUG );
+		SWITCHCASE_TO_STRING( DXC_E_MACRO_EXPANSION_FAILURE );
+		SWITCHCASE_TO_STRING( DXC_E_OPTIMIZATION_FAILED );
+		SWITCHCASE_TO_STRING( DXC_E_GENERAL_INTERNAL_ERROR );
+		SWITCHCASE_TO_STRING( DXC_E_ABORT_COMPILATION_ERROR );
+		SWITCHCASE_TO_STRING( DXC_E_EXTENSION_ERROR );
+		SWITCHCASE_TO_STRING( DXC_E_LLVM_FATAL_ERROR );
+		SWITCHCASE_TO_STRING( DXC_E_LLVM_UNREACHABLE );
+		SWITCHCASE_TO_STRING( DXC_E_LLVM_CAST_ERROR );
+	}
+	return nullptr;
+#undef SWITCHCASE_TO_STRING
+}
+
 // Utility variable so we can place a breakpoint while debugging
 static int32 GBreakpointDXC = 0;
 
@@ -92,6 +134,10 @@ static void LogFailedHRESULT(const TCHAR* FailedExpressionStr, HRESULT Result)
 		FString UsedMemoryStr = FString::Printf(TEXT("%d MB"), MemoryStats.UsedPhysical / 1024 / 1024);
 		UE_LOG(LogD3D12ShaderCompiler, Fatal, TEXT("%s failed: Result=0x%08x (E_OUTOFMEMORY); System memory %s, used %s"), FailedExpressionStr, Result, *SystemMemoryStr, *UsedMemoryStr);
 	}
+	else if (const TCHAR* ErrorCodeStr = DxcErrorCodeToString(Result))
+	{
+		UE_LOG(LogD3D12ShaderCompiler, Fatal, TEXT("%s failed: Result=0x%08x (%s)"), FailedExpressionStr, Result, ErrorCodeStr);
+	}
 	else
 	{
 		// Turn HRESULT into human readable string for error report
@@ -103,7 +149,7 @@ static void LogFailedHRESULT(const TCHAR* FailedExpressionStr, HRESULT Result)
 
 #define VERIFYHRESULT(expr)									\
 	{														\
-		HRESULT HR##__LINE__ = expr;						\
+		const HRESULT HR##__LINE__ = expr;					\
 		if (FAILED(HR##__LINE__))							\
 		{													\
 			LogFailedHRESULT(TEXT(#expr), HR##__LINE__);	\
@@ -112,55 +158,57 @@ static void LogFailedHRESULT(const TCHAR* FailedExpressionStr, HRESULT Result)
 
 class FDxcMalloc final : public IMalloc
 {
-	int32 RefCount = 1;
+	ULONG RefCount = 1;
 
 public:
 
 	// IMalloc
 
-	virtual void* STDMETHODCALLTYPE Alloc(SIZE_T cb) final override
+	void* STDCALL Alloc(SIZE_T cb) override
 	{
+		cb = FMath::Max(SIZE_T(1), cb);
 		return FMemory::Malloc(cb);
 	}
 
-	virtual void* STDMETHODCALLTYPE Realloc(void* pv, SIZE_T cb) final override
+	void* STDCALL Realloc(void* pv, SIZE_T cb) override
 	{
+		cb = FMath::Max(SIZE_T(1), cb);
 		return FMemory::Realloc(pv, cb);
 	}
 
-	virtual void STDMETHODCALLTYPE Free(void* pv) final override
+	void STDCALL Free(void* pv) override
 	{
 		return FMemory::Free(pv);
 	}
 
-	virtual SIZE_T STDMETHODCALLTYPE GetSize(void* pv) final override
+	SIZE_T STDCALL GetSize(void* pv) override
 	{
 		return FMemory::GetAllocSize(pv);
 	}
 
-	virtual int STDMETHODCALLTYPE DidAlloc(void* pv) final override
+	int STDCALL DidAlloc(void* pv) override
 	{
 		return 1; // assume that all allocation queries coming from DXC belong to our allocator
 	}
 
-	virtual void STDMETHODCALLTYPE HeapMinimize() final override
+	void STDCALL HeapMinimize() override
 	{
 		// nothing
 	}
 
 	// IUnknown
 
-	ULONG STDMETHODCALLTYPE AddRef()
+	ULONG STDCALL AddRef()
 	{
 		return ++RefCount;
 	}
 
-	ULONG STDMETHODCALLTYPE Release() {
+	ULONG STDCALL Release() {
 		check(RefCount != 0);
 		return --RefCount;
 	}
 
-	STDMETHODIMP QueryInterface(REFIID iid, void** ppvObject)
+	HRESULT STDCALL QueryInterface(REFIID iid, void** ppvObject)
 	{
 		checkNoEntry(); // We do not expect or support QI on DXC allocator replacement
 		return ERROR_NOINTERFACE;
@@ -802,12 +850,7 @@ bool CompileAndProcessD3DShaderDXC(FString& PreprocessedShaderSource,
 		if (bIsRayTracingShader)
 		{
 			TRefCountPtr<ID3D12LibraryReflection> LibraryReflection;
-			const HRESULT CreateReflectionResult = Utils->CreateReflection(&ReflBuffer, IID_PPV_ARGS(LibraryReflection.GetInitReference()));
-
-			if (FAILED(CreateReflectionResult))
-			{
-				UE_LOG(LogD3D12ShaderCompiler, Fatal, TEXT("CreateReflection failed: Result=%08x"), CreateReflectionResult);
-			}
+			VERIFYHRESULT(Utils->CreateReflection(&ReflBuffer, IID_PPV_ARGS(LibraryReflection.GetInitReference())));
 
 			D3D12_LIBRARY_DESC LibraryDesc = {};
 			LibraryReflection->GetDesc(&LibraryDesc);
@@ -922,11 +965,7 @@ bool CompileAndProcessD3DShaderDXC(FString& PreprocessedShaderSource,
 		else
 		{
 			TRefCountPtr<ID3D12ShaderReflection> ShaderReflection;
-			const HRESULT CreateReflectionResult = Utils->CreateReflection(&ReflBuffer, IID_PPV_ARGS(ShaderReflection.GetInitReference()));
-			if (FAILED(CreateReflectionResult))
-			{
-				UE_LOG(LogD3D12ShaderCompiler, Fatal, TEXT("CreateReflection failed: Result=%08x"), CreateReflectionResult);
-			}
+			VERIFYHRESULT(Utils->CreateReflection(&ReflBuffer, IID_PPV_ARGS(ShaderReflection.GetInitReference())));
 
 			D3D12_SHADER_DESC ShaderDesc = {};
 			ShaderReflection->GetDesc(&ShaderDesc);
