Fixed compile issue

#rnx

diff --git a/Engine/Source/Developer/Windows/LiveCodingServer/Private/External/LC_CoffDetail.h b/Engine/Source/Developer/Windows/LiveCodingServer/Private/External/LC_CoffDetail.h
--- a/Engine/Source/Developer/Windows/LiveCodingServer/Private/External/LC_CoffDetail.h
+++ b/Engine/Source/Developer/Windows/LiveCodingServer/Private/External/LC_CoffDetail.h
@@ -69,11 +69,9 @@ namespace coffDetail
 
 
 	template <typename AuxSymbolType>
-	inline uint32_t GetAssociatedComdatSectionIndex(const AuxSymbolType* symbol)
-	{
-		static_assert(false, "Invalid function overload");
-		return 0u;
-	}
+	// BEGIN EPIC MOD
+	uint32_t GetAssociatedComdatSectionIndex(const AuxSymbolType* symbol);
+	// END EPIC MOD
 
 	template <>
 	inline uint32_t GetAssociatedComdatSectionIndex(const IMAGE_AUX_SYMBOL* symbol)
