After extracting all parameters from D3D12 reflection, the "bSucceeded" flag was being forced to true which would prevent any parameter validation from being reported. Most notably, this would hide errors caused by missing global parameters when using shader parameter structs.
Bumping D3D_SM5 and D3D_SM6 versions.

diff --git a/Engine/Source/Developer/Windows/ShaderFormatD3D/Private/ShaderFormatD3D.cpp b/Engine/Source/Developer/Windows/ShaderFormatD3D/Private/ShaderFormatD3D.cpp
--- a/Engine/Source/Developer/Windows/ShaderFormatD3D/Private/ShaderFormatD3D.cpp
+++ b/Engine/Source/Developer/Windows/ShaderFormatD3D/Private/ShaderFormatD3D.cpp
@@ -20,8 +20,8 @@ class FShaderFormatD3D : public IShaderFormat
 		UE_SHADER_PCD3D_SHARED_VER = 5,
 
 		/** Version for shader format, this becomes part of the DDC key. */
-		UE_SHADER_PCD3D_SM6_VER = 6,
-		UE_SHADER_PCD3D_SM5_VER = 11,
+		UE_SHADER_PCD3D_SM6_VER = 7,
+		UE_SHADER_PCD3D_SM5_VER = 12,
 		UE_SHADER_PCD3D_ES3_1_VER = 8,
 		UE_SHADER_D3D_ES3_1_HOLOLENS_VER = UE_SHADER_PCD3D_ES3_1_VER,
 	};
