Adding correct checks for SV_StencilRef support
- SM6.0+ shader reflection now checks for D3D_SHADER_REQUIRES_STENCIL_REF and sets EShaderCodeFeatures::StencilRef if it's set.
- D3D12 initialization now correctly sets GRHISupportsStencilRefFromPixelShader.
- D3D12 shader loading makes sure that EShaderCodeFeatures::StencilRef shaders are only loaded when GRHISupportsStencilRefFromPixelShader is set.

diff --git a/Engine/Source/Developer/Windows/ShaderFormatD3D/Private/D3DShaderCompilerDXC.cpp b/Engine/Source/Developer/Windows/ShaderFormatD3D/Private/D3DShaderCompilerDXC.cpp
--- a/Engine/Source/Developer/Windows/ShaderFormatD3D/Private/D3DShaderCompilerDXC.cpp
+++ b/Engine/Source/Developer/Windows/ShaderFormatD3D/Private/D3DShaderCompilerDXC.cpp
@@ -1082,6 +1082,11 @@ bool CompileAndProcessD3DShaderDXC(FString& PreprocessedShaderSource,
 					EnumAddFlags(CodeFeatures.CodeFeatures, EShaderCodeFeatures::BindlessSamplers);
 				}
 
+				if ((ShaderRequiresFlags & D3D_SHADER_REQUIRES_STENCIL_REF) != 0)
+				{
+					EnumAddFlags(CodeFeatures.CodeFeatures, EShaderCodeFeatures::StencilRef);
+				}
+
 				// We only need this to appear when using a DXC shader
 				ShaderCode.AddOptionalData<FShaderCodeFeatures>(CodeFeatures);
 
diff --git a/Engine/Source/Developer/Windows/ShaderFormatD3D/Private/ShaderFormatD3D.cpp b/Engine/Source/Developer/Windows/ShaderFormatD3D/Private/ShaderFormatD3D.cpp
--- a/Engine/Source/Developer/Windows/ShaderFormatD3D/Private/ShaderFormatD3D.cpp
+++ b/Engine/Source/Developer/Windows/ShaderFormatD3D/Private/ShaderFormatD3D.cpp
@@ -58,13 +58,21 @@ public:
 			ShaderModelHash ^= 0x96ED7F56;
 		#endif
 
+			// Make sure we recompile if EShaderCodeFeatures gets bigger
+			ShaderModelHash = HashCombine(ShaderModelHash, GetTypeHash(sizeof(EShaderCodeFeatures)));
+
 			return HashCombine(DxcVersionHash, ShaderModelHash);
 		}
 		else if (Format == NAME_PCD3D_SM5)
 		{
+			uint32 ShaderModelHash = GetVersionHash(UE_SHADER_PCD3D_SM6_VER);
+
+			// Make sure we recompile if EShaderCodeFeatures gets bigger
+			ShaderModelHash = HashCombine(ShaderModelHash, GetTypeHash(sizeof(EShaderCodeFeatures)));
+
 			// Technically not needed for regular SM5 compiled with legacy compiler,
 			// but PCD3D_SM5 currently includes ray tracing shaders that are compiled with new compiler stack.
-			return HashCombine(DxcVersionHash, GetVersionHash(UE_SHADER_PCD3D_SM5_VER));
+			return HashCombine(DxcVersionHash, ShaderModelHash);
 		}
 		else if (Format == NAME_PCD3D_ES3_1) 
 		{
diff --git a/Engine/Source/Runtime/D3D12RHI/Private/D3D12Adapter.cpp b/Engine/Source/Runtime/D3D12RHI/Private/D3D12Adapter.cpp
--- a/Engine/Source/Runtime/D3D12RHI/Private/D3D12Adapter.cpp
+++ b/Engine/Source/Runtime/D3D12RHI/Private/D3D12Adapter.cpp
@@ -953,6 +953,15 @@ PRAGMA_ENABLE_DEPRECATION_WARNINGS
 				UE_LOG(LogD3D12RHI, Log, TEXT("Bindless resources are supported"));
 			}
 
+			{
+				D3D12_FEATURE_DATA_D3D12_OPTIONS Features{};
+				RootDevice->CheckFeatureSupport(D3D12_FEATURE_D3D12_OPTIONS, &Features, sizeof(Features));
+
+				GRHISupportsStencilRefFromPixelShader = (Features.PSSpecifiedStencilRefSupported != 0);
+
+				UE_LOG(LogD3D12RHI, Log, TEXT("Stencil ref from pixel shader is %s"), GRHISupportsStencilRefFromPixelShader ? TEXT("supported") : TEXT("not supported"));
+			}
+
 			// Detect availability of shader model 6.0 wave operations
 			{
 				D3D12_FEATURE_DATA_D3D12_OPTIONS1 Features{};
diff --git a/Engine/Source/Runtime/D3D12RHI/Private/D3D12Shaders.cpp b/Engine/Source/Runtime/D3D12RHI/Private/D3D12Shaders.cpp
--- a/Engine/Source/Runtime/D3D12RHI/Private/D3D12Shaders.cpp
+++ b/Engine/Source/Runtime/D3D12RHI/Private/D3D12Shaders.cpp
@@ -76,6 +76,11 @@ static bool ValidateShaderIsUsable(FD3D12ShaderData* InShader, EShaderFrequency
 			return false;
 		}
 	}
+
+	if (InFrequency == SF_Pixel && EnumHasAnyFlags(InShader->Features, EShaderCodeFeatures::StencilRef) && !GRHISupportsStencilRefFromPixelShader)
+	{
+		return false;
+	}
 #endif
 
 	return true;
diff --git a/Engine/Source/Runtime/RenderCore/Public/ShaderCore.h b/Engine/Source/Runtime/RenderCore/Public/ShaderCore.h
--- a/Engine/Source/Runtime/RenderCore/Public/ShaderCore.h
+++ b/Engine/Source/Runtime/RenderCore/Public/ShaderCore.h
@@ -664,7 +664,7 @@ struct FShaderCodeResourceMasks
 };
 
 // if this changes you need to make sure all shaders get invalidated
-enum class EShaderCodeFeatures : uint8
+enum class EShaderCodeFeatures : uint16
 {
 	None                    = 0,
 	WaveOps                 = 1 << 0,
@@ -674,6 +674,7 @@ enum class EShaderCodeFeatures : uint8
 	DiagnosticBuffer        = 1 << 4,
 	BindlessResources       = 1 << 5,
 	BindlessSamplers        = 1 << 6,
+	StencilRef              = 1 << 7,
 };
 ENUM_CLASS_FLAGS(EShaderCodeFeatures);
 
