Fix temporary string object in D3DShaderCompiler. Related CL 20260406.

#rnx

diff --git a/Engine/Source/Developer/Windows/ShaderFormatD3D/Private/D3DShaderCompiler.cpp b/Engine/Source/Developer/Windows/ShaderFormatD3D/Private/D3DShaderCompiler.cpp
--- a/Engine/Source/Developer/Windows/ShaderFormatD3D/Private/D3DShaderCompiler.cpp
+++ b/Engine/Source/Developer/Windows/ShaderFormatD3D/Private/D3DShaderCompiler.cpp
@@ -563,11 +563,12 @@ static bool CompileErrorsContainInternalError(ID3DBlob* Errors)
 {
 	if (Errors)
 	{
-		void* ErrorBuffer = Errors->GetBufferPointer();
-		if (ErrorBuffer)
+		if (void* ErrorBuffer = Errors->GetBufferPointer())
 		{
-			const FStringView ErrorString = ANSI_TO_TCHAR(ErrorBuffer);
-			return ErrorString.Contains(TEXT("internal error:")) || ErrorString.Contains(TEXT("Internal Compiler Error:"));
+			const ANSICHAR* ErrorString = reinterpret_cast<const ANSICHAR*>(ErrorBuffer);
+			return
+				FCStringAnsi::Strstr(ErrorString, "internal error:") != nullptr ||
+				FCStringAnsi::Strstr(ErrorString, "Internal Compiler Error:") != nullptr;
 		}
 	}
 	return false;
