Fixed LiveCoding issue where if a typo is made in the preloaded module list, the game would hit a breakpoint.

diff --git a/Engine/Source/Developer/Windows/LiveCoding/Private/LiveCodingModule.cpp b/Engine/Source/Developer/Windows/LiveCoding/Private/LiveCodingModule.cpp
--- a/Engine/Source/Developer/Windows/LiveCoding/Private/LiveCodingModule.cpp
+++ b/Engine/Source/Developer/Windows/LiveCoding/Private/LiveCodingModule.cpp
@@ -1037,12 +1037,17 @@ void FLiveCodingModule::UpdateModules(bool bAllowStarting)
 		TSet<FName> PreloadedFileNames;
 		{
 			FModuleManager& Manager = FModuleManager::Get();
-			for (FName moduleName : Settings->PreloadNamedModules)
+			for (FName ModuleName : Settings->PreloadNamedModules)
 			{
-				FString FileName = Manager.GetModuleFilename(moduleName);
-				if (!FileName.IsEmpty())
+				TArray<FName> OutModules;
+				Manager.FindModules(*ModuleName.ToString(), OutModules);
+				for (FName OutModule : OutModules)
 				{
-					PreloadedFileNames.Add(FName(FPaths::GetBaseFilename(FileName, true)));
+					FString FileName = Manager.GetModuleFilename(OutModule);
+					if (!FileName.IsEmpty())
+					{
+						PreloadedFileNames.Add(FName(FPaths::GetBaseFilename(FileName, true)));
+					}
 				}
 			}
 		}
