Fix loop unrolling issue in DOFGatherKernel shader for updated DXC.

#rnx

diff --git a/Engine/Source/Developer/Windows/ShaderFormatD3D/Private/D3DShaderCompiler.inl b/Engine/Source/Developer/Windows/ShaderFormatD3D/Private/D3DShaderCompiler.inl
--- a/Engine/Source/Developer/Windows/ShaderFormatD3D/Private/D3DShaderCompiler.inl
+++ b/Engine/Source/Developer/Windows/ShaderFormatD3D/Private/D3DShaderCompiler.inl
@@ -195,6 +195,9 @@ public:
 
 		// disable undesired warnings
 		ExtraArguments.Add(TEXT("-Wno-parentheses-equality"));
+
+		// @lh-todo: This fixes a loop unrolling issue that showed up in DOFGatherKernel with cs_6_6 with the latest DXC revision
+		ExtraArguments.Add(TEXT("-disable-lifetime-markers"));
 	}
 
 	inline FString GetDumpDebugInfoPath() const
