[shaders] remove unneccessary global vars from D3D compile backend

diff --git a/Engine/Source/Developer/Windows/ShaderFormatD3D/Private/D3DShaderCompiler.cpp b/Engine/Source/Developer/Windows/ShaderFormatD3D/Private/D3DShaderCompiler.cpp
--- a/Engine/Source/Developer/Windows/ShaderFormatD3D/Private/D3DShaderCompiler.cpp
+++ b/Engine/Source/Developer/Windows/ShaderFormatD3D/Private/D3DShaderCompiler.cpp
@@ -37,7 +37,6 @@ DEFINE_LOG_CATEGORY_STATIC(LogD3D11ShaderCompiler, Log, All);
 
 static const uint32 GD3DMaximumNumUAVs = 8; // Limit for feature level 11.0
 
-int32 GD3DAllowRemoveUnused = 0;
 static int32 GD3DCheckForDoubles = 1;
 static int32 GD3DDumpAMDCodeXLFile = 0;
 
@@ -489,9 +488,6 @@ static HRESULT D3DCompileWrapper(
 #endif
 }
 
-// Utility variable so we can place a breakpoint while debugging
-static int32 GBreakpoint = 0;
-
 inline bool IsCompatibleBinding(const D3D11_SHADER_INPUT_BIND_DESC& BindDesc, uint32 BindingSpace)
 {
 	return true;
@@ -821,12 +817,6 @@ static bool CompileAndProcessD3DShaderFXCExt(
 
 			if (Input.Target.Frequency == SF_Pixel)
 			{
-				if (GD3DAllowRemoveUnused != 0 && Input.bCompilingForShaderPipeline)
-				{
-					// Handy place for a breakpoint for debugging...
-					++GBreakpoint;
-				}
-
 				bool bFoundUnused = false;
 				for (uint32 Index = 0; Index < ShaderDesc.InputParameters; ++Index)
 				{
@@ -866,7 +856,7 @@ static bool CompileAndProcessD3DShaderFXCExt(
 					}
 				}
 
-				if (GD3DAllowRemoveUnused && Input.bCompilingForShaderPipeline && bFoundUnused && !bSecondPassAferUnusedInputRemoval)
+				if (Input.Environment.CompilerFlags.Contains(CFLAG_ForceRemoveUnusedInterpolators) && Input.bCompilingForShaderPipeline && bFoundUnused && !bSecondPassAferUnusedInputRemoval)
 				{
 					// Rewrite the source removing the unused inputs so the bindings will match.
 					// We may need to do this more than once if unused inputs change after the removal. Ie. for complex shaders, what can happen is:
@@ -1059,11 +1049,6 @@ static bool CompileAndProcessD3DShaderFXCExt(
 		}
 	}
 
-	if (FAILED(Result))
-	{
-		++GBreakpoint;
-	}
-
 	return SUCCEEDED(Result);
 }
 
@@ -1174,12 +1159,10 @@ void CompileD3DShader(const FShaderCompilerInput& Input, FShaderCompilerOutput&
 		}
 	}
 
-	GD3DAllowRemoveUnused = Input.Environment.CompilerFlags.Contains(CFLAG_ForceRemoveUnusedInterpolators) ? 1 : 0;
-
 	FString EntryPointName = Input.EntryPointName;
 
 	Output.bFailedRemovingUnused = false;
-	if (GD3DAllowRemoveUnused == 1 && Input.Target.Frequency == SF_Vertex && Input.bCompilingForShaderPipeline)
+	if (Input.Environment.CompilerFlags.Contains(CFLAG_ForceRemoveUnusedInterpolators) && Input.Target.Frequency == SF_Vertex && Input.bCompilingForShaderPipeline)
 	{
 		// Always add SV_Position
 		TArray<FString> UsedOutputs = Input.UsedOutputs;
diff --git a/Engine/Source/Developer/Windows/ShaderFormatD3D/Private/D3DShaderCompiler.inl b/Engine/Source/Developer/Windows/ShaderFormatD3D/Private/D3DShaderCompiler.inl
--- a/Engine/Source/Developer/Windows/ShaderFormatD3D/Private/D3DShaderCompiler.inl
+++ b/Engine/Source/Developer/Windows/ShaderFormatD3D/Private/D3DShaderCompiler.inl
@@ -576,13 +576,11 @@ inline void GenerateFinalOutput(TRefCountPtr<TBlob>& CompressedData,
 		BuildResourceTableTokenStream(GenericSRT.UnorderedAccessViewMap, GenericSRT.MaxBoundResourceTable, SRT.UnorderedAccessViewMap);
 	}
 
-	if (GD3DAllowRemoveUnused != 0 && Input.Target.Frequency == SF_Pixel && Input.bCompilingForShaderPipeline && bProcessingSecondTime)
+	if (Input.Environment.CompilerFlags.Contains(CFLAG_ForceRemoveUnusedInterpolators) && Input.Target.Frequency == SF_Pixel && Input.bCompilingForShaderPipeline && bProcessingSecondTime)
+
 	{
 		Output.bSupportsQueryingUsedAttributes = true;
-		if (GD3DAllowRemoveUnused == 1)
-		{
-			Output.UsedAttributes = ShaderInputs;
-		}
+		Output.UsedAttributes = ShaderInputs;
 	}
 
 	// Generate the final Output
diff --git a/Engine/Source/Developer/Windows/ShaderFormatD3D/Private/ShaderFormatD3D.h b/Engine/Source/Developer/Windows/ShaderFormatD3D/Private/ShaderFormatD3D.h
--- a/Engine/Source/Developer/Windows/ShaderFormatD3D/Private/ShaderFormatD3D.h
+++ b/Engine/Source/Developer/Windows/ShaderFormatD3D/Private/ShaderFormatD3D.h
@@ -44,5 +44,3 @@ bool CompileAndProcessD3DShaderDXC(FString& PreprocessedShaderSource,
 	FShaderCompilerOutput& Output);
 
 bool ValidateResourceCounts(uint32 NumSRVs, uint32 NumSamplers, uint32 NumUAVs, uint32 NumCBs, TArray<FString>& OutFilteredErrors);
-
-extern int32 GD3DAllowRemoveUnused;
