Added a new LiveCoding command that allows for a single command to add both types of modules and a list of reserved pages.

#rnx

diff --git a/Engine/Source/Developer/Windows/LiveCoding/Private/External/LC_API.cpp b/Engine/Source/Developer/Windows/LiveCoding/Private/External/LC_API.cpp
--- a/Engine/Source/Developer/Windows/LiveCoding/Private/External/LC_API.cpp
+++ b/Engine/Source/Developer/Windows/LiveCoding/Private/External/LC_API.cpp
@@ -359,6 +359,19 @@ LPP_DLL_API(void) LppDisableCompileFinishNotification()
 }
 // END EPIC MOD
 
+// BEGIN EPIC MOD
+LPP_DLL_API(void*) LppEnableModulesEx(const wchar_t* moduleNames[], unsigned int moduleCount, const wchar_t* lazyLoadModuleNames[], unsigned int lazyLoadModuleCount, const uintptr_t* reservedPages, unsigned int reservedPagesCount)
+{
+	if (!CheckForStartup())
+	{
+		return nullptr;
+	}
+
+	return g_startupThread->EnableModulesEx(moduleNames, moduleCount, lazyLoadModuleNames, lazyLoadModuleCount, reservedPages, reservedPagesCount);
+}
+// END EPIC MOD
+
+
 LPP_DLL_API(void) LppApplySettingBool(const char* settingName, int value)
 {
 	if (!CheckForStartup())
diff --git a/Engine/Source/Developer/Windows/LiveCoding/Private/External/LC_API.h b/Engine/Source/Developer/Windows/LiveCoding/Private/External/LC_API.h
--- a/Engine/Source/Developer/Windows/LiveCoding/Private/External/LC_API.h
+++ b/Engine/Source/Developer/Windows/LiveCoding/Private/External/LC_API.h
@@ -43,6 +43,7 @@ LPP_DLL_API(void) LppSetBuildArguments(const wchar_t* arguments);
 LPP_DLL_API(void*) LppEnableLazyLoadedModule(const wchar_t* nameOfExeOrDll);
 LPP_DLL_API(void) LppSetReinstancingFlow(bool enable);
 LPP_DLL_API(void) LppDisableCompileFinishNotification();
+LPP_DLL_API(void*) LppEnableModulesEx(const wchar_t* moduleNames[], unsigned int moduleCount, const wchar_t* lazyLoadModuleNames[], unsigned int lazyLoadModuleCount, const uintptr_t* reservedPages, unsigned int reservedPagesCount);
 // END EPIC MOD
 LPP_DLL_API(void) LppApplySettingBool(const char* settingName, int value);
 LPP_DLL_API(void) LppApplySettingInt(const char* settingName, int value);
diff --git a/Engine/Source/Developer/Windows/LiveCoding/Private/External/LC_ClientStartupThread.cpp b/Engine/Source/Developer/Windows/LiveCoding/Private/External/LC_ClientStartupThread.cpp
--- a/Engine/Source/Developer/Windows/LiveCoding/Private/External/LC_ClientStartupThread.cpp
+++ b/Engine/Source/Developer/Windows/LiveCoding/Private/External/LC_ClientStartupThread.cpp
@@ -410,6 +410,22 @@ void ClientStartupThread::DisableCompileFinishNotification()
 }
 // END EPIC MOD
 
+// BEGIN EPIC MOD
+void* ClientStartupThread::EnableModulesEx(const wchar_t* moduleNames[], unsigned int moduleCount, const wchar_t* lazyLoadModuleNames[], unsigned int lazyLoadModuleCount, const uintptr_t* reservedPages, unsigned int reservedPagesCount)
+{
+	// we cannot wait for commands in the user command thread as long as startup hasn't finished
+	Join();
+
+	if (m_userCommandThread)
+	{
+		return m_userCommandThread->EnableModulesEx(moduleNames, moduleCount, lazyLoadModuleNames, lazyLoadModuleCount, reservedPages, reservedPagesCount);
+	}
+
+	return nullptr;
+}
+// END EPIC MOD
+
+
 void ClientStartupThread::ApplySettingBool(const char* settingName, int value)
 {
 	// wait for the startup thread to finish initialization
diff --git a/Engine/Source/Developer/Windows/LiveCoding/Private/External/LC_ClientStartupThread.h b/Engine/Source/Developer/Windows/LiveCoding/Private/External/LC_ClientStartupThread.h
--- a/Engine/Source/Developer/Windows/LiveCoding/Private/External/LC_ClientStartupThread.h
+++ b/Engine/Source/Developer/Windows/LiveCoding/Private/External/LC_ClientStartupThread.h
@@ -85,6 +85,10 @@ public:
 	void DisableCompileFinishNotification();
 	// END EPIC MOD
 
+	// BEGIN EPIC MOD
+	void* EnableModulesEx(const wchar_t* moduleNames[], unsigned int moduleCount, const wchar_t* lazyLoadModuleNames[], unsigned int lazyLoadModuleCount, const uintptr_t* reservedPages, unsigned int reservedPagesCount);
+	// END EPIC MOD
+
 private:
 	Thread::ReturnValue ThreadFunction(const std::wstring& groupName, RunMode::Enum runMode);
 
diff --git a/Engine/Source/Developer/Windows/LiveCoding/Private/External/LC_ClientUserCommandThread.cpp b/Engine/Source/Developer/Windows/LiveCoding/Private/External/LC_ClientUserCommandThread.cpp
--- a/Engine/Source/Developer/Windows/LiveCoding/Private/External/LC_ClientUserCommandThread.cpp
+++ b/Engine/Source/Developer/Windows/LiveCoding/Private/External/LC_ClientUserCommandThread.cpp
@@ -575,6 +575,74 @@ void ClientUserCommandThread::DisableCompileFinishNotification()
 }
 // END EPIC MOD
 
+// BEGIN EPIC MOD
+namespace
+{
+	std::vector<Windows::HMODULE> GatherModuleHandles(const wchar_t* moduleNames[], unsigned int moduleCount)
+	{
+		std::vector<Windows::HMODULE> moduleHandles;
+		moduleHandles.reserve(moduleCount);
+
+		for (unsigned int i = 0u; i < moduleCount; ++i)
+		{
+			Windows::HMODULE module = ::GetModuleHandleW(moduleNames[i]);
+			if (module)
+			{
+				moduleHandles.push_back(module);
+			}
+			else
+			{
+				LC_ERROR_USER("Cannot enable module %S because it is not loaded by this process.", moduleNames[i]);
+			}
+		}
+		return moduleHandles;
+	}
+
+	void AppendModuleHandles(ProxyCommand<commands::EnableModulesEx>* proxy, const std::vector<Windows::HMODULE>& moduleHandles)
+	{
+		for (Windows::HMODULE moduleHandle : moduleHandles)
+		{
+			commands::ModuleData moduleData = {};
+			moduleData.base = moduleHandle;
+			::GetModuleFileNameW(moduleHandle, moduleData.path, MAX_PATH);
+			proxy->m_payload.Write(moduleData);
+		}
+	}
+}
+
+void* ClientUserCommandThread::EnableModulesEx(const wchar_t* moduleNames[], unsigned int moduleCount, const wchar_t* lazyLoadModuleNames[], unsigned int lazyLoadModuleCount, const uintptr_t* reservedPages, unsigned int reservedPagesCount)
+{
+	std::vector<Windows::HMODULE> moduleHandles = GatherModuleHandles(moduleNames, moduleCount);
+	std::vector<Windows::HMODULE> lazyLoadModuleHandles = GatherModuleHandles(lazyLoadModuleNames, lazyLoadModuleCount);
+	if (moduleHandles.size() == 0 && lazyLoadModuleHandles.size() == 0)
+	{
+		return nullptr;
+	}
+
+	ProxyCommand<commands::EnableModulesEx>* proxy = new ProxyCommand<commands::EnableModulesEx>(true,
+		sizeof(commands::ModuleData) * moduleHandles.size() +
+		sizeof(commands::ModuleData) * lazyLoadModuleHandles.size() +
+		sizeof(uintptr_t) * reservedPagesCount);
+
+	proxy->m_command.processId = Process::Current::GetId();
+	proxy->m_command.moduleCount = static_cast<unsigned int>(moduleHandles.size());
+	proxy->m_command.lazyLoadModuleCount = static_cast<unsigned int>(lazyLoadModuleHandles.size());
+	proxy->m_command.reservedPagesCount = reservedPagesCount;
+	proxy->m_command.token = new Event(nullptr, Event::Type::AUTO_RESET);
+	AppendModuleHandles(proxy, moduleHandles);
+	AppendModuleHandles(proxy, lazyLoadModuleHandles);
+	for (unsigned int i = 0; i < reservedPagesCount; ++i)
+	{
+		proxy->m_payload.Write(reservedPages[i]);
+	}
+
+	PushUserCommand(proxy);
+
+	return proxy->m_command.token;
+
+}
+// END EPIC MOD
+
 void ClientUserCommandThread::InstallExceptionHandler(void)
 {
 	// BEGIN EPIC MOD - Using internal CrashReporter
diff --git a/Engine/Source/Developer/Windows/LiveCoding/Private/External/LC_ClientUserCommandThread.h b/Engine/Source/Developer/Windows/LiveCoding/Private/External/LC_ClientUserCommandThread.h
--- a/Engine/Source/Developer/Windows/LiveCoding/Private/External/LC_ClientUserCommandThread.h
+++ b/Engine/Source/Developer/Windows/LiveCoding/Private/External/LC_ClientUserCommandThread.h
@@ -100,6 +100,10 @@ public:
 	void DisableCompileFinishNotification();
 	// END EPIC MOD
 
+	// BEGIN EPIC MOD
+	void* EnableModulesEx(const wchar_t* moduleNames[], unsigned int moduleCount, const wchar_t* lazyLoadModuleNames[], unsigned int lazyLoadModuleCount, const uintptr_t* reservedPages, unsigned int reservedPagesCount);
+	// END EPIC MOD
+
 	void ApplySettingBool(const char* settingName, int value);
 	void ApplySettingInt(const char* settingName, int value);
 	void ApplySettingString(const char* settingName, const wchar_t* value);
diff --git a/Engine/Source/Developer/Windows/LiveCoding/Private/External/LC_Commands.h b/Engine/Source/Developer/Windows/LiveCoding/Private/External/LC_Commands.h
--- a/Engine/Source/Developer/Windows/LiveCoding/Private/External/LC_Commands.h
+++ b/Engine/Source/Developer/Windows/LiveCoding/Private/External/LC_Commands.h
@@ -353,6 +353,20 @@ namespace commands
 		Process::Id processId;
 	};
 
-	static const uint32_t COUNT = DisableCompileFinishNotification::ID + 1u;
+	// This is a version of EnableModules that includes support for lazy modules and pre-reserved page addresses.
+	struct EnableModulesEx
+	{
+		static const uint32_t ID = DisableCompileFinishNotification::ID + 1u;
+
+		Process::Id processId;
+		unsigned int moduleCount;
+		unsigned int lazyLoadModuleCount;
+		unsigned int reservedPagesCount;
+		void* token;
+
+		// this command always contains an array of 'ModuleData x moduleCount + ModuleData x lazyLoadModuleCount + sizeof(uintptr_t) * reservedPagesCount' as payload
+	};
+
+	static const uint32_t COUNT = EnableModulesEx::ID + 1u;
 	// END EPIC MOD
 }
diff --git a/Engine/Source/Developer/Windows/LiveCodingServer/Private/External/LC_LiveProcess.h b/Engine/Source/Developer/Windows/LiveCodingServer/Private/External/LC_LiveProcess.h
--- a/Engine/Source/Developer/Windows/LiveCodingServer/Private/External/LC_LiveProcess.h
+++ b/Engine/Source/Developer/Windows/LiveCodingServer/Private/External/LC_LiveProcess.h
@@ -119,6 +119,13 @@ public:
 	}
 	// END EPIC MOD
 
+	// BEGIN EPIC MOD
+	void AddPage(void *page)
+	{
+		m_virtualMemoryRange.AddPage(page);
+	}
+	// END EPIC MOD
+
 private:
 	Process::Handle m_processHandle;
 	Process::Id m_processId;
diff --git a/Engine/Source/Developer/Windows/LiveCodingServer/Private/External/LC_ServerCommandThread.cpp b/Engine/Source/Developer/Windows/LiveCodingServer/Private/External/LC_ServerCommandThread.cpp
--- a/Engine/Source/Developer/Windows/LiveCodingServer/Private/External/LC_ServerCommandThread.cpp
+++ b/Engine/Source/Developer/Windows/LiveCodingServer/Private/External/LC_ServerCommandThread.cpp
@@ -1460,6 +1460,9 @@ Thread::ReturnValue ServerCommandThread::CommandThread(DuplexPipeServer* pipe, E
 	commandMap.RegisterAction<actions::SetReinstancingFlow>();
 	commandMap.RegisterAction<actions::DisableCompileFinishNotification>();
 	// END EPIC MOD
+	// BEGIN EPIC MOD
+	commandMap.RegisterAction<actions::EnableModulesEx>();
+	// END EPIC MOD
 
 	for (;;)
 	{
@@ -2106,6 +2109,160 @@ bool ServerCommandThread::actions::EnableModules::Execute(const CommandType* com
 	return true;
 }
 
+// BEGIN EPIC MOD
+bool ServerCommandThread::actions::EnableModulesEx::Execute(const CommandType* command, const DuplexPipe* pipe, void* context, const void* payload, size_t payloadSize)
+{
+	pipe->SendAck();
+
+	ServerCommandThread* commandThread = static_cast<ServerCommandThread*>(context);
+
+	// protect against several client DLLs calling into this action at the same time.
+	// this ensures that all modules are loaded serialized per process.
+	CriticalSection::ScopedLock lock(&commandThread->m_actionCS);
+
+	telemetry::Scope moduleLoadingScope("Module loading ex");
+
+	// Locate the process
+	LiveProcess* liveProcess = nullptr;
+	for (LiveProcess* process : commandThread->m_liveProcesses)
+	{
+		if (process->GetProcessId() == command->processId)
+		{
+			liveProcess = process;
+			break;
+		}
+	}
+
+	// set up virtual drives before loading anything, otherwise files won't be detected and therefore discarded
+	AddVirtualDrive();
+
+	scheduler::TaskBase* rootTask = scheduler::CreateEmptyTask();
+	types::vector<scheduler::Task<LiveModule*>*> loadModuleTasks;
+
+	const unsigned int moduleCount = command->moduleCount;
+	loadModuleTasks.reserve(moduleCount);
+
+	memoryStream::Reader payloadStream(payload, payloadSize);
+	for (unsigned int i = 0u; i < moduleCount; ++i)
+	{
+		const commands::ModuleData moduleData = payloadStream.Read<commands::ModuleData>();
+		scheduler::Task<LiveModule*>* task = commandThread->LoadModule(command->processId, moduleData.base, moduleData.path, rootTask);
+
+		// the module could have failed to load
+		if (task)
+		{
+			loadModuleTasks.push_back(task);
+		}
+	}
+
+	const unsigned int lazyModuleCount = command->lazyLoadModuleCount;
+	unsigned int loadedLazyModules = 0;
+	for (unsigned int i = 0u; i < lazyModuleCount; ++i)
+	{
+		const commands::ModuleData moduleData = payloadStream.Read<commands::ModuleData>();
+
+		// Check if this module is already enabled - it may have been lazy-loaded, then fully loaded, by a restarted process. If so, translate this into a call to EnableModules.
+		bool isEnabled = false;
+		const std::wstring modulePath = Filesystem::NormalizePath(moduleData.path).GetString();
+		for (LiveModule* module : commandThread->m_liveModules)
+		{
+			if (module->GetModuleName() == modulePath)
+			{
+				scheduler::Task<LiveModule*>* task = commandThread->LoadModule(command->processId, moduleData.base, moduleData.path, rootTask);
+				if (task)
+				{
+					loadModuleTasks.push_back(task);
+				}
+				isEnabled = true;
+				break;
+			}
+		}
+		
+		if (!isEnabled && liveProcess != nullptr)
+		{
+			liveProcess->AddLazyLoadedModule(modulePath, static_cast<Windows::HMODULE>(moduleData.base));
+			LC_LOG_DEV("Registered module %S for lazy-loading", modulePath.c_str());
+			++loadedLazyModules;
+		}
+	}
+
+	const unsigned int reservedPagesCount = command->reservedPagesCount;
+	for (unsigned int i = 0u; i < reservedPagesCount; ++i)
+	{
+		const uintptr_t page = payloadStream.Read<uintptr_t>();
+		if (liveProcess != nullptr)
+		{
+			liveProcess->AddPage(reinterpret_cast<void*>(page));
+		}
+	}
+
+	// wait for all tasks to finish
+	scheduler::RunTask(rootTask);
+	scheduler::WaitForTask(rootTask);
+
+	const size_t loadModuleTaskCount = loadModuleTasks.size();
+	commandThread->m_liveModules.reserve(loadModuleTaskCount);
+
+	size_t loadedTranslationUnits = 0u;
+
+	// update all live modules loaded by the tasks
+	for (size_t i = 0u; i < loadModuleTaskCount; ++i)
+	{
+		scheduler::Task<LiveModule*>* task = loadModuleTasks[i];
+		LiveModule* liveModule = task->GetResult();
+
+		commandThread->m_liveModules.push_back(liveModule);
+
+		// update directory cache for this live module
+		liveModule->UpdateDirectoryCache(commandThread->m_directoryCache);
+
+		// update the number of loaded translation units
+		loadedTranslationUnits += liveModule->GetCompilandDatabase()->compilands.size();
+	}
+
+	scheduler::DestroyTasks(loadModuleTasks);
+	scheduler::DestroyTask(rootTask);
+
+	// dump memory statistics
+	{
+		LC_LOG_INDENT_TELEMETRY;
+		g_symbolAllocator.PrintStats();
+		g_immutableStringAllocator.PrintStats();
+		g_contributionAllocator.PrintStats();
+		g_compilandAllocator.PrintStats();
+		g_dependencyAllocator.PrintStats();
+	}
+
+	if (+commandThread->m_compileThread == nullptr || Thread::Current::GetId() != Thread::GetId(commandThread->m_compileThread))
+	{
+		if (loadModuleTaskCount > 0u)
+		{
+			LC_SUCCESS_USER("Loaded %zu module(s), %zu lazy load module(s), and %zu reserved page ranges (%.3fs, %zu translation units)", 
+				loadModuleTaskCount, loadedLazyModules, reservedPagesCount, moduleLoadingScope.ReadSeconds(), loadedTranslationUnits);
+		}
+
+		// tell user we are ready, but only once to not clutter the log
+		{
+			static bool showedOnce = false;
+			if (!showedOnce)
+			{
+				showedOnce = true;
+				const int shortcut = appSettings::g_compileShortcut->GetValue();
+				const std::wstring& shortcutText = shortcut::ConvertShortcutToText(shortcut);
+				LC_SUCCESS_USER("Live coding ready - Save changes and press %S to re-compile code", shortcutText.c_str());
+			}
+		}
+	}
+
+	// remove virtual drives once we're finished
+	RemoveVirtualDrive();
+
+	// tell server we are finished
+	pipe->SendCommandAndWaitForAck(commands::EnableModulesFinished{ command->token }, nullptr, 0u);
+
+	return true;
+}
+// END EPIC MOD
 
 bool ServerCommandThread::actions::DisableModules::Execute(const CommandType* command, const DuplexPipe* pipe, void* context, const void* payload, size_t payloadSize)
 {
diff --git a/Engine/Source/Developer/Windows/LiveCodingServer/Private/External/LC_ServerCommandThread.h b/Engine/Source/Developer/Windows/LiveCodingServer/Private/External/LC_ServerCommandThread.h
--- a/Engine/Source/Developer/Windows/LiveCodingServer/Private/External/LC_ServerCommandThread.h
+++ b/Engine/Source/Developer/Windows/LiveCodingServer/Private/External/LC_ServerCommandThread.h
@@ -134,6 +134,10 @@ private:
 		DECLARE_ACTION(DisableCompileFinishNotification);
 		// END EPIC MOD
 
+		// BEGIN EPIC MOD
+		DECLARE_ACTION(EnableModulesEx);
+		// END EPIC MOD
+
 		#undef DECLARE_ACTION
 	};
 
diff --git a/Engine/Source/Developer/Windows/LiveCodingServer/Private/External/LC_VirtualMemoryRange.cpp b/Engine/Source/Developer/Windows/LiveCodingServer/Private/External/LC_VirtualMemoryRange.cpp
--- a/Engine/Source/Developer/Windows/LiveCodingServer/Private/External/LC_VirtualMemoryRange.cpp
+++ b/Engine/Source/Developer/Windows/LiveCodingServer/Private/External/LC_VirtualMemoryRange.cpp
@@ -115,3 +115,12 @@ void VirtualMemoryRange::FreePages(const void* addressStart, const void* address
 		}
 	}
 }
+
+
+// BEGIN EPIC MOD
+void VirtualMemoryRange::AddPage(void* page)
+{
+	m_pageData.emplace_back(PageData{ page });
+}
+// END EPIC MOD
+
diff --git a/Engine/Source/Developer/Windows/LiveCodingServer/Private/External/LC_VirtualMemoryRange.h b/Engine/Source/Developer/Windows/LiveCodingServer/Private/External/LC_VirtualMemoryRange.h
--- a/Engine/Source/Developer/Windows/LiveCodingServer/Private/External/LC_VirtualMemoryRange.h
+++ b/Engine/Source/Developer/Windows/LiveCodingServer/Private/External/LC_VirtualMemoryRange.h
@@ -21,6 +21,10 @@ public:
 	void ReservePages(const void* addressStart, const void* addressEnd, size_t alignment);
 	void FreePages(const void* addressStart, const void* addressEnd);
 
+	// BEGIN EPIC MOD
+	void AddPage(void* page);
+	// END EPIC MOD
+
 private:
 	struct PageData
 	{
