Fix compilation issue with VS2019

#rnx

diff --git a/Engine/Source/Developer/Windows/LiveCoding/Private/External/LC_API.h b/Engine/Source/Developer/Windows/LiveCoding/Private/External/LC_API.h
--- a/Engine/Source/Developer/Windows/LiveCoding/Private/External/LC_API.h
+++ b/Engine/Source/Developer/Windows/LiveCoding/Private/External/LC_API.h
@@ -2,6 +2,9 @@
 
 #pragma once
 
+// BEGIN EPIC MOD
+#include "stddef.h"
+// END EPIC MOD
 #include "LPP_API.h"
 
 
