Fixed a message formatting issue

#rnx

diff --git a/Engine/Source/Developer/Windows/LiveCodingServer/Private/External/LC_ServerCommandThread.cpp b/Engine/Source/Developer/Windows/LiveCodingServer/Private/External/LC_ServerCommandThread.cpp
--- a/Engine/Source/Developer/Windows/LiveCodingServer/Private/External/LC_ServerCommandThread.cpp
+++ b/Engine/Source/Developer/Windows/LiveCodingServer/Private/External/LC_ServerCommandThread.cpp
@@ -2237,7 +2237,7 @@ bool ServerCommandThread::actions::EnableModulesEx::Execute(const CommandType* c
 	{
 		if (loadModuleTaskCount > 0u)
 		{
-			LC_SUCCESS_USER("Loaded %zu module(s), %zu lazy load module(s), and %zu reserved page ranges (%.3fs, %zu translation units)", 
+			LC_SUCCESS_USER("Loaded %zu module(s), %u lazy load module(s), and %u reserved page ranges (%.3fs, %zu translation units)", 
 				loadModuleTaskCount, loadedLazyModules, reservedPagesCount, moduleLoadingScope.ReadSeconds(), loadedTranslationUnits);
 		}
 
