Fixed issue with LiveCoding where large projects might not have enough memory reserved around libraries preventing patching.

diff --git a/Engine/Source/Developer/Windows/LiveCoding/Private/LiveCodingModule.cpp b/Engine/Source/Developer/Windows/LiveCoding/Private/LiveCodingModule.cpp
--- a/Engine/Source/Developer/Windows/LiveCoding/Private/LiveCodingModule.cpp
+++ b/Engine/Source/Developer/Windows/LiveCoding/Private/LiveCodingModule.cpp
@@ -156,6 +156,127 @@ private:
 };
 #endif
 
+namespace LivePP
+{
+
+	// Based on liveProcess::GetLowerBoundIn4GBRange
+	uintptr_t GetLowerBound(uintptr_t ModuleBase, uint64 Padding)
+	{
+		constexpr uintptr_t LowestPossibleAddress = 64ull * 1024ull;
+		return ModuleBase >= LowestPossibleAddress + Padding ? ModuleBase - Padding : LowestPossibleAddress;
+	}
+
+	// Based on liveProcess::GetUpperBoundIn4GBRange
+	uintptr_t GetUpperBound(uintptr_t ModuleBase, uint64 Padding)
+	{
+		constexpr uintptr_t HighestPossibleAddress = 0x00007FFFFFFF0000ull;
+		return ModuleBase <= HighestPossibleAddress - Padding ? ModuleBase + Padding : HighestPossibleAddress;
+	}
+
+	// Based on VirtualMemoryRange::ReservePages
+	void ReservePages(TArray<uintptr_t>& ReservedPages, uintptr_t AddressStart, uintptr_t AddressEnd, uint64 PageAlignment)
+	{
+		HANDLE ProcessHandle = GetCurrentProcess();
+
+		// reserve all free pages in the virtual memory range.
+		// pages must be aligned to the given alignment.
+		for (uintptr_t Address = AddressStart; Address < AddressEnd; /* nothing */)
+		{
+			// align address to be scanned
+			Address = Align(Address, PageAlignment);
+
+			if (Address < AddressStart)
+			{
+				// overflow happened because we scanned too far
+				break;
+			}
+
+			::MEMORY_BASIC_INFORMATION MemoryInfo = {};
+			const size_t BytesReturned = ::VirtualQueryEx(ProcessHandle, (const void*)Address, &MemoryInfo, sizeof(::MEMORY_BASIC_INFORMATION));
+			if (BytesReturned == 0)
+			{
+				break;
+			}
+
+			// we are only interested in free pages
+			if (MemoryInfo.State == MEM_FREE)
+			{
+				// work out the maximum size of the page allocation.
+				// we should not allocate past the end of the range.
+				const uint64 BytesLeft = AddressEnd - (uintptr_t)MemoryInfo.BaseAddress;
+				const uint64 Size = std::min<uint64>(MemoryInfo.RegionSize, BytesLeft);
+
+				// try to reserve this page.
+				// if we are really unlucky, the process might have allocated this region in the meantime.
+				void* BaseAddress = ::VirtualAllocEx(ProcessHandle, MemoryInfo.BaseAddress, Size, MEM_RESERVE, PAGE_NOACCESS);
+				if (BaseAddress)
+				{
+					ReservedPages.Add((uintptr_t)BaseAddress);
+				}
+			}
+
+			// keep on searching
+			Address = (uintptr_t)MemoryInfo.BaseAddress + MemoryInfo.RegionSize;
+		}
+	}
+}
+
+namespace
+{
+	constexpr uint64 DefaultPadding = 128ull * 1024ull * 1024ull; // LivePP used 2ull * 1024ull * 1024ull * 1024ull;
+	constexpr uint64 DefaultPageAlignment = 64ull * 1024ull;
+	constexpr uint32 NewModuleCountThreshhold = 128;
+
+	void ReservePagesBefore(TArray<uintptr_t>& ReservedPages, uintptr_t ModuleBase, uint64 Padding, uint32 PageAlignment)
+	{
+		LivePP::ReservePages(ReservedPages, LivePP::GetLowerBound(ModuleBase, Padding), ModuleBase, PageAlignment);
+	}
+
+	void ReservePagesBetween(TArray<uintptr_t>& ReservedPages, uintptr_t StartModuleBase, uintptr_t EndModuleBase, uint64 Padding, uint32 PageAlignment)
+	{
+		LivePP::ReservePages(ReservedPages, LivePP::GetLowerBound(StartModuleBase, Padding), LivePP::GetUpperBound(EndModuleBase, Padding), PageAlignment);
+	}
+
+	void ReservePagesAfter(TArray<uintptr_t>& ReservedPages, uintptr_t ModuleBase, uint64 Padding, uint32 PageAlignment)
+	{
+		LivePP::ReservePages(ReservedPages, ModuleBase, LivePP::GetUpperBound(ModuleBase, Padding), PageAlignment);
+	}
+
+	void ReservePages(TArray<uintptr_t>& ReservedPages, const TArray<uintptr_t>& ModuleBases, uint64 Padding, uint64 PageAlignment)
+	{
+		if (ModuleBases.IsEmpty())
+		{
+			return;
+		}
+
+		// We use an ordered list to avoid rescanning the same ranges over and over again
+		TArray<uintptr_t> SortedModuleBases(ModuleBases);
+		SortedModuleBases.Sort();
+
+		// Reserve prior to the first module
+		ReservePagesBefore(ReservedPages, SortedModuleBases[0], Padding, PageAlignment);
+
+		// Loop through the pairs of modules
+		for (int i = 0; i < SortedModuleBases.Num() - 1; ++i)
+		{
+			uintptr_t StartModuleBase = SortedModuleBases[i];
+			uintptr_t EndModuleBase = SortedModuleBases[i + 1];
+			if (EndModuleBase - StartModuleBase < Padding * 2)
+			{
+				ReservePagesBetween(ReservedPages, StartModuleBase, EndModuleBase, Padding, PageAlignment);
+			}
+			else
+			{
+				ReservePagesAfter(ReservedPages, StartModuleBase, Padding, PageAlignment);
+				ReservePagesBefore(ReservedPages, EndModuleBase, Padding, PageAlignment);
+			}
+		}
+
+		// Reserve after the last module
+		ReservePagesAfter(ReservedPages, SortedModuleBases.Last(), Padding, PageAlignment);
+	}
+}
+
 // Helper structure to load the NTDLL library and work with an API
 struct FNtDllFunction
 {
@@ -951,20 +1072,7 @@ bool FLiveCodingModule::StartLiveCoding(ELiveCodingStartupMode StartupMode)
 		FString ProcessGroup = FString::Printf(TEXT("UE_%s_0x%08x"), FApp::GetProjectName(), GetTypeHash(ProjectPath));
 		LppRegisterProcessGroup(TCHAR_TO_ANSI(*ProcessGroup));
 
-		// Build the command line
-		FString KnownTargetName = FPlatformMisc::GetUBTTargetName();
-		FString Arguments = FString::Printf(TEXT("%s %s %s"),
-			*KnownTargetName,
-			FPlatformMisc::GetUBTPlatform(),
-			LexToString(FApp::GetBuildConfiguration()));
-
-		UE_LOG(LogLiveCoding, Display, TEXT("LiveCodingConsole Arguments: %s"), *Arguments);
-
-		if(SourceProject.Len() > 0)
-		{
-			Arguments += FString::Printf(TEXT(" -Project=\"%s\""), *FPaths::ConvertRelativePathToFull(SourceProject));
-		}
-		LppSetBuildArguments(*Arguments);
+		SetBuildArguments();
 
 #if WITH_EDITOR
 		if (IsReinstancingEnabled())
@@ -1012,22 +1120,63 @@ bool FLiveCodingModule::StartLiveCoding(ELiveCodingStartupMode StartupMode)
 			ShowConsole(true);
 		}
 
+		// Register a delegate to listen for new modules loaded from this point onwards
+		ModulesChangedDelegateHandle = FModuleManager::Get().OnModulesChanged().AddRaw(this, &FLiveCodingModule::OnModulesChanged);
+
 		// Mark it as started
 		State = EState::RunningAndEnabled;
 	}
 	return true;
 }
 
-void FLiveCodingModule::UpdateModules(bool bAllowStarting)
+void FLiveCodingModule::SetBuildArguments()
 {
-	TArray<ModuleChange> Changes;
+	FScopeLock lock(&SetBuildArgumentsCs);
+
+	// Build the command line
+	FString KnownTargetName = FPlatformMisc::GetUBTTargetName();
+	if (KnownTargetName != LastKnownTargetName)
 	{
-		FScopeLock lock(&ModuleChangeCs);
-		Swap(Changes, ModuleChanges);
+		LastKnownTargetName = KnownTargetName;
+		FString Arguments = FString::Printf(TEXT("%s %s %s"),
+			*KnownTargetName,
+			FPlatformMisc::GetUBTPlatform(),
+			LexToString(FApp::GetBuildConfiguration()));
+
+		UE_LOG(LogLiveCoding, Display, TEXT("LiveCodingConsole Arguments: %s"), *Arguments);
+
+		FString SourceProject = SourceProjectVariable->GetString();
+		if (SourceProject.Len() > 0)
+		{
+			Arguments += FString::Printf(TEXT(" -Project=\"%s\""), *FPaths::ConvertRelativePathToFull(SourceProject));
+		}
+		LppSetBuildArguments(*Arguments);
 	}
+}
 
+void FLiveCodingModule::UpdateModules(bool bAllowStarting)
+{
 	if (HasStarted(bAllowStarting))
 	{
+		// The target name might have changed now that we are loading early
+		SetBuildArguments();
+
+		TArray<ModuleChange> Changes;
+		TArray<uintptr_t> Pages;
+		FGraphEventRef Task;
+		{
+			FScopeLock lock(&ModuleChangeCs);
+			Swap(Changes, ModuleChanges);
+			Swap(Pages, ReservedPages);
+			Task = ReservePagesTaskRef;
+			LastReservePagesModuleCount = 0;
+		}
+
+		if (Task.IsValid())
+		{
+			FTaskGraphInterface::Get().WaitUntilTaskCompletes(Task);
+		}
+
 #if IS_MONOLITHIC
 		wchar_t FullFilePath[WINDOWS_MAX_PATH];
 		verify(GetModuleFileName(hInstance, FullFilePath, UE_ARRAY_COUNT(FullFilePath)));
@@ -1053,6 +1202,7 @@ void FLiveCodingModule::UpdateModules(bool bAllowStarting)
 		}
 
 		TArray<FString> EnableModules;
+		TArray<FString> LazyLoadModules;
 		for (const ModuleChange& Change : Changes)
 		{
 			if (Change.bLoaded)
@@ -1065,9 +1215,7 @@ void FLiveCodingModule::UpdateModules(bool bAllowStarting)
 				}
 				else
 				{
-					TRACE_CPUPROFILER_EVENT_SCOPE(LppEnableLazyLoadedModule);
-					void* LppEnableLazyLoadedModuleToken = LppEnableLazyLoadedModule(*FullFilePath);
-					LppPendingTokens.Add(LppEnableLazyLoadedModuleToken);
+					LazyLoadModules.Add(FullFilePath);
 				}
 			}
 			else
@@ -1084,16 +1232,66 @@ void FLiveCodingModule::UpdateModules(bool bAllowStarting)
 				EnableModuleFileNames.Add(*EnableModule);
 			}
 
+			TArray<const TCHAR*> LazyLoadModuleFileNames;
+			for (const FString& LazyLoadModule : LazyLoadModules)
+			{
+				LazyLoadModuleFileNames.Add(*LazyLoadModule);
+			}
+
 			{
-				TRACE_CPUPROFILER_EVENT_SCOPE(LppEnableModules);
-				void* LppEnableModulesToken = LppEnableModules(EnableModuleFileNames.GetData(), EnableModuleFileNames.Num());
-				LppPendingTokens.Add(LppEnableModulesToken);
+				TRACE_CPUPROFILER_EVENT_SCOPE(LppEnableModulesEx);
+				void* LppEnableModulesExToken = LppEnableModulesEx(
+					EnableModuleFileNames.GetData(), EnableModuleFileNames.Num(),
+					LazyLoadModuleFileNames.GetData(), LazyLoadModuleFileNames.Num(),
+					Pages.GetData(), Pages.Num());
+				LppPendingTokens.Add(LppEnableModulesExToken);
 			}
 		}
 #endif
 	}
 }
 
+void FLiveCodingModule::ReservePagesTask()
+{
+	for (;;)
+	{
+		TArray<ModuleChange> Changes;
+		{
+			FScopeLock lock(&ModuleChangeCs);
+			uint32 NewModules = ModuleChanges.Num() - LastReservePagesModuleCount;
+			if (NewModules < NewModuleCountThreshhold)
+			{
+				ReservePagesTaskRef = FGraphEventRef();
+				return;
+			}
+			Changes.Append(ModuleChanges.GetData() + LastReservePagesModuleCount, NewModules);
+			LastReservePagesModuleCount = ModuleChanges.Num();
+		}
+
+		TRACE_CPUPROFILER_EVENT_SCOPE(LiveCodingReservePages);
+
+		TArray<uintptr_t> ModuleBases;
+		ModuleBases.Reserve(Changes.Num());
+		for (const ModuleChange& Change : Changes)
+		{
+			if (Change.bLoaded)
+			{
+				HMODULE Handle = ::GetModuleHandle(*Change.FullName.ToString());
+				ModuleBases.Add(reinterpret_cast<uintptr_t>(Handle));
+			}
+		}
+
+		TArray<uintptr_t> Pages;
+		ReservePages(Pages, ModuleBases, DefaultPadding, DefaultPageAlignment);
+
+		{
+			FScopeLock lock(&ModuleChangeCs);
+			ReservedPages.Append(Pages);
+		}
+	}
+}
+
+
 void FLiveCodingModule::OnModulesChanged(FName ModuleName, EModuleChangeReason Reason)
 {
 #if !IS_MONOLITHIC
@@ -1237,6 +1435,10 @@ void FLiveCodingModule::OnDllLoaded(const FString& FullPath)
 	{
 		FScopeLock lock(&ModuleChangeCs);
 		ModuleChanges.Emplace(ModuleChange{ FName(FullPath), true });
+		if (!ReservePagesTaskRef.IsValid() && ModuleChanges.Num() - LastReservePagesModuleCount >= NewModuleCountThreshhold)
+		{
+			ReservePagesTaskRef = FFunctionGraphTask::CreateAndDispatchWhenReady([this]() { ReservePagesTask(); });
+		}
 	}
 }
 
diff --git a/Engine/Source/Developer/Windows/LiveCoding/Private/LiveCodingModule.h b/Engine/Source/Developer/Windows/LiveCoding/Private/LiveCodingModule.h
--- a/Engine/Source/Developer/Windows/LiveCoding/Private/LiveCodingModule.h
+++ b/Engine/Source/Developer/Windows/LiveCoding/Private/LiveCodingModule.h
@@ -8,6 +8,7 @@
 #include "Modules/ModuleManager.h"
 #include "Templates/SharedPointer.h"
 #include "Internationalization/Text.h"
+#include "Async/TaskGraphFwd.h"
 #include <atomic>
 
 struct IConsoleCommand;
@@ -63,11 +64,13 @@ private:
 	bool StartLiveCoding(ELiveCodingStartupMode StartupMode);
 	void OnModulesChanged(FName ModuleName, EModuleChangeReason Reason);
 	void UpdateModules(bool bAllowStarting);
+	void ReservePagesTask();
 	bool ShouldPreloadModule(const TSet<FName>& PreloadedFileNames, const FString& FullFilePath) const;
 	bool IsReinstancingEnabled() const;
 	void WaitForStartup();
 	bool HasStarted(bool bAllowStarting) const;
 	void ShowConsole(bool bAllowStarting);
+	void SetBuildArguments();
 
 #if WITH_EDITOR
 	void ShowNotification(bool Success, const FText& Title, const FText* SubText);
@@ -86,14 +89,16 @@ private:
 	TSharedPtr<ISettingsSection> SettingsSection;
 	bool bSettingsEnabledLastTick = false;
 	bool bEnableReinstancingLastTick = false;
+	bool bBuildArgumentsSet = false;
 	std::atomic<EState> State = EState::NotRunning;
-	bool bUpdateModulesInTick = false;
+	std::atomic<bool> bUpdateModulesInTick = false;
 	bool bHasReinstancingOccurred = false;
 	bool bHasPatchBeenLoaded = false;
 	ELiveCodingCompileResult LastResults = ELiveCodingCompileResult::Success;
 	TSet<FName> ConfiguredModules;
 	TArray<void*> LppPendingTokens;
 	void* CallbackCookie = nullptr;
+	FString LastKnownTargetName = FString();
 
 	FText EnableErrorText;
 
@@ -115,8 +120,12 @@ private:
 		FName FullName;
 		bool bLoaded;
 	};
+	FCriticalSection SetBuildArgumentsCs;
 	FCriticalSection ModuleChangeCs;
 	TArray<ModuleChange> ModuleChanges;
+	TArray<uintptr_t> ReservedPages;
+	int LastReservePagesModuleCount = 0;
+	FGraphEventRef ReservePagesTaskRef;
 
 #if WITH_EDITOR
 	TUniquePtr<FReload> Reload;
diff --git a/Engine/Source/Runtime/Engine/Private/UnrealEngine.cpp b/Engine/Source/Runtime/Engine/Private/UnrealEngine.cpp
--- a/Engine/Source/Runtime/Engine/Private/UnrealEngine.cpp
+++ b/Engine/Source/Runtime/Engine/Private/UnrealEngine.cpp
@@ -2080,11 +2080,6 @@ void UEngine::Init(IEngineLoop* InEngineLoop)
 	// Record large world coordinate state
 	CSV_METADATA(TEXT("LargeWorldCoordinates"), TEXT("1"));	// LWC_TODO: Remove
 
-	// Enable the live coding module if it's a developer build
-#if WITH_LIVE_CODING
-	FModuleManager::Get().LoadModule("LiveCoding");
-#endif
-
 	// Finish asset manager loading
 	if (AssetManager)
 	{
diff --git a/Engine/Source/Runtime/Launch/Private/LaunchEngineLoop.cpp b/Engine/Source/Runtime/Launch/Private/LaunchEngineLoop.cpp
--- a/Engine/Source/Runtime/Launch/Private/LaunchEngineLoop.cpp
+++ b/Engine/Source/Runtime/Launch/Private/LaunchEngineLoop.cpp
@@ -4447,6 +4447,11 @@ bool FEngineLoop::LoadStartupCoreModules()
 		FModuleManager::Get().LoadModule(TEXT("Networking"));
 	}
 
+	// Enable the live coding module if it's a developer build
+#if WITH_LIVE_CODING
+	FModuleManager::Get().LoadModule("LiveCoding");
+#endif
+
 	SlowTask.EnterProgressFrame(10);
 	FPlatformApplicationMisc::LoadStartupModules();
 
