[shaders] d3d (both dxc & fxc) compile backend multithreading fixes; stripping out vestigial path allowing FXC DLL to be overridden by a parameter.

diff --git a/Engine/Source/Developer/Windows/ShaderFormatD3D/Private/D3DShaderCompiler.cpp b/Engine/Source/Developer/Windows/ShaderFormatD3D/Private/D3DShaderCompiler.cpp
--- a/Engine/Source/Developer/Windows/ShaderFormatD3D/Private/D3DShaderCompiler.cpp
+++ b/Engine/Source/Developer/Windows/ShaderFormatD3D/Private/D3DShaderCompiler.cpp
@@ -378,67 +378,43 @@ typedef HRESULT(WINAPI *pD3DStripShader)
 // look for IID_ID3D11ShaderReflection in d3d11shader.h for the SDK matching the compiler DLL.
 DEFINE_GUID_FOR_CURRENT_COMPILER(IID_ID3D11ShaderReflectionForCurrentCompiler, 0x8d536ca1, 0x0cca, 0x4956, 0xa8, 0x37, 0x78, 0x69, 0x63, 0x75, 0x55, 0x84);
 
-/**
- * GetD3DCompilerFuncs - gets function pointers from the dll at NewCompilerPath
- * @param OutD3DCompile - function pointer for D3DCompile (0 if not found)
- * @param OutD3DReflect - function pointer for D3DReflect (0 if not found)
- * @param OutD3DDisassemble - function pointer for D3DDisassemble (0 if not found)
- * @param OutD3DStripShader - function pointer for D3DStripShader (0 if not found)
- * @return bool - true if functions were retrieved from NewCompilerPath
- */
-static bool GetD3DCompilerFuncs(const FString& NewCompilerPath, pD3DCompile* OutD3DCompile,
-	pD3DReflect* OutD3DReflect, pD3DDisassemble* OutD3DDisassemble, pD3DStripShader* OutD3DStripShader)
+// Helper class to load the engine-packaged FXC DLL and retrieve function pointers for the various FXC functions from it.
+class FxcCompilerFunctions
 {
-	static FString CurrentCompiler;
-	static HMODULE CompilerDLL = 0;
-
-	if(CurrentCompiler != *NewCompilerPath)
-	{
-		CurrentCompiler = *NewCompilerPath;
+public:
 
-		if(CompilerDLL)
-		{
-			FreeLibrary(CompilerDLL);
-			CompilerDLL = 0;
-		}
+	static pD3DCompile GetCompile() { return Instance().Compile; }
+	static pD3DReflect GetReflect() { return Instance().Reflect; }
+	static pD3DDisassemble GetDisassemble() { return Instance().Disassemble; }
+	static pD3DStripShader GetStripShader() { return Instance().StripShader; }
 
-		if(CurrentCompiler.Len())
-		{
-			CompilerDLL = LoadLibrary(*CurrentCompiler);
-		}
-
-		if(!CompilerDLL && NewCompilerPath.Len())
+private:
+	FxcCompilerFunctions()
+	{
+		FString CompilerPath = FPaths::EngineDir() / TEXT("Binaries/ThirdParty/Windows/DirectX/x64/d3dcompiler_47.dll");
+		CompilerDLL = LoadLibrary(*CompilerPath);
+		if (!CompilerDLL)
 		{
-			// Couldn't find HLSL compiler in specified path. We fail the first compile.
-			*OutD3DCompile = 0;
-			*OutD3DReflect = 0;
-			*OutD3DDisassemble = 0;
-			*OutD3DStripShader = 0;
-			return false;
+			UE_LOG(LogD3D11ShaderCompiler, Fatal, TEXT("Cannot find the compiler DLL '%s'"), *CompilerPath);
 		}
+		Compile = (pD3DCompile)(void*)GetProcAddress(CompilerDLL, "D3DCompile");
+		Reflect = (pD3DReflect)(void*)GetProcAddress(CompilerDLL, "D3DReflect");
+		Disassemble = (pD3DDisassemble)(void*)GetProcAddress(CompilerDLL, "D3DDisassemble");
+		StripShader = (pD3DStripShader)(void*)GetProcAddress(CompilerDLL, "D3DStripShader");
 	}
 
-	if(CompilerDLL)
+	static FxcCompilerFunctions& Instance()
 	{
-		// from custom folder e.g. "C:/DXWin8/D3DCompiler_44.dll"
-		*OutD3DCompile = (pD3DCompile)(void*)GetProcAddress(CompilerDLL, "D3DCompile");
-		*OutD3DReflect = (pD3DReflect)(void*)GetProcAddress(CompilerDLL, "D3DReflect");
-		*OutD3DDisassemble = (pD3DDisassemble)(void*)GetProcAddress(CompilerDLL, "D3DDisassemble");
-		*OutD3DStripShader = (pD3DStripShader)(void*)GetProcAddress(CompilerDLL, "D3DStripShader");
-		return true;
+		static FxcCompilerFunctions Instance;
+		return Instance;
 	}
 
-    // if we cannot find the bundled DLL, this is a fatal error. We _do_not_ want to use a system-specific library as it can make the shaders (and DDC) system-specific
-	UE_LOG(LogD3D11ShaderCompiler, Fatal, TEXT("Cannot find the compiler DLL '%s'"), *CurrentCompiler);
-#if 0
-	// D3D SDK we compiled with (usually D3DCompiler_43.dll from windows folder)
-	*OutD3DCompile = &D3DCompile;
-	*OutD3DReflect = &D3DReflect;
-	*OutD3DDisassemble = &D3DDisassemble;
-	*OutD3DStripShader = &D3DStripShader;
-#endif
-	return false;
-}
+	HMODULE CompilerDLL = 0;
+	pD3DCompile Compile = nullptr;
+	pD3DReflect Reflect = nullptr;
+	pD3DDisassemble Disassemble = nullptr;
+	pD3DStripShader StripShader = nullptr;
+};
 
 static int D3DExceptionFilter(bool bCatchException)
 {
@@ -572,7 +548,7 @@ static bool CompileErrorsContainInternalError(ID3DBlob* Errors)
 
 // Generate the dumped usf file; call the D3D compiler, gather reflection information and generate the output data
 static bool CompileAndProcessD3DShaderFXCExt(
-	FString& PreprocessedShaderSource, const FString& CompilerPath,
+	FString& PreprocessedShaderSource,
 	uint32 CompileFlags,
 	const FShaderCompilerInput& Input,
 	const FShaderParameterParser& ShaderParameterParser,
@@ -608,11 +584,11 @@ static bool CompileAndProcessD3DShaderFXCExt(
 	TRefCountPtr<ID3DBlob> Shader;
 
 	HRESULT Result = S_OK;
-	pD3DCompile D3DCompileFunc = nullptr;
-	pD3DReflect D3DReflectFunc = nullptr;
-	pD3DDisassemble D3DDisassembleFunc = nullptr;
-	pD3DStripShader D3DStripShaderFunc = nullptr;
-	bool bCompilerPathFunctionsUsed = GetD3DCompilerFuncs(CompilerPath, &D3DCompileFunc, &D3DReflectFunc, &D3DDisassembleFunc, &D3DStripShaderFunc);
+	pD3DCompile D3DCompileFunc = FxcCompilerFunctions::GetCompile();
+	pD3DReflect D3DReflectFunc = FxcCompilerFunctions::GetReflect();
+	pD3DDisassemble D3DDisassembleFunc = FxcCompilerFunctions::GetDisassemble();
+	pD3DStripShader D3DStripShaderFunc = FxcCompilerFunctions::GetStripShader();
+
 	TRefCountPtr<ID3DBlob> Errors;
 
 	if (D3DCompileFunc)
@@ -720,7 +696,7 @@ static bool CompileAndProcessD3DShaderFXCExt(
 	}
 	else
 	{
-		FilteredErrors.Add(FString::Printf(TEXT("Couldn't find shader compiler: %s"), *CompilerPath));
+		FilteredErrors.Add(TEXT("Couldn't find D3D shader compiler DLL"));
 		Result = E_FAIL;
 	}
 
@@ -802,10 +778,7 @@ static bool CompileAndProcessD3DShaderFXCExt(
 			Output.bSucceeded = true;
 			ID3D11ShaderReflection* Reflector = NULL;
 
-			// IID_ID3D11ShaderReflectionForCurrentCompiler is defined in this file and needs to match the IID from the dll in CompilerPath
-			// if the function pointers from that dll are being used
-			const IID ShaderReflectionInterfaceID = bCompilerPathFunctionsUsed ? IID_ID3D11ShaderReflectionForCurrentCompiler : IID_ID3D11ShaderReflection;
-			Result = D3DReflectFunc(Shader->GetBufferPointer(), Shader->GetBufferSize(), ShaderReflectionInterfaceID, (void**)&Reflector);
+			Result = D3DReflectFunc(Shader->GetBufferPointer(), Shader->GetBufferSize(), IID_ID3D11ShaderReflectionForCurrentCompiler, (void**)&Reflector);
 			if (FAILED(Result))
 			{
 				UE_LOG(LogD3D11ShaderCompiler, Fatal, TEXT("D3DReflect failed: Result=%08x"), Result);
@@ -877,7 +850,7 @@ static bool CompileAndProcessD3DShaderFXCExt(
 						if (RemoveUnusedInputs(PreprocessedShaderSource, ShaderInputs, EntryPointName, RemoveErrors))
 						{
 							Output = OriginalOutput;
-							if (!CompileAndProcessD3DShaderFXCExt(PreprocessedShaderSource, CompilerPath, CompileFlags, Input, ShaderParameterParser, EntryPointName, ShaderProfile, true, FilteredErrors, Output))
+							if (!CompileAndProcessD3DShaderFXCExt(PreprocessedShaderSource, CompileFlags, Input, ShaderParameterParser, EntryPointName, ShaderProfile, true, FilteredErrors, Output))
 							{
 								// if we failed to compile the shader, propagate the error up
 								return false;
@@ -960,7 +933,7 @@ static bool CompileAndProcessD3DShaderFXCExt(
 		}
 		else
 		{
-			FilteredErrors.Add(FString::Printf(TEXT("Couldn't find shader reflection function in %s"), *CompilerPath));
+			FilteredErrors.Add(TEXT("Couldn't find shader reflection function in D3D Compiler DLL"));
 			Result = E_FAIL;
 			Output.bSucceeded = false;
 		}
@@ -1060,12 +1033,8 @@ bool CompileAndProcessD3DShaderFXC(FString& PreprocessedShaderSource,
 	const TCHAR* ShaderProfile, bool bSecondPassAferUnusedInputRemoval,
 	FShaderCompilerOutput& Output)
 {
-	// Override default compiler path to newer dll
-	FString CompilerPath = FPaths::EngineDir();
-	CompilerPath.Append(TEXT("Binaries/ThirdParty/Windows/DirectX/x64/d3dcompiler_47.dll"));
-
 	TArray<FString> FilteredErrors;
-	const bool bSuccess = CompileAndProcessD3DShaderFXCExt(PreprocessedShaderSource, CompilerPath, CompileFlags, Input, ShaderParameterParser, EntryPointName, ShaderProfile, false, FilteredErrors, Output);
+	const bool bSuccess = CompileAndProcessD3DShaderFXCExt(PreprocessedShaderSource, CompileFlags, Input, ShaderParameterParser, EntryPointName, ShaderProfile, false, FilteredErrors, Output);
 
 	// Process errors
 	for (int32 ErrorIndex = 0; ErrorIndex < FilteredErrors.Num(); ErrorIndex++)
diff --git a/Engine/Source/Developer/Windows/ShaderFormatD3D/Private/D3DShaderCompilerDXC.cpp b/Engine/Source/Developer/Windows/ShaderFormatD3D/Private/D3DShaderCompilerDXC.cpp
--- a/Engine/Source/Developer/Windows/ShaderFormatD3D/Private/D3DShaderCompilerDXC.cpp
+++ b/Engine/Source/Developer/Windows/ShaderFormatD3D/Private/D3DShaderCompilerDXC.cpp
@@ -155,7 +155,7 @@ static void LogFailedHRESULT(const TCHAR* FailedExpressionStr, HRESULT Result)
 
 class FDxcMalloc final : public IMalloc
 {
-	ULONG RefCount = 1;
+	std::atomic<ULONG> RefCount{ 1 };
 
 public:
 
@@ -219,16 +219,20 @@ static IMalloc* GetDxcMalloc()
 	return &Instance;
 }
 
+
 static dxc::DxcDllSupport& GetDxcDllHelper()
 {
-	static dxc::DxcDllSupport DxcDllSupport;
-	static bool DxcDllInitialized = false;
-	if (!DxcDllInitialized)
+	struct DxcDllHelper
 	{
-		VERIFYHRESULT(DxcDllSupport.Initialize());
-		DxcDllInitialized = true;
-	}
-	return DxcDllSupport;
+		DxcDllHelper()
+		{
+			VERIFYHRESULT(DxcDllSupport.Initialize());
+		}
+		dxc::DxcDllSupport DxcDllSupport;
+	};
+
+	static DxcDllHelper DllHelper;
+	return DllHelper.DxcDllSupport;
 }
 
 static FString DxcBlobEncodingToFString(TRefCountPtr<IDxcBlobEncoding> DxcBlob)
