Fixed LiveCodingConsole to handle the patching of coroutines.

diff --git a/Engine/Source/Developer/Windows/LiveCodingServer/Private/External/LC_RelocationPatcher.cpp b/Engine/Source/Developer/Windows/LiveCodingServer/Private/External/LC_RelocationPatcher.cpp
--- a/Engine/Source/Developer/Windows/LiveCodingServer/Private/External/LC_RelocationPatcher.cpp
+++ b/Engine/Source/Developer/Windows/LiveCodingServer/Private/External/LC_RelocationPatcher.cpp
@@ -35,6 +35,10 @@ bool relocations::WouldPatchRelocation(const ImmutableString& dstSymbolName)
 	{
 		return false;
 	}
+	else if (string::Contains(dstSymbolName.c_str(), "$__resumable"))
+	{
+		return false;
+	}
 
 	return true;
 }
@@ -193,6 +197,12 @@ relocations::Record relocations::PatchRelocation
 			LC_LOG_DEV("Ignoring relocation to %s", dstSymbolName.c_str());
 			return record;
 		}
+
+		if (string::Contains(dstSymbolName.c_str(), "$__resumable"))
+		{
+			LC_LOG_DEV("Ignoring relocation to %s", dstSymbolName.c_str());
+			return record;
+		}
 	}
 
 	// BEGIN EPIC MOD
diff --git a/Engine/Source/Developer/Windows/LiveCodingServer/Private/External/LC_SymbolReconstruction.cpp b/Engine/Source/Developer/Windows/LiveCodingServer/Private/External/LC_SymbolReconstruction.cpp
--- a/Engine/Source/Developer/Windows/LiveCodingServer/Private/External/LC_SymbolReconstruction.cpp
+++ b/Engine/Source/Developer/Windows/LiveCodingServer/Private/External/LC_SymbolReconstruction.cpp
@@ -125,6 +125,10 @@ void symbols::ReconstructFromExecutableCoff
 					// are being relocated to, but store absolute values encoded in their offset in the PDB.
 					LC_LOG_DEV("Compiler- or linker-generated symbol %s without an RVA", symbolName.c_str());
 				}
+				else if (string::Contains(symbolName.c_str(), "$__resumable"))
+				{
+					LC_LOG_DEV("Coroutine symbol %s", symbolName.c_str());
+				}
 				else
 				{
 					LC_LOG_DEV("Unknown unresolved symbol %s", symbolName.c_str());
@@ -253,6 +257,11 @@ walkOpenSymbols:
 				LC_LOG_DEV("Ignoring section symbol \"%s\"", dstSymbolName.c_str());
 				continue;
 			}
+			else if (string::Contains(dstSymbolName.c_str(), "$__resumable"))
+			{
+				LC_LOG_DEV("Ignoring Coroutine symbol %s", dstSymbolName.c_str());
+				continue;
+			}
 
 			if (!srcSymbol)
 			{
