Report SCW error code for E_OUTOFMEMORY in DXBC backend.

#rnx

diff --git a/Engine/Source/Developer/Windows/ShaderFormatD3D/Private/D3DShaderCompiler.cpp b/Engine/Source/Developer/Windows/ShaderFormatD3D/Private/D3DShaderCompiler.cpp
--- a/Engine/Source/Developer/Windows/ShaderFormatD3D/Private/D3DShaderCompiler.cpp
+++ b/Engine/Source/Developer/Windows/ShaderFormatD3D/Private/D3DShaderCompiler.cpp
@@ -620,6 +620,13 @@ static bool CompileAndProcessD3DShaderFXCExt(
 		// Some materials give FXC a hard time to optimize and the compiler fails with an internal error.
 		if (bPrecompileWithDXC || Result == HRESULT_FROM_WIN32(ERROR_ARITHMETIC_OVERFLOW) || Result == E_OUTOFMEMORY || Result == E_FAIL || (Result != S_OK && CompileErrorsContainInternalError(Errors.GetReference())))
 		{
+			// If we ran out of memory, it's likely the next attempt will crash, too.
+			// Report the error now in case CompileHlslToSpirv throws an exception.
+			if (Result == E_OUTOFMEMORY)
+			{
+				FSCWErrorCode::Report(FSCWErrorCode::OutOfMemory);
+			}
+
 			CrossCompiler::FShaderConductorContext CompilerContext;
 
 			// Load shader source into compiler context
@@ -689,6 +696,9 @@ static bool CompileAndProcessD3DShaderFXCExt(
 
 			if (!bPrecompileWithDXC && SUCCEEDED(Result))
 			{
+				// Reset our previously set error code
+				FSCWErrorCode::Reset();
+
 				// Let the user know this shader had to be cross-compiled due to a crash in FXC. Only shows up if CVar 'r.ShaderDevelopmentMode' is enabled.
 				Output.Errors.Add(FShaderCompilerError(FString::Printf(TEXT("Cross-compiled shader to intermediate HLSL after first attempt crashed FXC: %s"), *Input.GenerateShaderName())));
 			}
