Fixed issue in LiveCoding preload module support where if a module name is known but not currently loaded, the engine would generate a check.

diff --git a/Engine/Source/Developer/Windows/LiveCoding/Private/LiveCodingModule.cpp b/Engine/Source/Developer/Windows/LiveCoding/Private/LiveCodingModule.cpp
--- a/Engine/Source/Developer/Windows/LiveCoding/Private/LiveCodingModule.cpp
+++ b/Engine/Source/Developer/Windows/LiveCoding/Private/LiveCodingModule.cpp
@@ -1192,10 +1192,13 @@ void FLiveCodingModule::UpdateModules(bool bAllowStarting)
 				Manager.FindModules(*ModuleName.ToString(), OutModules);
 				for (FName OutModule : OutModules)
 				{
-					FString FileName = Manager.GetModuleFilename(OutModule);
-					if (!FileName.IsEmpty())
+					FModuleStatus ModuleStatus;
+					if (Manager.QueryModule(OutModule, ModuleStatus))
 					{
-						PreloadedFileNames.Add(FName(FPaths::GetBaseFilename(FileName, true)));
+						if (!ModuleStatus.FilePath.IsEmpty())
+						{
+							PreloadedFileNames.Add(FName(FPaths::GetBaseFilename(ModuleStatus.FilePath, true)));
+						}
 					}
 				}
 			}
