Moving some headers to an Internal path.
Touching the build.cs to invalidate the old paths.

diff --git a/Engine/Source/Developer/Windows/WindowsTargetPlatform/Internal/GenericWindowsTargetPlatform.h b/Engine/Source/Developer/Windows/WindowsTargetPlatform/Internal/GenericWindowsTargetPlatform.h
new file mode 100644
--- /dev/null
+++ b/Engine/Source/Developer/Windows/WindowsTargetPlatform/Internal/GenericWindowsTargetPlatform.h
@@ -0,0 +1,410 @@
+// Copyright Epic Games, Inc. All Rights Reserved.
+
+#pragma once
+
+#include "CoreMinimal.h"
+#include "Common/TargetPlatformBase.h"
+#include "Misc/ConfigCacheIni.h"
+#include "AnalyticsEventAttribute.h"
+#if PLATFORM_WINDOWS
+#include "LocalPcTargetDevice.h"
+#endif
+#include "Serialization/MemoryLayout.h"
+#include "SteamDeck/SteamDeckDevice.h"
+
+#if WITH_ENGINE
+	#include "Sound/SoundWave.h"
+	#include "TextureResource.h"
+	#include "Engine/VolumeTexture.h"
+	#include "StaticMeshResources.h"
+	#include "RHI.h"
+	#include "DataDrivenShaderPlatformInfo.h"
+	#include "AudioCompressionSettings.h"
+#endif // WITH_ENGINE
+#include "Windows/WindowsPlatformProperties.h"
+#include "HAL/IConsoleManager.h"
+
+#define LOCTEXT_NAMESPACE "TGenericWindowsTargetPlatform"
+
+/**
+ * Template for Windows target platforms
+ */
+#if PLATFORM_WINDOWS
+template<typename TProperties, typename TTargetDevice = TLocalPcTargetDevice<PLATFORM_64BITS>>
+#else
+template<typename TProperties>
+#endif
+
+class TGenericWindowsTargetPlatform : public TTargetPlatformBase<TProperties>
+{
+public:
+	typedef TTargetPlatformBase<TProperties> TSuper;
+
+	/**
+	 * Default constructor.
+	 */
+	TGenericWindowsTargetPlatform( )
+	{
+#if PLATFORM_WINDOWS
+		// only add local device if actually running on Windows
+		LocalDevice = MakeShareable(new TTargetDevice(*this));
+
+		// quick solution to not having WinGDK steamdeck devices
+		if (this->PlatformName().StartsWith(TEXT("Windows")))
+		{
+			// Check if we have any SteamDeck devices around
+			SteamDevices = TSteamDeckDevice<TLocalPcTargetDevice<true>>::DiscoverDevices(*this, TEXT("Proton"));
+		}
+#endif
+
+	#if WITH_ENGINE
+		TextureLODSettings = nullptr; // These are registered by the device profile system.
+		StaticMeshLODSettings.Initialize(this);
+
+
+		// Get the Target RHIs for this platform, we do not always want all those that are supported.
+		TArray<FName> TargetedShaderFormats;
+		TGenericWindowsTargetPlatform::GetAllTargetedShaderFormats(TargetedShaderFormats);
+
+		static FName NAME_PCD3D_SM6(TEXT("PCD3D_SM6"));
+		static FName NAME_PCD3D_SM5(TEXT("PCD3D_SM5"));
+		static FName NAME_VULKAN_SM5(TEXT("SF_VULKAN_SM5"));
+		static FName NAME_PCD3D_ES3_1(TEXT("PCD3D_ES31"));
+		static FName NAME_SF_VULKAN_ES31(TEXT("SF_VULKAN_ES31"));
+		static FName NAME_OPENGL_150_ES3_1(TEXT("GLSL_150_ES31"));
+
+		// If we are targeting ES3.1, we also must cook encoded HDR reflection captures
+		bRequiresEncodedHDRReflectionCaptures =	TargetedShaderFormats.Contains(NAME_SF_VULKAN_ES31)
+												|| TargetedShaderFormats.Contains(NAME_OPENGL_150_ES3_1)
+												|| TargetedShaderFormats.Contains(NAME_PCD3D_ES3_1);
+
+	#endif
+	}
+
+public:
+
+	//~ Begin ITargetPlatform Interface
+
+	virtual void EnableDeviceCheck(bool OnOff) override {}
+
+	virtual void GetAllDevices( TArray<ITargetDevicePtr>& OutDevices ) const override
+	{
+		OutDevices.Reset();
+		if (LocalDevice.IsValid())
+		{
+			OutDevices.Add(LocalDevice);
+		}
+
+		for (const ITargetDevicePtr& SteamDeck : SteamDevices)
+		{
+			if (SteamDeck.IsValid())
+			{
+				OutDevices.Add(SteamDeck);
+			}
+		}
+	}
+
+	virtual bool GenerateStreamingInstallManifest(const TMultiMap<FString, int32>& PakchunkMap, const TSet<int32>& PakchunkIndicesInUse) const override
+	{
+		return true;
+	}
+
+	virtual ITargetDevicePtr GetDefaultDevice( ) const override
+	{
+		if (LocalDevice.IsValid())
+		{
+			return LocalDevice;
+		}
+
+		return nullptr;
+	}
+
+	virtual ITargetDevicePtr GetDevice( const FTargetDeviceId& DeviceId )
+	{
+		if (LocalDevice.IsValid() && (DeviceId == LocalDevice->GetId()))
+		{
+			return LocalDevice;
+		}
+
+		for (const ITargetDevicePtr& SteamDeck : SteamDevices)
+		{
+			if (SteamDeck.IsValid() && DeviceId == SteamDeck->GetId())
+			{
+				return SteamDeck;
+			}
+		}
+
+		return nullptr;
+	}
+
+	virtual bool IsRunningPlatform( ) const override
+	{
+		// Must be Windows platform as editor for this to be considered a running platform
+		return PLATFORM_WINDOWS && !UE_SERVER && !UE_GAME && WITH_EDITOR && TProperties::HasEditorOnlyData();
+	}
+
+	virtual void GetShaderCompilerDependencies(TArray<FString>& OutDependencies) const override
+	{		
+		FTargetPlatformBase::AddDependencySCArrayHelper(OutDependencies, TEXT("Binaries/ThirdParty/Windows/DirectX/x64/d3dcompiler_47.dll"));
+		FTargetPlatformBase::AddDependencySCArrayHelper(OutDependencies, TEXT("Binaries/ThirdParty/ShaderConductor/Win64/ShaderConductor.dll"));
+		FTargetPlatformBase::AddDependencySCArrayHelper(OutDependencies, TEXT("Binaries/ThirdParty/ShaderConductor/Win64/dxcompiler.dll"));
+		FTargetPlatformBase::AddDependencySCArrayHelper(OutDependencies, TEXT("Binaries/ThirdParty/ShaderConductor/Win64/dxil.dll"));
+	}
+
+	virtual bool SupportsFeature( ETargetPlatformFeatures Feature ) const override
+	{
+		// we currently do not have a build target for WindowsServer
+		if (Feature == ETargetPlatformFeatures::Packaging)
+		{
+			return (TProperties::HasEditorOnlyData() || !TProperties::IsServerOnly());
+		}
+
+		if ( Feature == ETargetPlatformFeatures::ShouldSplitPaksIntoSmallerSizes )
+		{
+			return TProperties::IsClientOnly();
+		}
+
+		if (Feature == ETargetPlatformFeatures::MobileRendering)
+		{
+			static bool bCachedSupportsMobileRendering = false;
+#if WITH_ENGINE
+			static bool bHasCachedValue = false;
+			if (!bHasCachedValue)
+			{
+				TArray<FName> TargetedShaderFormats;
+				GetAllTargetedShaderFormats(TargetedShaderFormats);
+
+				for (const FName& Format : TargetedShaderFormats)
+				{
+					if (IsMobilePlatform(ShaderFormatToLegacyShaderPlatform(Format)))
+					{
+						bCachedSupportsMobileRendering = true;
+						break;
+					}
+				}
+				bHasCachedValue = true;
+			}
+#endif
+
+			return bCachedSupportsMobileRendering;
+		}
+
+		return TSuper::SupportsFeature(Feature);
+	}
+
+	virtual void GetBuildProjectSettingKeys(FString& OutSection, TArray<FString>& InBoolKeys, TArray<FString>& InIntKeys, TArray<FString>& InStringKeys) const override
+	{
+		OutSection = TEXT("/Script/WindowsTargetPlatform.WindowsTargetSettings");
+		InStringKeys.Add(TEXT("MinimumOSVersion"));
+	}
+
+	virtual void GetAllPossibleShaderFormats( TArray<FName>& OutFormats ) const override
+	{
+		// no shaders needed for dedicated server target
+		if (!TProperties::IsServerOnly())
+		{
+			static FName NAME_PCD3D_SM6(TEXT("PCD3D_SM6"));
+			static FName NAME_PCD3D_SM5(TEXT("PCD3D_SM5"));
+			static FName NAME_VULKAN_ES31(TEXT("SF_VULKAN_ES31"));
+			static FName NAME_OPENGL_150_ES3_1(TEXT("GLSL_150_ES31"));
+			static FName NAME_VULKAN_SM5(TEXT("SF_VULKAN_SM5"));
+			static FName NAME_PCD3D_ES3_1(TEXT("PCD3D_ES31"));
+
+			OutFormats.AddUnique(NAME_PCD3D_SM5);
+			OutFormats.AddUnique(NAME_PCD3D_SM6);
+			OutFormats.AddUnique(NAME_VULKAN_ES31);
+			OutFormats.AddUnique(NAME_OPENGL_150_ES3_1);
+			OutFormats.AddUnique(NAME_VULKAN_SM5);
+			OutFormats.AddUnique(NAME_PCD3D_ES3_1);
+		}
+	}
+
+	void GetAllTargetedShaderFormatsInternal(TArrayView<TCHAR const*> RelevantSettings, TArray<FName>& OutFormats) const
+	{
+		TArray<FString> TargetedShaderFormats;
+
+		for (const TCHAR* Name : RelevantSettings)
+		{
+			TArray<FString> NewTargetedShaderFormats;
+			GConfig->GetArray(TEXT("/Script/WindowsTargetPlatform.WindowsTargetSettings"), Name, NewTargetedShaderFormats, GEngineIni);
+
+			for (const FString& NewShaderFormat : NewTargetedShaderFormats)
+			{
+				TargetedShaderFormats.AddUnique(NewShaderFormat);
+			}
+		}
+		
+		// Gather the list of Target RHIs and filter out any that may be invalid.
+		TArray<FName> PossibleShaderFormats;
+		GetAllPossibleShaderFormats(PossibleShaderFormats);
+
+		for (int32 ShaderFormatIdx = TargetedShaderFormats.Num() - 1; ShaderFormatIdx >= 0; ShaderFormatIdx--)
+		{
+			FString ShaderFormat = TargetedShaderFormats[ShaderFormatIdx];
+			if (PossibleShaderFormats.Contains(FName(*ShaderFormat)) == false)
+			{
+				TargetedShaderFormats.RemoveAt(ShaderFormatIdx);
+			}
+		}
+
+		for(const FString& ShaderFormat : TargetedShaderFormats)
+		{
+			OutFormats.AddUnique(FName(*ShaderFormat));
+		}
+	}
+
+	virtual void GetAllTargetedShaderFormats( TArray<FName>& OutFormats ) const override 
+	{
+		// Get the Target RHIs for this platform, we do not always want all those that are supported. (reload in case user changed in the editor)
+
+		TCHAR const* RelevantSettings[] = 
+		{
+			TEXT("TargetedRHIs"),
+			TEXT("D3D12TargetedShaderFormats"),
+			TEXT("D3D11TargetedShaderFormats"),
+			TEXT("VulkanTargetedShaderFormats")
+		};
+
+		GetAllTargetedShaderFormatsInternal(RelevantSettings, OutFormats);
+	}
+
+	virtual void GetRayTracingShaderFormats( TArray<FName>& OutFormats ) const override 
+	{
+		if (UsesRayTracing())
+		{
+			TCHAR const* RelevantSettings[] = 
+			{
+				TEXT("VulkanTargetedShaderFormats")
+			};
+
+			GetAllTargetedShaderFormatsInternal(RelevantSettings, OutFormats);
+
+			// We always support ray tracing shaders when cooking for D3D12 SM6, however we may skip them for SM5 based on project settings.
+			OutFormats.AddUnique(FName(TEXT("PCD3D_SM6")));
+
+			static IConsoleVariable* RequireSM6CVar = IConsoleManager::Get().FindConsoleVariable(TEXT("r.RayTracing.RequireSM6"));
+			const bool bRequireSM6 = RequireSM6CVar && RequireSM6CVar->GetBool();
+			if (!bRequireSM6)
+			{
+				OutFormats.AddUnique(FName(TEXT("PCD3D_SM5")));
+			}
+		}
+	}
+
+	virtual void GetPlatformSpecificProjectAnalytics( TArray<FAnalyticsEventAttribute>& AnalyticsParamArray ) const override
+	{
+		TSuper::GetPlatformSpecificProjectAnalytics(AnalyticsParamArray);
+
+		AppendAnalyticsEventAttributeArray(AnalyticsParamArray,
+			TEXT("UsesRayTracing"), UsesRayTracing()
+		);
+	
+		TSuper::AppendAnalyticsEventConfigString(AnalyticsParamArray, TEXT("/Script/WindowsTargetPlatform.WindowsTargetSettings"), TEXT("DefaultGraphicsRHI"), GEngineIni);
+
+		TSuper::AppendAnalyticsEventConfigArray(AnalyticsParamArray, TEXT("/Script/WindowsTargetPlatform.WindowsTargetSettings"), TEXT("D3D12TargetedShaderFormats"), GEngineIni);
+		TSuper::AppendAnalyticsEventConfigArray(AnalyticsParamArray, TEXT("/Script/WindowsTargetPlatform.WindowsTargetSettings"), TEXT("D3D11TargetedShaderFormats"), GEngineIni);
+		TSuper::AppendAnalyticsEventConfigArray(AnalyticsParamArray, TEXT("/Script/WindowsTargetPlatform.WindowsTargetSettings"), TEXT("VulkanTargetedShaderFormats"), GEngineIni);
+		TSuper::AppendAnalyticsEventConfigArray(AnalyticsParamArray, TEXT("/Script/WindowsTargetPlatform.WindowsTargetSettings"), TEXT("TargetedRHIs"), GEngineIni, TEXT("TargetedRHIs_Deprecated") );
+	}
+
+#if WITH_ENGINE
+	virtual void GetReflectionCaptureFormats(TArray<FName>& OutFormats) const override
+	{
+		if (bRequiresEncodedHDRReflectionCaptures)
+		{
+			OutFormats.Add(FName(TEXT("EncodedHDR")));
+		}
+
+		OutFormats.Add(FName(TEXT("FullHDR")));
+	}
+
+	virtual void GetShaderFormatModuleHints(TArray<FName>& OutModuleNames) const override
+	{
+		OutModuleNames.Add(TEXT("ShaderFormatD3D"));
+		OutModuleNames.Add(TEXT("ShaderFormatOpenGL"));
+		OutModuleNames.Add(TEXT("VulkanShaderFormat"));
+	}
+
+	virtual const class FStaticMeshLODSettings& GetStaticMeshLODSettings( ) const override
+	{
+		return StaticMeshLODSettings;
+	}
+
+	virtual void GetTextureFormats( const UTexture* InTexture, TArray< TArray<FName> >& OutFormats) const override
+	{
+		if (!TProperties::IsServerOnly())
+		{
+			GetDefaultTextureFormatNamePerLayer(OutFormats.AddDefaulted_GetRef(), this, InTexture, true, 4, true);
+		}
+	}
+
+	virtual void GetAllTextureFormats(TArray<FName>& OutFormats) const override
+	{
+		if (!TProperties::IsServerOnly())
+		{
+			GetAllDefaultTextureFormats(this, OutFormats);
+		}
+	}
+
+	virtual const UTextureLODSettings& GetTextureLODSettings() const override
+	{
+		return *TextureLODSettings;
+	}
+
+	virtual void RegisterTextureLODSettings(const UTextureLODSettings* InTextureLODSettings) override
+	{
+		TextureLODSettings = InTextureLODSettings;
+	}
+
+#endif //WITH_ENGINE
+
+	virtual bool SupportsVariants() const override
+	{
+		return true;
+	}
+
+	virtual float GetVariantPriority() const override
+	{
+		return TProperties::GetVariantPriority();
+	}
+
+	virtual bool UsesDistanceFields() const override
+	{
+		bool bEnableDistanceFields = false;
+		GConfig->GetBool(TEXT("/Script/WindowsTargetPlatform.WindowsTargetSettings"), TEXT("bEnableDistanceFields"), bEnableDistanceFields, GEngineIni);
+
+		return bEnableDistanceFields && TSuper::UsesDistanceFields();
+	}
+
+	virtual bool UsesRayTracing() const override
+	{
+		bool bEnableRayTracing = false;
+		GConfig->GetBool(TEXT("/Script/WindowsTargetPlatform.WindowsTargetSettings"), TEXT("bEnableRayTracing"), bEnableRayTracing, GEngineIni);
+
+		return bEnableRayTracing && TSuper::UsesRayTracing();
+	}
+
+	//~ End ITargetPlatform Interface
+
+private:
+
+	// Holds the local device.
+	ITargetDevicePtr LocalDevice;
+
+	TArray<ITargetDevicePtr> SteamDevices;
+
+#if WITH_ENGINE
+	// Holds the texture LOD settings.
+	const UTextureLODSettings* TextureLODSettings;
+
+	// Holds static mesh LOD settings.
+	FStaticMeshLODSettings StaticMeshLODSettings;
+
+	// True if the project requires encoded HDR reflection captures
+	bool bRequiresEncodedHDRReflectionCaptures;
+
+#endif // WITH_ENGINE
+
+};
+
+#undef LOCTEXT_NAMESPACE
diff --git a/Engine/Source/Developer/Windows/WindowsTargetPlatform/Internal/LocalPcTargetDevice.h b/Engine/Source/Developer/Windows/WindowsTargetPlatform/Internal/LocalPcTargetDevice.h
new file mode 100644
--- /dev/null
+++ b/Engine/Source/Developer/Windows/WindowsTargetPlatform/Internal/LocalPcTargetDevice.h
@@ -0,0 +1,334 @@
+// Copyright Epic Games, Inc. All Rights Reserved.
+
+#pragma once
+
+#include "CoreMinimal.h"
+#include "Misc/Guid.h"
+#include "HAL/FileManager.h"
+#include "HAL/PlatformProcess.h"
+#include "Interfaces/ITargetDevice.h"
+#include "Interfaces/ITargetPlatform.h"
+
+#include "Windows/WindowsHWrapper.h"
+#include "Windows/AllowWindowsPlatformTypes.h"
+#include <TlHelp32.h>
+
+/**
+ * Template for local PC target devices.
+ *
+ * @param WIN64 - Whether the target platform is 64-bit Windows.
+ */
+template<bool WIN64>
+class TLocalPcTargetDevice
+	: public ITargetDevice
+{
+public:
+
+	/**
+	 * Creates and initializes a new device for the specified target platform.
+	 *
+	 * @param InTargetPlatform - The target platform.
+	 */
+	TLocalPcTargetDevice( const ITargetPlatform& InTargetPlatform )
+		: TargetPlatform(InTargetPlatform)
+	{ }
+
+public:
+
+	// ITargetDevice interface
+
+	virtual bool Connect( ) override
+	{
+		return true;
+	}
+
+	virtual void Disconnect( )
+	{ }
+
+	virtual ETargetDeviceTypes GetDeviceType( ) const override
+	{
+		if (::GetSystemMetrics(SM_TABLETPC) != 0)
+		{
+			return ETargetDeviceTypes::Tablet;
+		}
+
+		return ETargetDeviceTypes::Desktop;
+	}
+
+	virtual FTargetDeviceId GetId( ) const override
+	{
+		return FTargetDeviceId(TargetPlatform.PlatformName(), GetName());
+	}
+
+	virtual FString GetName( ) const override
+	{
+		return FPlatformProcess::ComputerName();
+	}
+
+	virtual FString GetOperatingSystemName( ) override
+	{
+		if (WIN64)
+		{
+			return TEXT("Windows (64-bit)");
+		}
+
+		return TEXT("Windows (32-bit)");
+	}
+
+	virtual int32 GetProcessSnapshot( TArray<FTargetDeviceProcessInfo>& OutProcessInfos ) override
+	{
+		// enumerate processes
+		HANDLE ProcessSnapshot = ::CreateToolhelp32Snapshot(TH32CS_SNAPPROCESS, 0);
+
+		if (ProcessSnapshot != INVALID_HANDLE_VALUE)
+		{
+			PROCESSENTRY32 ProcessEntry;
+			ProcessEntry.dwSize = sizeof(PROCESSENTRY32);
+
+			if (::Process32First(ProcessSnapshot, &ProcessEntry))
+			{
+				do 
+				{
+					// only include processes that the user has permission to query
+					HANDLE ProcessHandle = ::OpenProcess(PROCESS_QUERY_INFORMATION, false, ProcessEntry.th32ProcessID);
+
+					if (ProcessHandle != NULL)
+					{
+						FTargetDeviceProcessInfo ProcessInfo;
+
+						// get process details
+						ProcessInfo.Id = ProcessEntry.th32ProcessID;
+						ProcessInfo.Name = ProcessEntry.szExeFile;
+						ProcessInfo.ParentId = ProcessEntry.th32ParentProcessID;
+						ProcessInfo.UserName = TEXT("-");
+
+						// get process user name
+						HANDLE TokenHandle;
+
+						if (::OpenProcessToken(ProcessHandle, TOKEN_QUERY, &TokenHandle))
+						{
+							DWORD UserTokenSize;
+
+							::GetTokenInformation(TokenHandle, TokenUser, NULL, 0, &UserTokenSize);
+
+							if (::GetLastError() == ERROR_INSUFFICIENT_BUFFER)
+							{
+								PTOKEN_USER UserToken = reinterpret_cast<PTOKEN_USER>(new BYTE[UserTokenSize]);
+
+								if (UserToken != NULL)
+								{
+									if (::GetTokenInformation(TokenHandle, TokenUser, UserToken, UserTokenSize, &UserTokenSize))
+									{
+										WCHAR DomainName[256];
+										DWORD DomainNameLength = 256;
+
+										WCHAR UserName[256];
+										DWORD UserNameLength = 256;
+
+										SID_NAME_USE SidType;
+
+										if (::LookupAccountSid(NULL, UserToken->User.Sid, UserName, &UserNameLength, DomainName, &DomainNameLength, &SidType))
+										{
+											ProcessInfo.UserName = UserName;
+										}
+										else
+										{
+											ProcessInfo.UserName = TEXT("SYSTEM");
+										}
+									}
+								}
+								else
+								{
+									DWORD LastError = ::GetLastError();
+								}
+							}
+
+							::CloseHandle(TokenHandle);
+						}
+
+						::CloseHandle(ProcessHandle);
+
+						OutProcessInfos.Add(ProcessInfo);
+					}
+				}
+				while (::Process32Next(ProcessSnapshot, &ProcessEntry));
+			}
+
+			::CloseHandle(ProcessSnapshot);
+		}
+
+		// get thread details
+		HANDLE ThreadSnapshot = ::CreateToolhelp32Snapshot(TH32CS_SNAPTHREAD, 0);
+
+		if (ThreadSnapshot != INVALID_HANDLE_VALUE)
+		{
+			THREADENTRY32 ThreadEntry;
+			ThreadEntry.dwSize = sizeof(THREADENTRY32);
+
+			if (::Thread32First(ThreadSnapshot, &ThreadEntry))
+			{
+				do 
+				{
+					for (int32 ProcessInfoIndex = 0; ProcessInfoIndex < OutProcessInfos.Num(); ++ProcessInfoIndex)
+					{
+						FTargetDeviceProcessInfo& ProcessInfo = OutProcessInfos[ProcessInfoIndex];
+
+						if (ProcessInfo.Id == ThreadEntry.th32OwnerProcessID)
+						{
+							FTargetDeviceThreadInfo ThreadInfo;
+
+							ThreadInfo.ExitCode = 0;
+							ThreadInfo.Id = ThreadEntry.th32ThreadID;
+							ThreadInfo.StackSize = 0;
+							ThreadInfo.State = ETargetDeviceThreadStates::Unknown;
+							ThreadInfo.WaitState = ETargetDeviceThreadWaitStates::Unknown;
+
+							ProcessInfo.Threads.Add(ThreadInfo);
+
+							break;
+						}
+					}
+				}
+				while (::Thread32Next(ThreadSnapshot, &ThreadEntry));
+
+				::CloseHandle(ThreadSnapshot);
+			}
+		}
+
+		return OutProcessInfos.Num();
+	}
+
+	virtual const class ITargetPlatform& GetTargetPlatform( ) const override
+	{
+		return TargetPlatform;
+	}
+
+	virtual bool IsConnected( )
+	{
+		return true;
+	}
+
+	virtual bool IsDefault( ) const override
+	{
+		return true;
+	}
+
+	virtual bool PowerOff( bool Force ) override
+	{
+		if (!AdjustShutdownPrivileges())
+		{
+			return false;
+		}
+
+		CA_SUPPRESS(28159);
+		return (::ExitWindowsEx(EWX_POWEROFF | EWX_FORCE, SHTDN_REASON_MINOR_MAINTENANCE | SHTDN_REASON_FLAG_PLANNED) != 0);
+	}
+
+	virtual bool PowerOn( ) override
+	{
+		return false;
+	}
+
+	virtual bool Reboot( bool bReconnect = false ) override
+	{
+		if (!AdjustShutdownPrivileges())
+		{
+			return false;
+		}
+
+		CA_SUPPRESS(28159);
+		return (::ExitWindowsEx(EWX_REBOOT | EWX_FORCE, SHTDN_REASON_MINOR_MAINTENANCE | SHTDN_REASON_FLAG_PLANNED) != 0);
+	}
+
+	virtual bool SupportsFeature( ETargetDeviceFeatures Feature ) const override
+	{
+		switch (Feature)
+		{
+		case ETargetDeviceFeatures::MultiLaunch:
+			return true;
+
+		case ETargetDeviceFeatures::PowerOff:
+			return true;
+
+		// @todo gmp: implement turning on remote PCs (wake on LAN)
+		case ETargetDeviceFeatures::PowerOn:
+			return false;
+
+		case ETargetDeviceFeatures::ProcessSnapshot:
+			return true;
+
+		case ETargetDeviceFeatures::Reboot:
+			return true;
+		}
+
+		return false;
+	}
+
+	virtual void SetUserCredentials( const FString& UserName, const FString& UserPassword ) override
+	{
+	}
+
+	virtual bool GetUserCredentials( FString& OutUserName, FString& OutUserPassword ) override
+	{
+		return false;
+	}
+
+	virtual bool TerminateProcess( const int64 ProcessId ) override
+	{
+		HANDLE ProcessHandle = OpenProcess(PROCESS_TERMINATE, false, (DWORD)ProcessId );
+
+		if (ProcessHandle != NULL)
+		{
+			bool Result = (::TerminateProcess(ProcessHandle, 0) != 0);
+
+			CloseHandle(ProcessHandle);
+
+			return Result;
+		}
+
+		return false;
+	}
+
+protected:
+
+	/**
+	 * Adjusts shutdown privileges for the local host PC.
+	 *
+	 * @return true on success, false otherwise.
+	 */
+	bool AdjustShutdownPrivileges( )
+	{
+		HANDLE TokenHandle;
+		TOKEN_PRIVILEGES Token;
+
+		if (!::OpenProcessToken(GetCurrentProcess(), TOKEN_ADJUST_PRIVILEGES | TOKEN_QUERY, &TokenHandle))
+		{
+			return false;
+		}
+
+		::LookupPrivilegeValue(NULL, SE_SHUTDOWN_NAME, &Token.Privileges[0].Luid);
+
+		if (Token.Privileges[0].Attributes == SE_PRIVILEGE_ENABLED)
+		{
+			return true;
+		}
+
+		Token.PrivilegeCount = 1;
+		Token.Privileges[0].Attributes = SE_PRIVILEGE_ENABLED;
+
+		::AdjustTokenPrivileges(TokenHandle, 0, &Token, 0, (PTOKEN_PRIVILEGES)NULL, 0);
+      
+		return (GetLastError() == ERROR_SUCCESS);
+	}
+
+protected:
+
+	// Holds the collection of processes that were started using the Run() method.
+	TMap<FGuid, void*> Processes;
+
+	// Holds a reference to the device's target platform.
+	const ITargetPlatform& TargetPlatform;
+};
+
+#include "Windows/HideWindowsPlatformTypes.h"
+#undef PROCESSENTRY32
diff --git a/Engine/Source/Developer/Windows/WindowsTargetPlatform/Private/GenericWindowsTargetPlatform.h b/Engine/Source/Developer/Windows/WindowsTargetPlatform/Private/GenericWindowsTargetPlatform.h
deleted file mode 100644
diff --git a/Engine/Source/Developer/Windows/WindowsTargetPlatform/Private/LocalPcTargetDevice.h b/Engine/Source/Developer/Windows/WindowsTargetPlatform/Private/LocalPcTargetDevice.h
deleted file mode 100644
diff --git a/Engine/Source/Developer/Windows/WindowsTargetPlatform/WindowsTargetPlatform.Build.cs b/Engine/Source/Developer/Windows/WindowsTargetPlatform/WindowsTargetPlatform.Build.cs
--- a/Engine/Source/Developer/Windows/WindowsTargetPlatform/WindowsTargetPlatform.Build.cs
+++ b/Engine/Source/Developer/Windows/WindowsTargetPlatform/WindowsTargetPlatform.Build.cs
@@ -22,11 +22,6 @@ public class WindowsTargetPlatform : ModuleRules
 			}
 		);
 
-		PrivateIncludePaths.AddRange(
-			new string[] {
-			}
-		);
-
 		// compile with Engine
 		if (Target.bCompileAgainstEngine)
 		{
