LiveCoding no longer generates errors about object files being missing after a source file is removed from a unity file by live coding.

diff --git a/Engine/Source/Developer/Windows/LiveCodingServer/Private/External/LC_ServerCommandThread.cpp b/Engine/Source/Developer/Windows/LiveCodingServer/Private/External/LC_ServerCommandThread.cpp
--- a/Engine/Source/Developer/Windows/LiveCodingServer/Private/External/LC_ServerCommandThread.cpp
+++ b/Engine/Source/Developer/Windows/LiveCodingServer/Private/External/LC_ServerCommandThread.cpp
@@ -963,46 +963,28 @@ void ServerCommandThread::CompileChanges(bool didAllProcessesMakeProgress, comma
 		: LiveModule::UpdateType::NO_CLIENT_COMMUNICATION;
 
 	// BEGIN EPIC MOD
+	// For Epic, don't even bother calling update if we have nothing that has been reported as modified/new.
+	// This prevents the automatic change detection code from running which has issues with unity file changes.
 	static const types::vector<std::wstring> emptyLibFiles;
-	// END EPIC MOD
-
-	// has the user given us at least one modified or new .obj file for at least one of the modules?
-	const bool hasAtLeastOneOptionalObj = (m_liveModuleToModifiedOrNewObjFiles.size() != 0u);
 	for (size_t i = 0u; i < count; ++i)
 	{
 		LiveModule* liveModule = m_liveModules[i];
 
-		LiveModule::ErrorType::Enum moduleUpdateError = LiveModule::ErrorType::SUCCESS;
-		if (hasAtLeastOneOptionalObj)
-		{
-			// try to find the list of modified or new .objs for this module
-			const auto objFilesIt = m_liveModuleToModifiedOrNewObjFiles.find(liveModule->GetModuleName());
-			if (objFilesIt == m_liveModuleToModifiedOrNewObjFiles.end())
-			{
-				// no .objs for this module, ignore
-				continue;
-			}
-			else
-			{
-				// BEGIN EPIC MOD
-				const auto libFilesIt = m_liveModuleToAdditionalLibraries.find(liveModule->GetModuleName());
-				const types::vector<std::wstring>& libFiles = libFilesIt != m_liveModuleToAdditionalLibraries.end() ? libFilesIt->second : emptyLibFiles;
-
-				// build a patch with the given list of .objs for this module
-				const types::vector<symbols::ModifiedObjFile>& objFiles = objFilesIt->second;
-				moduleUpdateError = liveModule->Update(&fileCache, m_directoryCache, updateType, objFiles, libFiles);
-				// END EPIC MOD
-			}
-		}
-		else
+		// try to find the list of modified or new .objs for this module
+		const auto objFilesIt = m_liveModuleToModifiedOrNewObjFiles.find(liveModule->GetModuleName());
+		if (objFilesIt == m_liveModuleToModifiedOrNewObjFiles.end() || objFilesIt->second.empty())
 		{
-			// no optional .objs were given, update all live modules regularly
-			types::vector<symbols::ModifiedObjFile> emptyObjs;
-			// BEGIN EPIC MOD
-			moduleUpdateError = liveModule->Update(&fileCache, m_directoryCache, updateType, emptyObjs, emptyLibFiles);
-			// END EPIC MOD
+			// no .objs for this module, ignore
+			continue;
 		}
 
+		const auto libFilesIt = m_liveModuleToAdditionalLibraries.find(liveModule->GetModuleName());
+		const types::vector<std::wstring>& libFiles = libFilesIt != m_liveModuleToAdditionalLibraries.end() ? libFilesIt->second : emptyLibFiles;
+
+		// build a patch with the given list of .objs for this module
+		const types::vector<symbols::ModifiedObjFile>& objFiles = objFilesIt->second;
+		LiveModule::ErrorType::Enum moduleUpdateError = liveModule->Update(&fileCache, m_directoryCache, updateType, objFiles, libFiles);
+
 		// only accept new error conditions for this module if there haven't been any updates until now.
 		// this ensures that error conditions are kept and can be shown when updating several modules at once.
 		if (updateError == LiveModule::ErrorType::NO_CHANGE)
@@ -1010,6 +992,7 @@ void ServerCommandThread::CompileChanges(bool didAllProcessesMakeProgress, comma
 			updateError = moduleUpdateError;
 		}
 	}
+	// END EPIC MOD
 
 	// restart directory notifications for next compilation
 	m_directoryCache->RestartNotifications();
diff --git a/Engine/Source/Programs/LiveCodingConsole/Private/LiveCodingConsole.cpp b/Engine/Source/Programs/LiveCodingConsole/Private/LiveCodingConsole.cpp
--- a/Engine/Source/Programs/LiveCodingConsole/Private/LiveCodingConsole.cpp
+++ b/Engine/Source/Programs/LiveCodingConsole/Private/LiveCodingConsole.cpp
@@ -398,6 +398,11 @@ private:
 				ModuleToModuleFiles.FindOrAdd(Pair.Key).Libraries.Add(Library);
 			}
 		}
+
+		// Force the next patch time to include any CPP files modified this time around.  This avoids the issue 
+		// where a patch might be generated for a file twice instead of just once.
+		NextPatchStartTime = FDateTime::UtcNow();
+
 		return ELiveCodingCompileResult::Success;
 	}
 
