Fixed issue where LiveCoding would fail with "unresolved external symbols" when used with Clang.

diff --git a/Engine/Source/Developer/Windows/LiveCodingServer/Private/External/LC_LiveModule.cpp b/Engine/Source/Developer/Windows/LiveCodingServer/Private/External/LC_LiveModule.cpp
--- a/Engine/Source/Developer/Windows/LiveCodingServer/Private/External/LC_LiveModule.cpp
+++ b/Engine/Source/Developer/Windows/LiveCodingServer/Private/External/LC_LiveModule.cpp
@@ -2377,7 +2377,13 @@ LiveModule::ErrorType::Enum LiveModule::Update(FileAttributeCache* fileCache, Di
 					else if (string::Contains(upperCaseDirective.c_str(), "INCLUDE:"))
 					{
 						const std::size_t colonPos = directive.find(':');
-						const std::string symbolName(directive.c_str() + colonPos + 1u, directive.c_str() + directive.length());
+						// BEGIN EPIC MOD - This fix is already in Live++ 2
+						std::string symbolName(directive.c_str() + colonPos + 1u, directive.c_str() + directive.length());
+						if (symbolName.length() >= 2 && symbolName.front() == '\"' && symbolName.back() == '\"')
+						{
+							symbolName = symbolName.substr(1, symbolName.length() - 2);
+						}
+						// END EPIC MOD
 
 						const ModuleCache::FindSymbolData findData = m_moduleCache->FindSymbolByName(ModuleCache::SEARCH_ALL_MODULES, ImmutableString(symbolName.c_str()));
 						if (findData.symbol)
