Disabling the ignored-attributes warning when compiling bindless resources. DXC is incorrect that the attribute is ignored, see github issue: [REDACTED_PROFILE_URL]/DirectXShaderCompiler/issues/4537

diff --git a/Engine/Source/Developer/Windows/ShaderFormatD3D/Private/D3DShaderCompiler.inl b/Engine/Source/Developer/Windows/ShaderFormatD3D/Private/D3DShaderCompiler.inl
--- a/Engine/Source/Developer/Windows/ShaderFormatD3D/Private/D3DShaderCompiler.inl
+++ b/Engine/Source/Developer/Windows/ShaderFormatD3D/Private/D3DShaderCompiler.inl
@@ -31,7 +31,8 @@ public:
 		uint32 D3DCompileFlags,
 		uint32 AutoBindingSpace,
 		const TCHAR* InOptValidatorVersion,
-		uint32 HlslVersion = 2018
+		uint32 HlslVersion,
+		bool bBindlessResources
 	)
 		: ShaderProfile(InShaderProfile)
 		, EntryPoint(InEntryPoint)
@@ -196,6 +197,13 @@ public:
 		// disable undesired warnings
 		ExtraArguments.Add(TEXT("-Wno-parentheses-equality"));
 
+		// working around bindless conversion specific issue where globallycoherent on a function return type is flagged as ignored even though it is necessary.
+		// github issue: [REDACTED_PROFILE_URL]/DirectXShaderCompiler/issues/4537
+		if (bBindlessResources)
+		{
+			ExtraArguments.Add(TEXT("-Wno-ignored-attributes"));
+		}
+
 		// @lh-todo: This fixes a loop unrolling issue that showed up in DOFGatherKernel with cs_6_6 with the latest DXC revision
 		ExtraArguments.Add(TEXT("-disable-lifetime-markers"));
 	}
diff --git a/Engine/Source/Developer/Windows/ShaderFormatD3D/Private/D3DShaderCompilerDXC.cpp b/Engine/Source/Developer/Windows/ShaderFormatD3D/Private/D3DShaderCompilerDXC.cpp
--- a/Engine/Source/Developer/Windows/ShaderFormatD3D/Private/D3DShaderCompilerDXC.cpp
+++ b/Engine/Source/Developer/Windows/ShaderFormatD3D/Private/D3DShaderCompilerDXC.cpp
@@ -785,6 +785,7 @@ bool CompileAndProcessD3DShaderDXC(FString& PreprocessedShaderSource,
 	const bool bSymbolsBasedOnSource = Input.Environment.CompilerFlags.Contains(CFLAG_AllowUniqueSymbols);
 	const bool bHlslVersion2021 = Input.Environment.CompilerFlags.Contains(CFLAG_HLSL2021);
 	const uint32 HlslVersion = (bHlslVersion2021 ? 2021 : 2018);
+	const bool bBindlessResources = Input.Environment.CompilerFlags.Contains(CFLAG_BindlessResources);
 
 	FDxcArguments Args
 	(
@@ -799,7 +800,8 @@ bool CompileAndProcessD3DShaderDXC(FString& PreprocessedShaderSource,
 		DXCFlags,
 		AutoBindingSpace,
 		ValidatorVersion,
-		HlslVersion
+		HlslVersion,
+		bBindlessResources
 	);
 
 	if (bDumpDebugInfo)
