Always use SM6 when compiling with DXC, even when FeatureLevelES31 is requested.

#rnx

diff --git a/Engine/Source/Developer/Windows/ShaderFormatD3D/Private/D3DShaderCompiler.cpp b/Engine/Source/Developer/Windows/ShaderFormatD3D/Private/D3DShaderCompiler.cpp
--- a/Engine/Source/Developer/Windows/ShaderFormatD3D/Private/D3DShaderCompiler.cpp
+++ b/Engine/Source/Developer/Windows/ShaderFormatD3D/Private/D3DShaderCompiler.cpp
@@ -144,13 +144,13 @@ static const TCHAR* GetShaderProfileName(ELanguage Language, uint32 Frequency, b
 		switch(Frequency)
 		{
 		case SF_Pixel:
-			return TEXT("ps_5_0");
+			return bForceSM6 ? TEXT("ps_6_0") : TEXT("ps_5_0");
 		case SF_Vertex:
-			return TEXT("vs_5_0");
+			return bForceSM6 ? TEXT("vs_6_0") : TEXT("vs_5_0");
 		case SF_Geometry:
-			return TEXT("gs_5_0");
+			return bForceSM6 ? TEXT("gs_6_0") : TEXT("gs_5_0");
 		case SF_Compute:
-			return TEXT("cs_5_0");
+			return bForceSM6 ? TEXT("cs_6_0") : TEXT("cs_5_0");
 		}
 	}
 
