Remove DXC rewriter invocation from D3D12 backend.

This was originally added to have parity between Vulkan and D3D12, but it only emitted a debug file and wasn't even feeding the rewritten HLSL back into the D3D compiler.
In an effort to be independent of the DXC rewriter, this is a first step as this rewriter invocation served solely a debugging purpose anyway.

diff --git a/Engine/Source/Developer/Windows/ShaderFormatD3D/Private/D3DShaderCompilerDXC.cpp b/Engine/Source/Developer/Windows/ShaderFormatD3D/Private/D3DShaderCompilerDXC.cpp
--- a/Engine/Source/Developer/Windows/ShaderFormatD3D/Private/D3DShaderCompilerDXC.cpp
+++ b/Engine/Source/Developer/Windows/ShaderFormatD3D/Private/D3DShaderCompilerDXC.cpp
@@ -614,93 +614,6 @@ static bool DXCRewriteWrapper(
 #endif
 }
 
-static const TCHAR* GRewrittenBaseFilename = TEXT("Output.dxc.hlsl");
-static bool RewriteUsingSC(FString& PreprocessedShaderSource, const FShaderCompilerInput& Input, bool bIsRayTracingShader,
-	bool bDumpDebugInfo, ELanguage Language, FShaderCompilerOutput& Output)
-{
-	bool bResult = true;
-	if (bIsRayTracingShader)
-	{
-		bResult = false;
-	}
-	else
-	{
-		// Set up compile options for ShaderConductor (shader model, optimization settings etc.)
-		ShaderConductor::Compiler::Options Options;
-		Options.removeUnusedGlobals = false;
-		Options.packMatricesInRowMajor = false;
-		Options.enableDebugInfo = false;
-		Options.enable16bitTypes = false;
-		Options.disableOptimizations = false;
-		Options.shaderModel = ToDXCShaderModel(Language);
-
-		// Convert input source code from TCHAR to ANSI
-		std::string CStrSourceData(TCHAR_TO_ANSI(*PreprocessedShaderSource));
-		std::string CStrFileName(TCHAR_TO_ANSI(*Input.VirtualSourceFilePath));
-		std::string CStrEntryPointName(TCHAR_TO_ANSI(*Input.EntryPointName));
-
-		const ShaderConductor::MacroDefine BuiltinDefines[] =
-		{
-//			{ "COMPILER_HLSL", "1" },
-			{ "TextureExternal", "Texture2D" },
-		};
-
-		// Set up source description for ShaderConductor
-		ShaderConductor::Compiler::SourceDesc SourceDesc;
-		FMemory::Memzero(SourceDesc);
-		SourceDesc.source = CStrSourceData.c_str();
-		SourceDesc.fileName = CStrFileName.c_str();
-		SourceDesc.entryPoint = CStrEntryPointName.c_str();
-		SourceDesc.numDefines = sizeof(BuiltinDefines) / sizeof(BuiltinDefines[0]);
-		SourceDesc.defines = BuiltinDefines;
-		SourceDesc.stage = ToDXCShaderStage(Input.Target.GetFrequency());
-
-		ShaderConductor::Compiler::TargetDesc TargetDesc;
-		FMemory::Memzero(TargetDesc);
-		TargetDesc.language = ShaderConductor::ShadingLanguage::Dxil;
-
-		// Rewrite HLSL source to remove unused global variables (DXC retains them when compiling)
-		Options.removeUnusedGlobals = true;
-		bool bException = false;
-		ShaderConductor::Compiler::ResultDesc RewriteResultDesc;
-		DXCRewriteWrapper(SourceDesc, Options, RewriteResultDesc, bException);
-		Options.removeUnusedGlobals = false;
-		if (RewriteResultDesc.hasError || bException)
-		{
-			if (bException)
-			{
-				Output.Errors.Add(TEXT("ShaderConductor exception during rewrite"));
-			}
-			// Append compile error to output reports
-			if (RewriteResultDesc.errorWarningMsg.Size() > 0)
-			{
-				FUTF8ToTCHAR UTF8Converter(reinterpret_cast<const ANSICHAR*>(RewriteResultDesc.errorWarningMsg.Data()), RewriteResultDesc.errorWarningMsg.Size());
-				const FString ErrorString(RewriteResultDesc.errorWarningMsg.Size(), UTF8Converter.Get());
-				Output.Errors.Add(*ErrorString);
-				RewriteResultDesc.errorWarningMsg.Reset();
-				bResult = false;
-			}
-		}
-		else
-		{
-			// Copy rewritten HLSL code into new source data string
-			CStrSourceData.clear();
-			CStrSourceData.resize(RewriteResultDesc.target.Size());
-			FCStringAnsi::Strncpy(&CStrSourceData[0], static_cast<const char*>(RewriteResultDesc.target.Data()), RewriteResultDesc.target.Size());
-			PreprocessedShaderSource = CStrSourceData.c_str();
-
-			if (bDumpDebugInfo)
-			{
-				UE::ShaderCompilerCommon::FDebugShaderDataOptions DebugDataOptions;
-				DebugDataOptions.OverrideBaseFilename = GRewrittenBaseFilename;
-				UE::ShaderCompilerCommon::DumpDebugShaderData(Input, PreprocessedShaderSource, DebugDataOptions);
-			}
-		}
-	}
-
-	return bResult;
-}
-
 // Generate the dumped usf file; call the D3D compiler, gather reflection information and generate the output data
 bool CompileAndProcessD3DShaderDXC(FString& PreprocessedShaderSource,
 	const uint32 CompileFlags,
@@ -757,14 +670,6 @@ bool CompileAndProcessD3DShaderDXC(FString& PreprocessedShaderSource,
 
 	FString Filename = Input.GetSourceFilename();
 
-	if (Input.Environment.CompilerFlags.Contains(CFLAG_D3D12ForceShaderConductorRewrite))
-	{
-		if (RewriteUsingSC(PreprocessedShaderSource, Input, bIsRayTracingShader, bDumpDebugInfo, Language, Output))
-		{
-			Filename = GRewrittenBaseFilename;
-		}
-	}
-
 	FString DisasmFilename;
 	if (bDumpDebugInfo)
 	{
diff --git a/Engine/Source/Runtime/Engine/Private/ShaderCompiler/ShaderCompiler.cpp b/Engine/Source/Runtime/Engine/Private/ShaderCompiler/ShaderCompiler.cpp
--- a/Engine/Source/Runtime/Engine/Private/ShaderCompiler/ShaderCompiler.cpp
+++ b/Engine/Source/Runtime/Engine/Private/ShaderCompiler/ShaderCompiler.cpp
@@ -1120,14 +1120,6 @@ static TAutoConsoleVariable<int32> CVarD3DForceDXC(
 	TEXT(" 1: Force new compiler for all shaders"),
 	ECVF_ReadOnly);
 
-static TAutoConsoleVariable<int32> CVarD3DForceShaderConductorDXCRewrite(
-	TEXT("r.D3D.ForceShaderConductorDXCRewrite"),
-	0,
-	TEXT("Forces rewriting using ShaderConductor when DXC is enabled.\n")
-	TEXT(" 0: Do not rewrite (default)\n")
-	TEXT(" 1: Force ShaderConductor rewrite"),
-	ECVF_ReadOnly);
-
 static TAutoConsoleVariable<int32> CVarOpenGLForceDXC(
 	TEXT("r.OpenGL.ForceDXC"),
 	1,
@@ -6243,10 +6235,6 @@ void GlobalBeginCompileShader(
 			Input.Environment.CompilerFlags.Add(CFLAG_AllowTypedUAVLoads);
 		}
 
-		if (CVarD3DForceShaderConductorDXCRewrite.GetValueOnAnyThread() != 0)
-		{
-			Input.Environment.CompilerFlags.Add(CFLAG_D3D12ForceShaderConductorRewrite);
-		}
 		{
 			static const auto CVar = IConsoleManager::Get().FindConsoleVariable(TEXT("r.D3D.CheckedForTypedUAVs"));
 			if (CVar && CVar->GetInt() == 0)
diff --git a/Engine/Source/Runtime/RenderCore/Public/ShaderCompilerCore.h b/Engine/Source/Runtime/RenderCore/Public/ShaderCompilerCore.h
--- a/Engine/Source/Runtime/RenderCore/Public/ShaderCompilerCore.h
+++ b/Engine/Source/Runtime/RenderCore/Public/ShaderCompilerCore.h
@@ -114,7 +114,7 @@ enum ECompilerFlags
 	// Enable support of inline raytracing in compute shader.
 	CFLAG_InlineRayTracing,
 	// Force using the SC rewrite functionality before calling DXC on D3D12
-	CFLAG_D3D12ForceShaderConductorRewrite,
+	CFLAG_D3D12ForceShaderConductorRewrite UE_DEPRECATED(5.3, "CFLAG_D3D12ForceShaderConductorRewrite has been deprecated since UE5.3 and the flag is ignored"),
 	// Enable support of C-style data types for platforms that can. Check for PLATFORM_SUPPORTS_REAL_TYPES and FDataDrivenShaderPlatformInfo::GetSupportsRealTypes()
 	CFLAG_AllowRealTypes,
 	// Precompile HLSL to optimized HLSL, then forward to FXC. Speeds up some shaders that take longer with FXC and works around crashes in FXC.
