Changing vendor extensions code to use EGpuVendorId instead of hardcoded IDs.

diff --git a/Engine/Source/Developer/Windows/ShaderFormatD3D/Private/D3DShaderCompiler.inl b/Engine/Source/Developer/Windows/ShaderFormatD3D/Private/D3DShaderCompiler.inl
--- a/Engine/Source/Developer/Windows/ShaderFormatD3D/Private/D3DShaderCompiler.inl
+++ b/Engine/Source/Developer/Windows/ShaderFormatD3D/Private/D3DShaderCompiler.inl
@@ -473,7 +473,7 @@ template <typename ID3D1xShaderReflection, typename D3D1x_SHADER_DESC, typename
 
 			if (bIsVendorParameter)
 			{
-				VendorExtensions.Emplace(0x1002 /*AMD*/, 0, BindDesc.BindPoint, BindCount, ParameterType);
+				VendorExtensions.Emplace(EGpuVendorId::Amd, 0, BindDesc.BindPoint, BindCount, ParameterType);
 			}
 			else if (ParameterType == EShaderParameterType::Sampler)
 			{
@@ -512,11 +512,11 @@ template <typename ID3D1xShaderReflection, typename D3D1x_SHADER_DESC, typename
 			const uint32 BindCount = 1;
 			if (bIsVendorParameter)
 			{
-				const uint32 VendorId =
-					bIsNVExtension ? 0x10DE : // NVIDIA
-					(bIsAMDExtensionDX11 || bIsAMDExtensionDX12) ? 0x1002 : // AMD
-					bIsIntelExtension ? 0x8086 : // Intel
-					0;
+				const EGpuVendorId VendorId =
+					bIsNVExtension ? EGpuVendorId::Nvidia :
+					(bIsAMDExtensionDX11 || bIsAMDExtensionDX12) ? EGpuVendorId::Amd :
+					bIsIntelExtension ? EGpuVendorId::Intel :
+					EGpuVendorId::Unknown;
 				VendorExtensions.Emplace(VendorId, 0, BindDesc.BindPoint, BindCount, EShaderParameterType::UAV);
 			}
 			else if (bIsDiagnosticBufferParameter)
diff --git a/Engine/Source/Runtime/D3D12RHI/Private/D3D12Util.cpp b/Engine/Source/Runtime/D3D12RHI/Private/D3D12Util.cpp
--- a/Engine/Source/Runtime/D3D12RHI/Private/D3D12Util.cpp
+++ b/Engine/Source/Runtime/D3D12RHI/Private/D3D12Util.cpp
@@ -1051,7 +1051,7 @@ bool NeedsAgsIntrinsicsSpace(const FD3D12ShaderData& ShaderData)
 #if D3D12RHI_NEEDS_VENDOR_EXTENSIONS
 	for (const FShaderCodeVendorExtension& Extension : ShaderData.VendorExtensions)
 	{
-		if (Extension.VendorId == 0x1002) // AMD
+		if (Extension.VendorId == EGpuVendorId::Amd)
 		{
 			// [REDACTED_PROFILE_URL]/AGS_SDK/blob/master/ags_lib/hlsl/ags_shader_intrinsics_dx12.hlsl
 			return true;
diff --git a/Engine/Source/Runtime/D3D12RHI/Private/Windows/WindowsD3D12PipelineState.cpp b/Engine/Source/Runtime/D3D12RHI/Private/Windows/WindowsD3D12PipelineState.cpp
--- a/Engine/Source/Runtime/D3D12RHI/Private/Windows/WindowsD3D12PipelineState.cpp
+++ b/Engine/Source/Runtime/D3D12RHI/Private/Windows/WindowsD3D12PipelineState.cpp
@@ -888,7 +888,7 @@ static void CreatePipelineStateWithExtensions(ID3D12PipelineState** PSO, FD3D12A
 	for (const FShaderCodeVendorExtension& Extension : VendorExtensions)
 	{
 #if WITH_NVAPI
-		if (Extension.VendorId == 0x10DE) // NVIDIA
+		if (Extension.VendorId == EGpuVendorId::Nvidia)
 		{
 			if (Extension.Parameter.Type == EShaderParameterType::UAV)
 			{
@@ -915,14 +915,14 @@ static void CreatePipelineStateWithExtensions(ID3D12PipelineState** PSO, FD3D12A
 			}
 		}
 #endif //  WITH_NVAPI
-		if (Extension.VendorId == 0x1002) // AMD
+		if (Extension.VendorId == EGpuVendorId::Amd)
 		{
 			// [REDACTED_PROFILE_URL]/AGS_SDK/blob/master/ags_lib/hlsl/ags_shader_intrinsics_dx12.hlsl
 			// No special create override needed, pass through to default:
 			CreatePipelineStateWrapper(PSO, Adapter, CreationArgs);
 			return;
 		}
-		else if (Extension.VendorId == 0x8086) // INTEL
+		else if (Extension.VendorId == EGpuVendorId::Intel)
 		{
 			// [REDACTED_PROFILE_URL]/intel-graphics-compiler/blob/master/inc/IntelExtensions.hlsl
 			// No special create override needed, pass through to default:
diff --git a/Engine/Source/Runtime/RHI/Public/RHIDefinitions.h b/Engine/Source/Runtime/RHI/Public/RHIDefinitions.h
--- a/Engine/Source/Runtime/RHI/Public/RHIDefinitions.h
+++ b/Engine/Source/Runtime/RHI/Public/RHIDefinitions.h
@@ -804,9 +804,9 @@ ENUM_CLASS_FLAGS(EBufferUsageFlags);
 #define BUF_MultiGPUAllocate       EBufferUsageFlags::MultiGPUAllocate
 #define BUF_MultiGPUGraphIgnore    EBufferUsageFlags::MultiGPUGraphIgnore
 
-enum class EGpuVendorId
+enum class EGpuVendorId : uint32
 {
-	Unknown		= -1,
+	Unknown		= 0xffffffff,
 	NotQueried	= 0,
 
 	Amd			= 0x1002,
diff --git a/Engine/Source/Runtime/RenderCore/Public/ShaderCore.h b/Engine/Source/Runtime/RenderCore/Public/ShaderCore.h
--- a/Engine/Source/Runtime/RenderCore/Public/ShaderCore.h
+++ b/Engine/Source/Runtime/RenderCore/Public/ShaderCore.h
@@ -769,11 +769,11 @@ struct FShaderCodeVendorExtension
 	// for FindOptionalData() and AddOptionalData()
 	static const uint8 Key = 'v';
 
-	uint32 VendorId = 0;
+	EGpuVendorId VendorId = EGpuVendorId::NotQueried;
 	FParameterAllocation Parameter;
 
 	FShaderCodeVendorExtension() = default;
-	FShaderCodeVendorExtension(uint32 InVendorId, uint16 InBufferIndex, uint16 InBaseIndex, uint16 InSize, EShaderParameterType InType)
+	FShaderCodeVendorExtension(EGpuVendorId InVendorId, uint16 InBufferIndex, uint16 InBaseIndex, uint16 InSize, EShaderParameterType InType)
 		: VendorId(InVendorId)
 		, Parameter(InBufferIndex, InBaseIndex, InSize, InType)
 	{
diff --git a/Engine/Source/Runtime/Windows/D3D11RHI/Private/D3D11Shaders.cpp b/Engine/Source/Runtime/Windows/D3D11RHI/Private/D3D11Shaders.cpp
--- a/Engine/Source/Runtime/Windows/D3D11RHI/Private/D3D11Shaders.cpp
+++ b/Engine/Source/Runtime/Windows/D3D11RHI/Private/D3D11Shaders.cpp
@@ -69,7 +69,7 @@ static bool ApplyVendorExtensions(ID3D11Device* Direct3DDevice, EShaderFrequency
 #if !PLATFORM_HOLOLENS
 	for (const FShaderCodeVendorExtension& Extension : ShaderData->VendorExtensions)
 	{
-		if (Extension.VendorId == 0x10DE) // NVIDIA
+		if (Extension.VendorId == EGpuVendorId::Nvidia)
 		{
 			if (!IsRHIDeviceNVIDIA())
 			{
@@ -86,7 +86,7 @@ static bool ApplyVendorExtensions(ID3D11Device* Direct3DDevice, EShaderFrequency
 			}
 #endif
 		}
-		else if (Extension.VendorId == 0x1002) // AMD
+		else if (Extension.VendorId == EGpuVendorId::Amd)
 		{
 			if (!IsRHIDeviceAMD())
 			{
@@ -95,7 +95,7 @@ static bool ApplyVendorExtensions(ID3D11Device* Direct3DDevice, EShaderFrequency
 			}
 			// TODO: [REDACTED_PROFILE_URL]/AGS_SDK/blob/master/ags_lib/hlsl/ags_shader_intrinsics_dx11.hlsl
 		}
-		else if (Extension.VendorId == 0x8086) // Intel
+		else if (Extension.VendorId == EGpuVendorId::Intel)
 		{
 			if (!IsRHIDeviceIntel())
 			{
