D3D Shader Compiler cleanup
- Changes focused around making our DXC interface less dependent on FXC flags.
- Moving FDxcArguments to D3DShaderCompilerDXC
- Making FDxcArguments depend on FShaderCompilerInput instead of reinterpreting the FXC API arguments.
- Changing DXC arguments to use "-" prefix instead of "/"

diff --git a/Engine/Source/Developer/Windows/ShaderFormatD3D/Private/D3DShaderCompiler.cpp b/Engine/Source/Developer/Windows/ShaderFormatD3D/Private/D3DShaderCompiler.cpp
--- a/Engine/Source/Developer/Windows/ShaderFormatD3D/Private/D3DShaderCompiler.cpp
+++ b/Engine/Source/Developer/Windows/ShaderFormatD3D/Private/D3DShaderCompiler.cpp
@@ -1033,13 +1033,44 @@ static bool CompileAndProcessD3DShaderFXCExt(
 }
 
 bool CompileAndProcessD3DShaderFXC(FString& PreprocessedShaderSource,
-	uint32 CompileFlags,
 	const FShaderCompilerInput& Input,
 	const FShaderParameterParser& ShaderParameterParser,
 	FString& EntryPointName,
-	const TCHAR* ShaderProfile, bool bSecondPassAferUnusedInputRemoval,
+	const TCHAR* ShaderProfile,
+	bool bSecondPassAferUnusedInputRemoval,
 	FShaderCompilerOutput& Output)
 {
+	// @TODO - implement different material path to allow us to remove backwards compat flag on sm5 shaders
+	uint32 CompileFlags = D3DCOMPILE_ENABLE_BACKWARDS_COMPATIBILITY
+		// Unpack uniform matrices as row-major to match the CPU layout.
+		| D3DCOMPILE_PACK_MATRIX_ROW_MAJOR;
+
+	if (Input.Environment.CompilerFlags.Contains(CFLAG_GenerateSymbols))
+	{
+		CompileFlags |= D3DCOMPILE_DEBUG;
+	}
+
+	if (Input.Environment.CompilerFlags.Contains(CFLAG_Debug))
+	{
+		CompileFlags |= D3DCOMPILE_SKIP_OPTIMIZATION;
+	}
+	else
+	{
+		if (Input.Environment.CompilerFlags.Contains(CFLAG_StandardOptimization))
+		{
+			CompileFlags |= D3DCOMPILE_OPTIMIZATION_LEVEL1;
+		}
+		else
+		{
+			CompileFlags |= D3DCOMPILE_OPTIMIZATION_LEVEL3;
+		}
+	}
+
+	Input.Environment.CompilerFlags.Iterate([&CompileFlags](uint32 Flag)
+		{
+			CompileFlags |= TranslateCompilerFlagD3D11((ECompilerFlags)Flag);
+		});
+
 	TArray<FString> FilteredErrors;
 	const bool bSuccess = CompileAndProcessD3DShaderFXCExt(PreprocessedShaderSource, CompileFlags, Input, ShaderParameterParser, EntryPointName, ShaderProfile, false, FilteredErrors, Output);
 
@@ -1207,40 +1238,9 @@ void CompileD3DShader(const FShaderCompilerInput& Input, FShaderCompilerOutput&
 	}
 	#endif // UE_D3D_SHADER_COMPILER_ALLOW_DEAD_CODE_REMOVAL
 
-	// @TODO - implement different material path to allow us to remove backwards compat flag on sm5 shaders
-	uint32 CompileFlags = D3DCOMPILE_ENABLE_BACKWARDS_COMPATIBILITY
-		// Unpack uniform matrices as row-major to match the CPU layout.
-		| D3DCOMPILE_PACK_MATRIX_ROW_MAJOR;
-
-	if (Input.Environment.CompilerFlags.Contains(CFLAG_GenerateSymbols))
-	{
-		CompileFlags |= D3DCOMPILE_DEBUG;
-	}
-
-	if (Input.Environment.CompilerFlags.Contains(CFLAG_Debug)) 
-	{
-		CompileFlags |= D3DCOMPILE_SKIP_OPTIMIZATION;
-	}
-	else
-	{
-		if (Input.Environment.CompilerFlags.Contains(CFLAG_StandardOptimization))
-		{
-			CompileFlags |= D3DCOMPILE_OPTIMIZATION_LEVEL1;
-		}
-		else
-		{
-			CompileFlags |= D3DCOMPILE_OPTIMIZATION_LEVEL3;
-		}
-	}
-
-	Input.Environment.CompilerFlags.Iterate([&CompileFlags](uint32 Flag)
-		{
-			CompileFlags |= TranslateCompilerFlagD3D11((ECompilerFlags)Flag);
-		});
-
 	const bool bSuccess = bUseDXC
-		? CompileAndProcessD3DShaderDXC(PreprocessedShaderSource, CompileFlags, Input, ShaderParameterParser, EntryPointName, ShaderProfile, Language, false, Output)
-		: CompileAndProcessD3DShaderFXC(PreprocessedShaderSource, CompileFlags, Input, ShaderParameterParser, EntryPointName, ShaderProfile, false, Output);
+		? CompileAndProcessD3DShaderDXC(PreprocessedShaderSource, Input, ShaderParameterParser, EntryPointName, ShaderProfile, Language, false, Output)
+		: CompileAndProcessD3DShaderFXC(PreprocessedShaderSource, Input, ShaderParameterParser, EntryPointName, ShaderProfile, false, Output);
 
 	if (!bSuccess && !Output.Errors.Num())
 	{
diff --git a/Engine/Source/Developer/Windows/ShaderFormatD3D/Private/D3DShaderCompiler.inl b/Engine/Source/Developer/Windows/ShaderFormatD3D/Private/D3DShaderCompiler.inl
--- a/Engine/Source/Developer/Windows/ShaderFormatD3D/Private/D3DShaderCompiler.inl
+++ b/Engine/Source/Developer/Windows/ShaderFormatD3D/Private/D3DShaderCompiler.inl
@@ -2,298 +2,6 @@
 
 #pragma once
 
-
-class FDxcArguments
-{
-protected:
-	FString ShaderProfile;
-	FString EntryPoint;
-	FString Exports;
-	FString DumpDisasmFilename;
-	FString BatchBaseFilename;
-	FString DumpDebugInfoPath;
-	bool bEnable16BitTypes = false;
-	bool bKeepEmbeddedPDB = false;
-	bool bDump = false;
-
-	TArray<FString> ExtraArguments;
-
-public:
-	FDxcArguments(
-		const FString& InEntryPoint,
-		const TCHAR* InShaderProfile,
-		const FString& InExports,
-		const FString& InDumpDebugInfoPath,	// Optional, empty when not dumping shader debug info
-		const FString& InBaseFilename,
-		bool bInEnable16BitTypes,
-		bool bGenerateSymbols,
-		bool bSymbolsBasedOnSource,
-		uint32 D3DCompileFlags,
-		uint32 AutoBindingSpace,
-		const TCHAR* InOptValidatorVersion,
-		uint32 HlslVersion,
-		bool bBindlessResources
-	)
-		: ShaderProfile(InShaderProfile)
-		, EntryPoint(InEntryPoint)
-		, Exports(InExports)
-		, DumpDebugInfoPath(InDumpDebugInfoPath)
-		, bEnable16BitTypes(bInEnable16BitTypes)
-	{
-		BatchBaseFilename = FPaths::GetBaseFilename(InBaseFilename);
-
-		if (InDumpDebugInfoPath.Len() > 0)
-		{
-			bDump = true;
-			DumpDisasmFilename = InDumpDebugInfoPath / TEXT("Output.d3dasm");
-		}
-
-		switch (HlslVersion)
-		{
-		case 2015:
-			ExtraArguments.Add(TEXT("-HV"));
-			ExtraArguments.Add(TEXT("2015"));
-			break;
-		case 2016:
-			ExtraArguments.Add(TEXT("-HV"));
-			ExtraArguments.Add(TEXT("2016"));
-			break;
-		case 2017:
-			ExtraArguments.Add(TEXT("-HV"));
-			ExtraArguments.Add(TEXT("2017"));
-			break;
-		case 2018:
-			break; // Default
-		case 2021:
-			ExtraArguments.Add(TEXT("-HV"));
-			ExtraArguments.Add(TEXT("2021"));
-			break;
-		default:
-			checkf(false, TEXT("Invalid HLSL version: expected 2015, 2016, 2017, 2018, or 2021 but %u was specified"), HlslVersion);
-			break;
-		}
-
-		if (AutoBindingSpace != ~0u)
-		{
-			ExtraArguments.Add(TEXT("/auto-binding-space"));
-			ExtraArguments.Add(FString::Printf(TEXT("%d"), AutoBindingSpace));
-		}
-
-		if (Exports.Len() > 0)
-		{
-			// Ensure that only the requested functions exists in the output DXIL.
-			// All other functions and their used resources must be eliminated.
-			ExtraArguments.Add(TEXT("/exports"));
-			ExtraArguments.Add(Exports);
-		}
-
-		if (D3DCompileFlags & D3DCOMPILE_PREFER_FLOW_CONTROL)
-		{
-			D3DCompileFlags &= ~D3DCOMPILE_PREFER_FLOW_CONTROL;
-			ExtraArguments.Add(TEXT("/Gfp"));
-		}
-
-		if (D3DCompileFlags & D3DCOMPILE_SKIP_VALIDATION)
-		{
-			D3DCompileFlags &= ~D3DCOMPILE_SKIP_VALIDATION;
-			ExtraArguments.Add(TEXT("/Vd"));
-		}
-
-		if (D3DCompileFlags & D3DCOMPILE_AVOID_FLOW_CONTROL)
-		{
-			D3DCompileFlags &= ~D3DCOMPILE_AVOID_FLOW_CONTROL;
-			ExtraArguments.Add(TEXT("/Gfa"));
-		}
-
-		if (D3DCompileFlags & D3DCOMPILE_PACK_MATRIX_ROW_MAJOR)
-		{
-			D3DCompileFlags &= ~D3DCOMPILE_PACK_MATRIX_ROW_MAJOR;
-			ExtraArguments.Add(TEXT("/Zpr"));
-		}
-
-		if (D3DCompileFlags & D3DCOMPILE_ENABLE_BACKWARDS_COMPATIBILITY)
-		{
-			D3DCompileFlags &= ~D3DCOMPILE_ENABLE_BACKWARDS_COMPATIBILITY;
-			ExtraArguments.Add(TEXT("/Gec"));
-		}
-
-		if (D3DCompileFlags & D3DCOMPILE_WARNINGS_ARE_ERRORS)
-		{
-			D3DCompileFlags &= ~D3DCOMPILE_WARNINGS_ARE_ERRORS;
-			ExtraArguments.Add(TEXT("/WX"));
-		}
-
-		if (D3DCompileFlags & D3DCOMPILE_SKIP_OPTIMIZATION)
-		{
-			D3DCompileFlags &= ~D3DCOMPILE_SKIP_OPTIMIZATION;
-			ExtraArguments.Add(TEXT("/Od"));
-		}
-		else
-		{
-			switch (D3DCompileFlags & SHADER_OPTIMIZATION_LEVEL_MASK)
-			{
-			case D3DCOMPILE_OPTIMIZATION_LEVEL0:
-				D3DCompileFlags &= ~D3DCOMPILE_OPTIMIZATION_LEVEL0;
-				ExtraArguments.Add(TEXT("/O0"));
-				break;
-
-			case D3DCOMPILE_OPTIMIZATION_LEVEL1:
-				D3DCompileFlags &= ~D3DCOMPILE_OPTIMIZATION_LEVEL1;
-				ExtraArguments.Add(TEXT("/O1"));
-				break;
-
-			case D3DCOMPILE_OPTIMIZATION_LEVEL2:
-				D3DCompileFlags &= ~D3DCOMPILE_OPTIMIZATION_LEVEL2;
-				ExtraArguments.Add(TEXT("/O2"));
-				break;
-
-			case D3DCOMPILE_OPTIMIZATION_LEVEL3:
-				D3DCompileFlags &= ~D3DCOMPILE_OPTIMIZATION_LEVEL3;
-				ExtraArguments.Add(TEXT("/O3"));
-				break;
-
-			default:
-				break;
-			}
-		}
-
-		if (D3DCompileFlags & D3DCOMPILE_DEBUG)
-		{
-			D3DCompileFlags &= ~D3DCOMPILE_DEBUG;
-			bGenerateSymbols = true;
-		}
-
-		if (bEnable16BitTypes)
-		{
-			ExtraArguments.Add(TEXT("/enable-16bit-types"));
-		}
-
-		checkf(D3DCompileFlags == 0, TEXT("Unhandled shader compiler flags 0x%x!"), D3DCompileFlags);
-
-		if (InOptValidatorVersion)
-		{
-			ExtraArguments.Add(TEXT("/validator-version"));
-			ExtraArguments.Add(FString(InOptValidatorVersion));
-		}
-
-		if (bGenerateSymbols)
-		{
-			// -Zsb Compute Shader Hash considering only output binary
-			// -Zss Compute Shader Hash considering source information
-			ExtraArguments.Add(bSymbolsBasedOnSource ? TEXT("/Zss") : TEXT("/Zsb"));
-
-			ExtraArguments.Add(TEXT("/Qembed_debug"));
-			ExtraArguments.Add(TEXT("/Zi"));
-
-			ExtraArguments.Add(TEXT("/Fd"));
-			ExtraArguments.Add(TEXT(".\\"));
-
-			bKeepEmbeddedPDB = true;
-		}
-
-		// Reflection will be removed later, otherwise the disassembly won't contain variables
-		//ExtraArguments.Add(TEXT("/Qstrip_reflect"));
-
-		// disable undesired warnings
-		ExtraArguments.Add(TEXT("-Wno-parentheses-equality"));
-
-		// working around bindless conversion specific issue where globallycoherent on a function return type is flagged as ignored even though it is necessary.
-		// github issue: [REDACTED_PROFILE_URL]/DirectXShaderCompiler/issues/4537
-		if (bBindlessResources)
-		{
-			ExtraArguments.Add(TEXT("-Wno-ignored-attributes"));
-		}
-
-		// @lh-todo: This fixes a loop unrolling issue that showed up in DOFGatherKernel with cs_6_6 with the latest DXC revision
-		ExtraArguments.Add(TEXT("-disable-lifetime-markers"));
-	}
-
-	inline FString GetDumpDebugInfoPath() const
-	{
-		return DumpDebugInfoPath;
-	}
-
-	inline bool ShouldKeepEmbeddedPDB() const
-	{
-		return bKeepEmbeddedPDB;
-	}
-
-	inline bool ShouldDump() const
-	{
-		return bDump;
-	}
-
-	FString GetEntryPointName() const
-	{
-		return Exports.Len() > 0 ? FString(TEXT("")) : EntryPoint;
-	}
-
-	const FString& GetShaderProfile() const
-	{
-		return ShaderProfile;
-	}
-
-	const FString& GetDumpDisassemblyFilename() const
-	{
-		return DumpDisasmFilename;
-	}
-
-	void GetCompilerArgsNoEntryNoProfileNoDisasm(TArray<const WCHAR*>& Out) const
-	{
-		for (const FString& Entry : ExtraArguments)
-		{
-			Out.Add(*Entry);
-		}
-	}
-
-	void GetCompilerArgs(TArray<const WCHAR*>& Out) const
-	{
-		GetCompilerArgsNoEntryNoProfileNoDisasm(Out);
-		if (Exports.Len() == 0)
-		{
-			Out.Add(TEXT("/E"));
-			Out.Add(*EntryPoint);
-		}
-
-		Out.Add(TEXT("/T"));
-		Out.Add(*ShaderProfile);
-
-		Out.Add(TEXT(" /Fc "));
-		Out.Add(TEXT("zzz.d3dasm"));	// Dummy
-
-		Out.Add(TEXT(" /Fo "));
-		Out.Add(TEXT("zzz.dxil"));	// Dummy
-	}
-
-	FString GetBatchCommandLineString(const FString& ShaderPath) const
-	{
-		FString DXCCommandline;
-		for (const FString& Entry : ExtraArguments)
-		{
-			DXCCommandline += TEXT(" ");
-			DXCCommandline += Entry;
-		}
-
-		DXCCommandline += TEXT(" /T ");
-		DXCCommandline += ShaderProfile;
-
-		if (Exports.Len() == 0)
-		{
-			DXCCommandline += TEXT(" /E ");
-			DXCCommandline += EntryPoint;
-		}
-
-		DXCCommandline += TEXT(" /Fc ");
-		DXCCommandline += BatchBaseFilename + TEXT(".d3dasm");
-
-		DXCCommandline += TEXT(" /Fo ");
-		DXCCommandline += BatchBaseFilename + TEXT(".dxil");
-
-		return DXCCommandline;
-	}
-};
-
-
 template <typename ID3D1xShaderReflection, typename D3D1x_SHADER_DESC, typename D3D1x_SHADER_INPUT_BIND_DESC,
 	typename ID3D1xShaderReflectionConstantBuffer, typename D3D1x_SHADER_BUFFER_DESC,
 	typename ID3D1xShaderReflectionVariable, typename D3D1x_SHADER_VARIABLE_DESC>
diff --git a/Engine/Source/Developer/Windows/ShaderFormatD3D/Private/D3DShaderCompilerDXC.cpp b/Engine/Source/Developer/Windows/ShaderFormatD3D/Private/D3DShaderCompilerDXC.cpp
--- a/Engine/Source/Developer/Windows/ShaderFormatD3D/Private/D3DShaderCompilerDXC.cpp
+++ b/Engine/Source/Developer/Windows/ShaderFormatD3D/Private/D3DShaderCompilerDXC.cpp
@@ -153,6 +153,256 @@ static void LogFailedHRESULT(const TCHAR* FailedExpressionStr, HRESULT Result)
 		}													\
 	}
 
+
+class FDxcArguments
+{
+protected:
+	FString ShaderProfile;
+	FString EntryPoint;
+	FString Exports;
+	FString DumpDisasmFilename;
+	FString BatchBaseFilename;
+	FString DumpDebugInfoPath;
+	bool bKeepEmbeddedPDB = false;
+	bool bDump = false;
+
+	TArray<FString> ExtraArguments;
+
+public:
+	FDxcArguments(
+		const FShaderCompilerInput& Input,
+		const FString& InEntryPoint,
+		const TCHAR* InShaderProfile,
+		ELanguage Language,
+		const FString& InExports
+	)
+		: ShaderProfile(InShaderProfile)
+		, EntryPoint(InEntryPoint)
+		, Exports(InExports)
+		, BatchBaseFilename(FPaths::GetBaseFilename(Input.GetSourceFilename()))
+		, DumpDebugInfoPath(Input.DumpDebugInfoPath)
+		, bDump(Input.DumpDebugInfoEnabled())
+	{
+		if (bDump)
+		{
+			DumpDisasmFilename = Input.DumpDebugInfoPath / TEXT("Output.d3dasm");
+		}
+
+		const bool bEnable16BitTypes =
+			// 16bit types are SM6.2 whereas Language == ELanguage::SM6 is SM6.6, so their support at runtime is guarented.
+			(Language == ELanguage::SM6 && Input.Environment.CompilerFlags.Contains(CFLAG_AllowRealTypes))
+
+			// Enable 16bit_types to reduce DXIL size (compiler bug - will be fixed)
+			|| Input.IsRayTracingShader();
+
+		const bool bHlslVersion2021 = Input.Environment.CompilerFlags.Contains(CFLAG_HLSL2021);
+		const uint32 HlslVersion = (bHlslVersion2021 ? 2021 : 2018);
+
+		switch (HlslVersion)
+		{
+		case 2015:
+			ExtraArguments.Add(TEXT("-HV"));
+			ExtraArguments.Add(TEXT("2015"));
+			break;
+		case 2016:
+			ExtraArguments.Add(TEXT("-HV"));
+			ExtraArguments.Add(TEXT("2016"));
+			break;
+		case 2017:
+			ExtraArguments.Add(TEXT("-HV"));
+			ExtraArguments.Add(TEXT("2017"));
+			break;
+		case 2018:
+			break; // Default
+		case 2021:
+			ExtraArguments.Add(TEXT("-HV"));
+			ExtraArguments.Add(TEXT("2021"));
+			break;
+		default:
+			checkf(false, TEXT("Invalid HLSL version: expected 2015, 2016, 2017, 2018, or 2021 but %u was specified"), HlslVersion);
+			break;
+		}
+
+		// Unpack uniform matrices as row-major to match the CPU layout.
+		ExtraArguments.Add(TEXT("-Zpr"));
+
+		if (Input.Environment.CompilerFlags.Contains(CFLAG_Debug) || Input.Environment.CompilerFlags.Contains(CFLAG_SkipOptimizationsDXC))
+		{
+			ExtraArguments.Add(TEXT("-Od"));
+		}
+		else if (Input.Environment.CompilerFlags.Contains(CFLAG_StandardOptimization))
+		{
+			ExtraArguments.Add(TEXT("-O1"));
+		}
+		else
+		{
+			ExtraArguments.Add(TEXT("-O3"));
+		}
+
+		if (Input.Environment.CompilerFlags.Contains(CFLAG_PreferFlowControl))
+		{
+			ExtraArguments.Add(TEXT("-Gfp"));
+		}
+
+		if (Input.Environment.CompilerFlags.Contains(CFLAG_AvoidFlowControl))
+		{
+			ExtraArguments.Add(TEXT("-Gfa"));
+		}
+
+		if (Input.Environment.CompilerFlags.Contains(CFLAG_WarningsAsErrors))
+		{
+			ExtraArguments.Add(TEXT("-WX"));
+		}
+
+		const uint32 AutoBindingSpace = GetAutoBindingSpace(Input.Target);
+		if (AutoBindingSpace != ~0u)
+		{
+			ExtraArguments.Add(TEXT("-auto-binding-space"));
+			ExtraArguments.Add(FString::Printf(TEXT("%d"), AutoBindingSpace));
+		}
+
+		if (Exports.Len() > 0)
+		{
+			// Ensure that only the requested functions exists in the output DXIL.
+			// All other functions and their used resources must be eliminated.
+			ExtraArguments.Add(TEXT("-exports"));
+			ExtraArguments.Add(Exports);
+		}
+
+		if (bEnable16BitTypes)
+		{
+			ExtraArguments.Add(TEXT("-enable-16bit-types"));
+		}
+
+		if (Input.Environment.CompilerFlags.Contains(CFLAG_GenerateSymbols))
+		{
+			if (Input.Environment.CompilerFlags.Contains(CFLAG_AllowUniqueSymbols))
+			{
+				// -Zss Compute Shader Hash considering source information
+				ExtraArguments.Add(TEXT("-Zss"));
+			}
+			else
+			{
+				// -Zsb Compute Shader Hash considering only output binary
+				ExtraArguments.Add(TEXT("-Zsb"));
+			}
+
+			ExtraArguments.Add(TEXT("-Qembed_debug"));
+			ExtraArguments.Add(TEXT("-Zi"));
+
+			ExtraArguments.Add(TEXT("-Fd"));
+			ExtraArguments.Add(TEXT(".\\"));
+
+			bKeepEmbeddedPDB = true;
+		}
+
+		// Reflection will be removed later, otherwise the disassembly won't contain variables
+		//ExtraArguments.Add(TEXT("-Qstrip_reflect"));
+
+		// disable undesired warnings
+		ExtraArguments.Add(TEXT("-Wno-parentheses-equality"));
+
+		// working around bindless conversion specific issue where globallycoherent on a function return type is flagged as ignored even though it is necessary.
+		// github issue: [REDACTED_PROFILE_URL]/DirectXShaderCompiler/issues/4537
+		if (Input.Environment.CompilerFlags.Contains(CFLAG_BindlessResources))
+		{
+			ExtraArguments.Add(TEXT("-Wno-ignored-attributes"));
+		}
+
+		// @lh-todo: This fixes a loop unrolling issue that showed up in DOFGatherKernel with cs_6_6 with the latest DXC revision
+		ExtraArguments.Add(TEXT("-disable-lifetime-markers"));
+	}
+
+	FString GetDumpDebugInfoPath() const
+	{
+		return DumpDebugInfoPath;
+	}
+
+	bool ShouldKeepEmbeddedPDB() const
+	{
+		return bKeepEmbeddedPDB;
+	}
+
+	bool ShouldDump() const
+	{
+		return bDump;
+	}
+
+	FString GetEntryPointName() const
+	{
+		return Exports.Len() > 0 ? FString(TEXT("")) : EntryPoint;
+	}
+
+	const FString& GetShaderProfile() const
+	{
+		return ShaderProfile;
+	}
+
+	const FString& GetDumpDisassemblyFilename() const
+	{
+		return DumpDisasmFilename;
+	}
+
+	void GetCompilerArgsNoEntryNoProfileNoDisasm(TArray<const WCHAR*>& Out) const
+	{
+		for (const FString& Entry : ExtraArguments)
+		{
+			Out.Add(*Entry);
+		}
+	}
+
+	void GetCompilerArgs(TArray<const WCHAR*>& Out) const
+	{
+		GetCompilerArgsNoEntryNoProfileNoDisasm(Out);
+		if (Exports.Len() == 0)
+		{
+			Out.Add(TEXT("-E"));
+			Out.Add(*EntryPoint);
+		}
+
+		Out.Add(TEXT("-T"));
+		Out.Add(*ShaderProfile);
+
+		Out.Add(TEXT(" -Fc "));
+		Out.Add(TEXT("zzz.d3dasm"));	// Dummy
+
+		Out.Add(TEXT(" -Fo "));
+		Out.Add(TEXT("zzz.dxil"));	// Dummy
+	}
+
+	const FString& GetBatchBaseFilename() const
+	{
+		return BatchBaseFilename;
+	}
+
+	FString GetBatchCommandLineString() const
+	{
+		FString DXCCommandline;
+		for (const FString& Entry : ExtraArguments)
+		{
+			DXCCommandline += TEXT(" ");
+			DXCCommandline += Entry;
+		}
+
+		DXCCommandline += TEXT(" -T ");
+		DXCCommandline += ShaderProfile;
+
+		if (Exports.Len() == 0)
+		{
+			DXCCommandline += TEXT(" -E ");
+			DXCCommandline += EntryPoint;
+		}
+
+		DXCCommandline += TEXT(" -Fc ");
+		DXCCommandline += BatchBaseFilename + TEXT(".d3dasm");
+
+		DXCCommandline += TEXT(" -Fo ");
+		DXCCommandline += BatchBaseFilename + TEXT(".dxil");
+
+		return DXCCommandline;
+	}
+};
+
 class FDxcMalloc final : public IMalloc
 {
 	std::atomic<ULONG> RefCount{ 1 };
@@ -310,7 +560,7 @@ static HRESULT InnerDXCCompileWrapper(
 static HRESULT DXCCompileWrapper(
 	TRefCountPtr<IDxcCompiler3>& Compiler,
 	TRefCountPtr<IDxcBlobEncoding>& TextBlob,
-	FDxcArguments& Arguments,
+	const FDxcArguments& Arguments,
 	TRefCountPtr<IDxcResult>& OutCompileResult)
 {
 	bool bExceptionError = false;
@@ -419,7 +669,7 @@ static bool RemoveContainerReflection(dxc::DxcDllSupport& DxcDllHelper, TRefCoun
 	return false;
 };
 
-static HRESULT D3DCompileToDxil(const char* SourceText, FDxcArguments& Arguments,
+static HRESULT D3DCompileToDxil(const char* SourceText, const FDxcArguments& Arguments,
 	TRefCountPtr<IDxcBlob>& OutDxilBlob, TRefCountPtr<IDxcBlob>& OutReflectionBlob, TRefCountPtr<IDxcBlobEncoding>& OutErrorBlob)
 {
 	dxc::DxcDllSupport& DxcDllHelper = GetDxcDllHelper();
@@ -505,7 +755,7 @@ static HRESULT D3DCompileToDxil(const char* SourceText, FDxcArguments& Arguments
 	return CompileResultCode;
 }
 
-static FString D3DCreateDXCCompileBatchFile(const FDxcArguments& Args, const FString& ShaderPath)
+static FString D3DCreateDXCCompileBatchFile(const FDxcArguments& Args)
 {
 	FString DxcPath = FPaths::ConvertRelativePathToFull(FPaths::EngineDir());
 
@@ -515,7 +765,8 @@ static FString D3DCreateDXCCompileBatchFile(const FDxcArguments& Args, const FSt
 	FString DxcFilename = FPaths::Combine(DxcPath, TEXT("dxc.exe"));
 	FPaths::MakePlatformFilename(DxcFilename);
 
-	const FString BatchCmdLineArgs = Args.GetBatchCommandLineString(ShaderPath);
+	const FString& BatchBaseFilename = Args.GetBatchBaseFilename();
+	const FString BatchCmdLineArgs = Args.GetBatchCommandLineString();
 
 	return FString::Printf(
 		TEXT(
@@ -532,7 +783,7 @@ static FString D3DCreateDXCCompileBatchFile(const FDxcArguments& Args, const FSt
 		*DxcFilename,
 		*DxcPath,
 		*BatchCmdLineArgs,
-		*ShaderPath
+		*BatchBaseFilename
 	);
 }
 
@@ -616,11 +867,12 @@ static bool DXCRewriteWrapper(
 
 // Generate the dumped usf file; call the D3D compiler, gather reflection information and generate the output data
 bool CompileAndProcessD3DShaderDXC(FString& PreprocessedShaderSource,
-	const uint32 CompileFlags,
 	const FShaderCompilerInput& Input,
 	const FShaderParameterParser& ShaderParameterParser,
-	FString& EntryPointName,
-	const TCHAR* ShaderProfile, ELanguage Language, bool bProcessingSecondTime,
+	const FString& EntryPointName,
+	const TCHAR* ShaderProfile,
+	ELanguage Language,
+	bool bProcessingSecondTime,
 	FShaderCompilerOutput& Output)
 {
 	TRACE_CPUPROFILER_EVENT_SCOPE(CompileAndProcessD3DShaderDXC);
@@ -636,13 +888,6 @@ bool CompileAndProcessD3DShaderDXC(FString& PreprocessedShaderSource,
 	FString RayIntersectionEntryPoint; // Optional for hit group shaders
 	FString RayTracingExports;
 
-	bool bEnable16BitTypes = false;
-	if (Language == ELanguage::SM6)
-	{
-		// 16bit types are SM6.2 whereas Language == ELanguage::SM6 is SM6.6, so their support at runtime is guarented.
-		bEnable16BitTypes = Input.Environment.CompilerFlags.Contains(CFLAG_AllowRealTypes);
-	}
-
 	if (bIsRayTracingShader)
 	{
 		UE::ShaderCompilerCommon::ParseRayTracingEntryPoint(Input.EntryPointName, RayEntryPoint, RayAnyHitEntryPoint, RayIntersectionEntryPoint);
@@ -660,59 +905,23 @@ bool CompileAndProcessD3DShaderDXC(FString& PreprocessedShaderSource,
 			RayTracingExports += TEXT(";");
 			RayTracingExports += RayIntersectionEntryPoint;
 		}
-
-		// Enable 16bit_types to reduce DXIL size (compiler bug - will be fixed)
-		bEnable16BitTypes = true;
 	}
 
 	UE::ShaderCompilerCommon::DumpDebugShaderData(Input, PreprocessedShaderSource);
-	bool bDumpDebugInfo = Input.DumpDebugInfoEnabled(); 
-
-	FString Filename = Input.GetSourceFilename();
-
-	FString DisasmFilename;
-	if (bDumpDebugInfo)
-	{
-		DisasmFilename = Input.DumpDebugInfoPath / Filename;
-	}
-
-	// Ignore backwards compatibility flag (/Gec) as it is deprecated.
-	// #dxr_todo: this flag should not be even passed into this function from the higher level.
-	uint32 DXCFlags = CompileFlags & (~D3DCOMPILE_ENABLE_BACKWARDS_COMPATIBILITY);
-	if (Input.Environment.CompilerFlags.Contains(CFLAG_SkipOptimizationsDXC))
-	{
-		DXCFlags |= D3DCOMPILE_SKIP_OPTIMIZATION;
-	}
-
-	const TCHAR* ValidatorVersion = nullptr; // Whatever default validator chosen by DXC
-
-	const bool bGenerateSymbols = Input.Environment.CompilerFlags.Contains(CFLAG_GenerateSymbols);
-	const bool bSymbolsBasedOnSource = Input.Environment.CompilerFlags.Contains(CFLAG_AllowUniqueSymbols);
-	const bool bHlslVersion2021 = Input.Environment.CompilerFlags.Contains(CFLAG_HLSL2021);
-	const uint32 HlslVersion = (bHlslVersion2021 ? 2021 : 2018);
-	const bool bBindlessResources = Input.Environment.CompilerFlags.Contains(CFLAG_BindlessResources);
 
 	FDxcArguments Args
 	(
+		Input,
 		EntryPointName,
 		ShaderProfile,
-		RayTracingExports,
-		Input.DumpDebugInfoPath,
-		Filename,
-		bEnable16BitTypes,
-		bGenerateSymbols,
-		bSymbolsBasedOnSource,
-		DXCFlags,
-		AutoBindingSpace,
-		ValidatorVersion,
-		HlslVersion,
-		bBindlessResources
+		Language,
+		RayTracingExports
 	);
 
-	if (bDumpDebugInfo)
+	if (Args.ShouldDump())
 	{
-		FString BatchFileContents = D3DCreateDXCCompileBatchFile(Args, Filename);
-		FFileHelper::SaveStringToFile(BatchFileContents, *(Input.DumpDebugInfoPath / TEXT("CompileDXC.bat")));
+		const FString BatchFileContents = D3DCreateDXCCompileBatchFile(Args);
+		FFileHelper::SaveStringToFile(BatchFileContents, *(Args.GetDumpDebugInfoPath() / TEXT("CompileDXC.bat")));
 	}
 
 	TRefCountPtr<IDxcBlob> ShaderBlob;
diff --git a/Engine/Source/Developer/Windows/ShaderFormatD3D/Private/ShaderFormatD3D.h b/Engine/Source/Developer/Windows/ShaderFormatD3D/Private/ShaderFormatD3D.h
--- a/Engine/Source/Developer/Windows/ShaderFormatD3D/Private/ShaderFormatD3D.h
+++ b/Engine/Source/Developer/Windows/ShaderFormatD3D/Private/ShaderFormatD3D.h
@@ -28,19 +28,20 @@ void CompileShader_Windows(const struct FShaderCompilerInput& Input,struct FShad
  * @param bSecondPassAferUnusedInputRemoval whether we're compiling the shader second time, after having removed the unused inputs discovered in the first pass
  */
 bool CompileAndProcessD3DShaderFXC(FString& PreprocessedShaderSource,
-	uint32 CompileFlags,
 	const FShaderCompilerInput& Input,
 	const FShaderParameterParser& ShaderParameterParser,
 	FString& EntryPointName,
-	const TCHAR* ShaderProfile, bool bSecondPassAferUnusedInputRemoval,
+	const TCHAR* ShaderProfile,
+	bool bSecondPassAferUnusedInputRemoval,
 	FShaderCompilerOutput& Output);
 
 bool CompileAndProcessD3DShaderDXC(FString& PreprocessedShaderSource,
-	uint32 CompileFlags,
 	const FShaderCompilerInput& Input,
 	const FShaderParameterParser& ShaderParameterParser,
-	FString& EntryPointName,
-	const TCHAR* ShaderProfile, ELanguage Language, bool bProcessingSecondTime,
+	const FString& EntryPointName,
+	const TCHAR* ShaderProfile,
+	ELanguage Language,
+	bool bProcessingSecondTime,
 	FShaderCompilerOutput& Output);
 
 bool ValidateResourceCounts(uint32 NumSRVs, uint32 NumSamplers, uint32 NumUAVs, uint32 NumCBs, TArray<FString>& OutFilteredErrors);
