Set LiveCodingServer to compile with c++17

#rnx

diff --git a/Engine/Source/Developer/Windows/LiveCodingServer/LiveCodingServer.Build.cs b/Engine/Source/Developer/Windows/LiveCodingServer/LiveCodingServer.Build.cs
--- a/Engine/Source/Developer/Windows/LiveCodingServer/LiveCodingServer.Build.cs
+++ b/Engine/Source/Developer/Windows/LiveCodingServer/LiveCodingServer.Build.cs
@@ -7,6 +7,8 @@ public class LiveCodingServer : ModuleRules
 {
 	public LiveCodingServer(ReadOnlyTargetRules Target) : base(Target)
 	{
+		CppStandard = CppStandardVersion.Cpp17;
+
         PrivateDependencyModuleNames.Add("Core");
         PrivateDependencyModuleNames.Add("Json");
         PrivateDependencyModuleNames.Add("LiveCoding");
