[shaders] migrate PC D3D backends to independent preprocessing shaderformat API

diff --git a/Engine/Source/Developer/Windows/ShaderFormatD3D/Private/D3DShaderCompiler.cpp b/Engine/Source/Developer/Windows/ShaderFormatD3D/Private/D3DShaderCompiler.cpp
--- a/Engine/Source/Developer/Windows/ShaderFormatD3D/Private/D3DShaderCompiler.cpp
+++ b/Engine/Source/Developer/Windows/ShaderFormatD3D/Private/D3DShaderCompiler.cpp
@@ -5,6 +5,7 @@
 #include "ShaderCompilerCommon.h"
 #include "ShaderMinifier.h"
 #include "ShaderParameterParser.h"
+#include "ShaderPreprocessTypes.h"
 #include "D3D12RHI.h"
 #include "Misc/Paths.h"
 #include "Misc/FileHelper.h"
@@ -548,11 +549,11 @@ static bool CompileErrorsContainInternalError(ID3DBlob* Errors)
 
 // Generate the dumped usf file; call the D3D compiler, gather reflection information and generate the output data
 static bool CompileAndProcessD3DShaderFXCExt(
-	FString& PreprocessedShaderSource,
+	const FString& PreprocessedShaderSource,
 	uint32 CompileFlags,
 	const FShaderCompilerInput& Input,
 	const FShaderParameterParser& ShaderParameterParser,
-	FString& EntryPointName,
+	const FString& EntryPointName,
 	const TCHAR* ShaderProfile, bool bSecondPassAferUnusedInputRemoval,
 	TArray<FString>& FilteredErrors, FShaderCompilerOutput& Output)
 {
@@ -560,8 +561,6 @@ static bool CompileAndProcessD3DShaderFXCExt(
 
 	auto AnsiSourceFile = StringCast<ANSICHAR>(*PreprocessedShaderSource);
 
-	// Write out the preprocessed file and a batch file to compile it if requested (DumpDebugInfoPath is valid)
-	UE::ShaderCompilerCommon::DumpDebugShaderData(Input, PreprocessedShaderSource);
 	bool bDumpDebugInfo = Input.DumpDebugInfoEnabled();
 	FString DisasmFilename;
 	if (bDumpDebugInfo)
@@ -848,19 +847,19 @@ static bool CompileAndProcessD3DShaderFXCExt(
 					// Since unused inputs are passed to the next stage, that will cause us to generate a vertex shader that does not output B, but our pixel shader will still be expecting B on input,
 					// as it was rewritten based on the pass1 results.
 
-					FString OriginalPreprocSource = PreprocessedShaderSource;
-					FString OriginalEntryPointName = EntryPointName;
+					FString ModifiedSource = PreprocessedShaderSource;
+					FString ModifiedEntryPointName = EntryPointName;
 					FShaderCompilerOutput OriginalOutput = Output;
 					const int kMaxReasonableAttempts = 64;
 					for (int32 Attempt = 0; Attempt < kMaxReasonableAttempts; ++Attempt)
 					{
 						TArray<FString> RemoveErrors;
-						PreprocessedShaderSource = OriginalPreprocSource;
-						EntryPointName = OriginalEntryPointName;
-						if (RemoveUnusedInputs(PreprocessedShaderSource, ShaderInputs, EntryPointName, RemoveErrors))
+						ModifiedSource = PreprocessedShaderSource;
+						ModifiedEntryPointName = EntryPointName;
+						if (RemoveUnusedInputs(ModifiedSource, ShaderInputs, ModifiedEntryPointName, RemoveErrors))
 						{
 							Output = OriginalOutput;
-							if (!CompileAndProcessD3DShaderFXCExt(PreprocessedShaderSource, CompileFlags, Input, ShaderParameterParser, EntryPointName, ShaderProfile, true, FilteredErrors, Output))
+							if (!CompileAndProcessD3DShaderFXCExt(ModifiedSource, CompileFlags, Input, ShaderParameterParser, ModifiedEntryPointName, ShaderProfile, true, FilteredErrors, Output))
 							{
 								// if we failed to compile the shader, propagate the error up
 								return false;
@@ -919,6 +918,11 @@ static bool CompileAndProcessD3DShaderFXCExt(
 							break;
 						}
 					}
+
+					if (Input.ExtraSettings.bExtractShaderSource)
+					{
+						Output.OptionalFinalShaderSource = ModifiedSource;
+					}
 				}
 			}
 
@@ -1032,10 +1036,10 @@ static bool CompileAndProcessD3DShaderFXCExt(
 	return SUCCEEDED(Result);
 }
 
-bool CompileAndProcessD3DShaderFXC(FString& PreprocessedShaderSource,
+bool CompileAndProcessD3DShaderFXC(const FString& PreprocessedShaderSource,
 	const FShaderCompilerInput& Input,
 	const FShaderParameterParser& ShaderParameterParser,
-	FString& EntryPointName,
+	const FString& EntryPointName,
 	const TCHAR* ShaderProfile,
 	bool bSecondPassAferUnusedInputRemoval,
 	FShaderCompilerOutput& Output)
@@ -1103,18 +1107,29 @@ bool CompileAndProcessD3DShaderFXC(FString& PreprocessedShaderSource,
 	return bSuccess;
 }
 
-void CompileD3DShader(const FShaderCompilerInput& Input, FShaderCompilerOutput& Output, FShaderCompilerDefinitions& AdditionalDefines, const FString& WorkingDirectory, ELanguage Language)
+bool PreprocessD3DShader(
+	const FShaderCompilerInput& Input,
+	const FShaderCompilerEnvironment& Environment,
+	FShaderPreprocessOutput& Output,
+	ELanguage Language)
 {
-	TRACE_CPUPROFILER_EVENT_SCOPE(CompileD3DShader);
-	FString PreprocessedShaderSource;
-	const bool bUseDXC = !Input.CanCompileWithLegacyFxc();
-	const TCHAR* ShaderProfile = GetShaderProfileName(Language, Input.Target.Frequency, bUseDXC);
-
-	if(!ShaderProfile)
+	FShaderCompilerDefinitions AdditionalDefines;
+	if (Language == ELanguage::SM6)
 	{
-		Output.Errors.Add(FShaderCompilerError(*FString::Printf(TEXT("Unrecognized shader frequency %s"), GetShaderFrequencyString((EShaderFrequency)Input.Target.Frequency))));
-		return;
+		AdditionalDefines.SetDefine(TEXT("SM6_PROFILE"), 1);
 	}
+	else if (Language == ELanguage::SM5)
+	{
+		AdditionalDefines.SetDefine(TEXT("SM5_PROFILE"), 1);
+	}
+	else if (Language == ELanguage::ES3_1)
+	{
+		AdditionalDefines.SetDefine(TEXT("ES3_1_PROFILE"), 1);
+	}
+
+	FString& PreprocessedSource = Output.EditSource();
+	const bool bUseDXC = !Input.CanCompileWithLegacyFxc();
+
 
 	// Set additional defines.
 	AdditionalDefines.SetDefine(TEXT("COMPILER_HLSL"), 1);
@@ -1128,40 +1143,85 @@ void CompileD3DShader(const FShaderCompilerInput& Input, FShaderCompilerOutput&
 		AdditionalDefines.SetDefine(TEXT("PLATFORM_SUPPORTS_DIAGNOSTIC_BUFFER"), 1);
 		AdditionalDefines.SetDefine(TEXT("COMPILER_SUPPORTS_NOINLINE"), 1);
 
-		if (Input.Environment.CompilerFlags.Contains(CFLAG_InlineRayTracing))
+		if (Environment.CompilerFlags.Contains(CFLAG_InlineRayTracing))
 		{
 			AdditionalDefines.SetDefine(TEXT("PLATFORM_SUPPORTS_INLINE_RAY_TRACING"), 1);
 		}
 
-		if (Language == ELanguage::SM6 && Input.Environment.CompilerFlags.Contains(CFLAG_AllowRealTypes))
+		if (Language == ELanguage::SM6 && Environment.CompilerFlags.Contains(CFLAG_AllowRealTypes))
 		{
 			AdditionalDefines.SetDefine(TEXT("PLATFORM_SUPPORTS_REAL_TYPES"), 1);
 		}
 	}
 
-	const double StartPreprocessTime = FPlatformTime::Seconds();
 	
 	if (Input.bSkipPreprocessedCache)
 	{
-		if (!FFileHelper::LoadFileToString(PreprocessedShaderSource, *Input.VirtualSourceFilePath))
+		if (!FFileHelper::LoadFileToString(PreprocessedSource, *Input.VirtualSourceFilePath))
 		{
-			return;
+			return false;
 		}
 
 		// Remove const as we are on debug-only mode
-		CrossCompiler::CreateEnvironmentFromResourceTable(PreprocessedShaderSource, (FShaderCompilerEnvironment&)Input.Environment);
+		CrossCompiler::CreateEnvironmentFromResourceTable(PreprocessedSource, const_cast<FShaderCompilerEnvironment&>(Environment));
 	}
 	else
 	{
-		if (!PreprocessShader(PreprocessedShaderSource, Output, Input, AdditionalDefines))
+		if (!PreprocessShader(Output, Input, Environment, AdditionalDefines))
 		{
 			// The preprocessing stage will add any relevant errors.
-			return;
+			return false;
 		}
 	}
 
-	FString EntryPointName = Input.EntryPointName;
+	if (!Output.ParseAndModify(Input, Environment, TEXT("cbuffer")))
+	{
+		// The FShaderParameterParser will add any relevant errors.
+		return false;
+	}
+
+	// Only use UniformBuffer structs on SM6 until we can fully vet SM5
+	if (Language != ELanguage::SM6)
+	{
+		RemoveUniformBuffersFromSource(Environment, PreprocessedSource);
+	}
+
+	// Process TEXT macro.
+	TransformStringIntoCharacterArray(PreprocessedSource);
 
+	TArray<FString> FilteredErrors;
+	// Run the shader minifier
+	#if UE_D3D_SHADER_COMPILER_ALLOW_DEAD_CODE_REMOVAL
+	if (Environment.CompilerFlags.Contains(CFLAG_RemoveDeadCode))
+	{
+		if (!UE::ShaderCompilerCommon::RemoveDeadCode(PreprocessedSource, Input.EntryPointName, Output.EditErrors()))
+		{
+			return false;
+		}
+	}
+	#endif // UE_D3D_SHADER_COMPILER_ALLOW_DEAD_CODE_REMOVAL
+	return true;
+}
+
+
+void CompileD3DShader(const FShaderCompilerInput& Input, const FShaderPreprocessOutput& PreprocessOutput, FShaderCompilerOutput& Output, const FString& WorkingDirectory, ELanguage Language)
+{
+	TRACE_CPUPROFILER_EVENT_SCOPE(CompileD3DShader);
+
+	const bool bUseDXC = !Input.CanCompileWithLegacyFxc(); 
+	const TCHAR* ShaderProfile = GetShaderProfileName(Language, Input.Target.Frequency, bUseDXC);
+
+	if(!ShaderProfile)
+	{
+		Output.Errors.Add(FShaderCompilerError(*FString::Printf(TEXT("Unrecognized shader frequency %s"), GetShaderFrequencyString((EShaderFrequency)Input.Target.Frequency))));
+		return;
+	}
+
+	const FString* PreprocessedSource = &PreprocessOutput.GetSource();
+	const FString* EntryPoint = &Input.EntryPointName;
+	
+	FString ModifiedSource;
+	FString ModifiedEntryPoint;
 	if (Input.Environment.CompilerFlags.Contains(CFLAG_ForceRemoveUnusedInterpolators) && Input.Target.Frequency == SF_Vertex && Input.bCompilingForShaderPipeline)
 	{
 		// Always add SV_Position
@@ -1191,12 +1251,12 @@ void CompileD3DShader(const FShaderCompilerInput& Input, FShaderCompilerOutput&
 		Exceptions.AddUnique(TEXT("SV_CullDistance5"));
 		Exceptions.AddUnique(TEXT("SV_CullDistance6"));
 		Exceptions.AddUnique(TEXT("SV_CullDistance7"));
-		
-		// Write the preprocessed file out in case so we can debug issues on HlslParser
-		UE::ShaderCompilerCommon::DumpDebugShaderData(Input, PreprocessedShaderSource);
+
+		ModifiedSource = *PreprocessedSource;
+		ModifiedEntryPoint = *EntryPoint;
 
 		TArray<FString> Errors;
-		if (!RemoveUnusedOutputs(PreprocessedShaderSource, UsedOutputs, Exceptions, EntryPointName, Errors))
+		if (!RemoveUnusedOutputs(ModifiedSource, UsedOutputs, Exceptions, ModifiedEntryPoint, Errors))
 		{
 			UE_LOG(LogD3D11ShaderCompiler, Warning, TEXT("Failed to remove unused outputs from shader: %s"), *Input.GenerateShaderName());
 			for (const FString& ErrorReport : Errors)
@@ -1208,81 +1268,39 @@ void CompileD3DShader(const FShaderCompilerInput& Input, FShaderCompilerOutput&
 				Output.Errors.Add(NewError);
 			}
 		}
+		else
+		{
+			PreprocessedSource = &ModifiedSource;
+			EntryPoint = &ModifiedEntryPoint;
+		}
 	}
 
-	FShaderParameterParser ShaderParameterParser(Input.Environment.CompilerFlags, TEXT("cbuffer"));
-	if (!ShaderParameterParser.ParseAndModify(Input, Output.Errors, PreprocessedShaderSource))
-	{
-		// The FShaderParameterParser will add any relevant errors.
-		return;
-	}
-
-	// Only use UniformBuffer structs on SM6 until we can fully vet SM5
-	if (Language != ELanguage::SM6)
-	{
-		RemoveUniformBuffersFromSource(Input.Environment, PreprocessedShaderSource);
-	}
-
-	// Process TEXT macro.
-	TransformStringIntoCharacterArray(PreprocessedShaderSource);
-
-	Output.PreprocessTime = FPlatformTime::Seconds() - StartPreprocessTime;
-
-	TArray<FString> FilteredErrors;
-
-	// Run the shader minifier
-	#if UE_D3D_SHADER_COMPILER_ALLOW_DEAD_CODE_REMOVAL
-	if (Input.Environment.CompilerFlags.Contains(CFLAG_RemoveDeadCode))
-	{
-		UE::ShaderCompilerCommon::RemoveDeadCode(PreprocessedShaderSource, EntryPointName, Output.Errors);
-	}
-	#endif // UE_D3D_SHADER_COMPILER_ALLOW_DEAD_CODE_REMOVAL
+	const FShaderParameterParser& ShaderParameterParser = PreprocessOutput.GetParameterParser();
 
 	const bool bSuccess = bUseDXC
-		? CompileAndProcessD3DShaderDXC(PreprocessedShaderSource, Input, ShaderParameterParser, EntryPointName, ShaderProfile, Language, false, Output)
-		: CompileAndProcessD3DShaderFXC(PreprocessedShaderSource, Input, ShaderParameterParser, EntryPointName, ShaderProfile, false, Output);
+		? CompileAndProcessD3DShaderDXC(*PreprocessedSource, Input, ShaderParameterParser, *EntryPoint, ShaderProfile, Language, false, Output)
+		: CompileAndProcessD3DShaderFXC(*PreprocessedSource, Input, ShaderParameterParser, *EntryPoint, ShaderProfile, false, Output);
 
 	if (!bSuccess && !Output.Errors.Num())
 	{
 		Output.Errors.Add(TEXT("Compile failed without errors!"));
 	}
 
+	ShaderParameterParser.ValidateShaderParameterTypes(Input, Output);
+
 	const bool bDirectCompile = FParse::Param(FCommandLine::Get(), TEXT("directcompile"));
 	if (bDirectCompile)
 	{
-		for (const auto& Error : Output.Errors)
+		for (const FShaderCompilerError& Error : Output.Errors)
 		{
 			FPlatformMisc::LowLevelOutputDebugStringf(TEXT("%s\n"), *Error.GetErrorStringWithLineMarker());
 		}
 	}
 
-	ShaderParameterParser.ValidateShaderParameterTypes(Input, Output);
-
-	if (Input.ExtraSettings.bExtractShaderSource)
+	// Only set OptionalFinalShaderSource if it wasn't already set by the compile function 
+	// (FXC path can further modify the preprocessed source during unused output removal)
+	if (Input.ExtraSettings.bExtractShaderSource && Output.OptionalFinalShaderSource.IsEmpty())
 	{
-		Output.OptionalFinalShaderSource = PreprocessedShaderSource;
+		Output.OptionalFinalShaderSource = *PreprocessedSource;
 	}
 }
-
-void CompileShader_Windows(const FShaderCompilerInput& Input,FShaderCompilerOutput& Output,const FString& WorkingDirectory, ELanguage Language)
-{
-	FShaderCompilerDefinitions AdditionalDefines;
-	if (Language == ELanguage::SM6)
-	{
-		AdditionalDefines.SetDefine(TEXT("SM6_PROFILE"), 1);
-	}
-	else if (Language == ELanguage::SM5)
-	{
-		AdditionalDefines.SetDefine(TEXT("SM5_PROFILE"), 1);
-	}
-	else if (Language == ELanguage::ES3_1)
-	{
-		AdditionalDefines.SetDefine(TEXT("ES3_1_PROFILE"), 1);
-	}
-	else
-	{
-		checkf(0, TEXT("Unknown ELanguage %d"), (int32)Language);
-	}
-
-	CompileD3DShader(Input, Output, AdditionalDefines, WorkingDirectory, Language);
-}
diff --git a/Engine/Source/Developer/Windows/ShaderFormatD3D/Private/D3DShaderCompilerDXC.cpp b/Engine/Source/Developer/Windows/ShaderFormatD3D/Private/D3DShaderCompilerDXC.cpp
--- a/Engine/Source/Developer/Windows/ShaderFormatD3D/Private/D3DShaderCompilerDXC.cpp
+++ b/Engine/Source/Developer/Windows/ShaderFormatD3D/Private/D3DShaderCompilerDXC.cpp
@@ -866,7 +866,7 @@ static bool DXCRewriteWrapper(
 }
 
 // Generate the dumped usf file; call the D3D compiler, gather reflection information and generate the output data
-bool CompileAndProcessD3DShaderDXC(FString& PreprocessedShaderSource,
+bool CompileAndProcessD3DShaderDXC(const FString& PreprocessedShaderSource,
 	const FShaderCompilerInput& Input,
 	const FShaderParameterParser& ShaderParameterParser,
 	const FString& EntryPointName,
@@ -907,8 +907,6 @@ bool CompileAndProcessD3DShaderDXC(FString& PreprocessedShaderSource,
 		}
 	}
 
-	UE::ShaderCompilerCommon::DumpDebugShaderData(Input, PreprocessedShaderSource);
-
 	FDxcArguments Args
 	(
 		Input,
diff --git a/Engine/Source/Developer/Windows/ShaderFormatD3D/Private/ShaderFormatD3D.cpp b/Engine/Source/Developer/Windows/ShaderFormatD3D/Private/ShaderFormatD3D.cpp
--- a/Engine/Source/Developer/Windows/ShaderFormatD3D/Private/ShaderFormatD3D.cpp
+++ b/Engine/Source/Developer/Windows/ShaderFormatD3D/Private/ShaderFormatD3D.cpp
@@ -13,7 +13,7 @@ static FName NAME_PCD3D_SM5(TEXT("PCD3D_SM5"));
 static FName NAME_PCD3D_ES3_1(TEXT("PCD3D_ES31"));
 static FName NAME_D3D_ES3_1_HOLOLENS(TEXT("D3D_ES3_1_HOLOLENS"));
 
-class FShaderFormatD3D : public IShaderFormat
+class FShaderFormatD3D : public UE::ShaderCompilerCommon::FBaseShaderFormat 
 {
 	enum EVersion
 	{
@@ -26,11 +26,6 @@ class FShaderFormatD3D : public IShaderFormat
 		UE_SHADER_D3D_ES3_1_HOLOLENS_VER = UE_SHADER_PCD3D_ES3_1_VER,
 	};
 
-	void CheckFormat(FName Format) const
-	{
-		check(Format == NAME_PCD3D_SM6 || Format == NAME_PCD3D_SM5 || Format == NAME_PCD3D_ES3_1 || Format == NAME_D3D_ES3_1_HOLOLENS);
-	}
-
 	uint32 DxcVersionHash = 0;
 
 public:
@@ -60,7 +55,6 @@ public:
 
 	virtual uint32 GetVersion(FName Format) const override
 	{
-		CheckFormat(Format);
 		if (Format == NAME_PCD3D_SM6)
 		{
 			uint32 ShaderModelHash = GetVersionHash(UE_SHADER_PCD3D_SM6_VER);
@@ -107,31 +101,46 @@ public:
 		OutFormats.Add(NAME_D3D_ES3_1_HOLOLENS);
 	}
 
-	virtual void CompileShader(FName Format, const struct FShaderCompilerInput& Input, struct FShaderCompilerOutput& Output,const FString& WorkingDirectory) const
+	ELanguage FormatToLanguage(FName Format) const
 	{
-		CheckFormat(Format);
 		if (Format == NAME_PCD3D_SM6)
 		{
-			CompileShader_Windows(Input, Output, WorkingDirectory, ELanguage::SM6);
-		}
-		else if(Format == NAME_PCD3D_SM5)
-		{
-			CompileShader_Windows(Input, Output, WorkingDirectory, ELanguage::SM5);
+			return ELanguage::SM6;
 		}
-		else if (Format == NAME_PCD3D_ES3_1)
+		else if (Format == NAME_PCD3D_SM5)
 		{
-			CompileShader_Windows(Input, Output, WorkingDirectory, ELanguage::ES3_1);
+			return ELanguage::SM5;
 		}
-		else if (Format == NAME_D3D_ES3_1_HOLOLENS)
+		else if (Format == NAME_PCD3D_ES3_1 || Format == NAME_D3D_ES3_1_HOLOLENS)
 		{
-			CompileShader_Windows(Input, Output, WorkingDirectory, ELanguage::ES3_1);
+			return ELanguage::ES3_1;
 		}
 		else
 		{
 			checkf(0, TEXT("Unknown format %s"), *Format.ToString());
+			return ELanguage::Invalid;
 		}
 	}
 
+	virtual bool PreprocessShader(const struct FShaderCompilerInput& Input, const struct FShaderCompilerEnvironment& MergedEnvironment, class FShaderPreprocessOutput& PreprocessOutput) const
+	{
+		return PreprocessD3DShader(Input, MergedEnvironment, PreprocessOutput, FormatToLanguage(Input.ShaderFormat));
+	}
+
+	virtual void CompilePreprocessedShader(
+		const struct FShaderCompilerInput& Input, 
+		const class FShaderPreprocessOutput& PreprocessOutput, 
+		struct FShaderCompilerOutput& Output,
+		const FString& WorkingDirectory) const
+	{
+		CompileD3DShader(Input, PreprocessOutput, Output, WorkingDirectory, FormatToLanguage(Input.ShaderFormat));
+	}
+
+	virtual bool SupportsIndependentPreprocessing() const
+	{
+		return true;
+	}
+
 	void AddShaderTargetDefines(FShaderCompilerInput& Input, uint32 ShaderTargetMajor, uint32 ShaderTargetMinor) const
 	{
 		// Inserting our own versions of these defines since we preprocess our shader source before we actually use something that defines them.
@@ -141,7 +150,6 @@ public:
 
 	void ModifyShaderCompilerInput(FShaderCompilerInput& Input) const final
 	{
-		CheckFormat(Input.ShaderFormat);
 		if (Input.ShaderFormat == NAME_PCD3D_SM6 || Input.IsRayTracingShader())
 		{
 			Input.Environment.SetDefine(TEXT("SM6_PROFILE"), 1);
diff --git a/Engine/Source/Developer/Windows/ShaderFormatD3D/Private/ShaderFormatD3D.h b/Engine/Source/Developer/Windows/ShaderFormatD3D/Private/ShaderFormatD3D.h
--- a/Engine/Source/Developer/Windows/ShaderFormatD3D/Private/ShaderFormatD3D.h
+++ b/Engine/Source/Developer/Windows/ShaderFormatD3D/Private/ShaderFormatD3D.h
@@ -20,22 +20,34 @@ enum class ELanguage
 	SM5,
 	SM6,
 	ES3_1,
+	Invalid,
 };
 
-void CompileShader_Windows(const struct FShaderCompilerInput& Input,struct FShaderCompilerOutput& Output,const class FString& WorkingDirectory, ELanguage Language);
+bool PreprocessD3DShader(
+	const struct FShaderCompilerInput& Input,
+	const struct FShaderCompilerEnvironment& MergedEnvironment,
+	class FShaderPreprocessOutput& PreprocessOutput,
+	ELanguage Language);
+
+void CompileD3DShader(
+	const struct FShaderCompilerInput& Input,
+	const class FShaderPreprocessOutput& PreprocessOutput,
+	struct FShaderCompilerOutput& Output,
+	const class FString& WorkingDirectory,
+	ELanguage Language);
 
 /**
  * @param bSecondPassAferUnusedInputRemoval whether we're compiling the shader second time, after having removed the unused inputs discovered in the first pass
  */
-bool CompileAndProcessD3DShaderFXC(FString& PreprocessedShaderSource,
+bool CompileAndProcessD3DShaderFXC(const FString& PreprocessedShaderSource,
 	const FShaderCompilerInput& Input,
 	const FShaderParameterParser& ShaderParameterParser,
-	FString& EntryPointName,
+	const FString& EntryPointName,
 	const TCHAR* ShaderProfile,
 	bool bSecondPassAferUnusedInputRemoval,
 	FShaderCompilerOutput& Output);
 
-bool CompileAndProcessD3DShaderDXC(FString& PreprocessedShaderSource,
+bool CompileAndProcessD3DShaderDXC(const FString& PreprocessedShaderSource,
 	const FShaderCompilerInput& Input,
 	const FShaderParameterParser& ShaderParameterParser,
 	const FString& EntryPointName,
