Fixed mismatched copy constructor/assignment operator in FixedSizeString.

diff --git a/Engine/Source/Developer/Windows/LiveCoding/Private/External/LC_FixedSizeString.cpp b/Engine/Source/Developer/Windows/LiveCoding/Private/External/LC_FixedSizeString.cpp
--- a/Engine/Source/Developer/Windows/LiveCoding/Private/External/LC_FixedSizeString.cpp
+++ b/Engine/Source/Developer/Windows/LiveCoding/Private/External/LC_FixedSizeString.cpp
@@ -27,6 +27,14 @@ FixedSizeString<T, N>::FixedSizeString(const T* other)
 }
 
 
+template <typename T, unsigned int N>
+FixedSizeString<T, N>::FixedSizeString(const FixedSizeString<T, N>& other)
+	: m_length(other.m_length)
+{
+	memcpy(m_data, other.m_data, sizeof(T)*(other.m_length + 1u));
+}
+
+
 template <typename T, unsigned int N>
 FixedSizeString<T, N>::FixedSizeString(const T* other, unsigned int length)
 	: m_length(length)
diff --git a/Engine/Source/Developer/Windows/LiveCoding/Private/External/LC_FixedSizeString.h b/Engine/Source/Developer/Windows/LiveCoding/Private/External/LC_FixedSizeString.h
--- a/Engine/Source/Developer/Windows/LiveCoding/Private/External/LC_FixedSizeString.h
+++ b/Engine/Source/Developer/Windows/LiveCoding/Private/External/LC_FixedSizeString.h
@@ -15,6 +15,7 @@ public:
 
 	FixedSizeString(void);
 	explicit FixedSizeString(const T* other);
+	FixedSizeString(const FixedSizeString<T, N>& other);
 	FixedSizeString(const T* other, unsigned int length);
 
 	FixedSizeString<T, N>& operator=(const FixedSizeString<T, N>& other);
