Removing the cvar lookup string for shader dead code removal in all shader formats. This cvar can't be loaded on SCW startup so it was causing a hash mismatch.

diff --git a/Engine/Source/Developer/Apple/MetalShaderFormat/Private/MetalShaderFormat.cpp b/Engine/Source/Developer/Apple/MetalShaderFormat/Private/MetalShaderFormat.cpp
--- a/Engine/Source/Developer/Apple/MetalShaderFormat/Private/MetalShaderFormat.cpp
+++ b/Engine/Source/Developer/Apple/MetalShaderFormat/Private/MetalShaderFormat.cpp
@@ -325,13 +325,7 @@ uint32 GetMetalFormatVersion(FName Format)
 	uint32 Result = Version.Raw;
 
 #if UE_METAL_SHADER_COMPILER_ALLOW_DEAD_CODE_REMOVAL
-	{
-		static const auto CVar = IConsoleManager::Get().FindConsoleVariable(TEXT("r.Shader.RemoveDeadCode"));
-		if (CVar && CVar->GetInt() != 0)
-		{
-			Result = HashCombine(Result, 0x75E2FE85);
-		}
-	}
+	Result = HashCombine(Result, 0x75E2FE85);
 #endif // UE_METAL_SHADER_COMPILER_ALLOW_DEAD_CODE_REMOVAL
 
 	return Result;
diff --git a/Engine/Source/Developer/ShaderFormatOpenGL/Private/ShaderFormatOpenGL.cpp b/Engine/Source/Developer/ShaderFormatOpenGL/Private/ShaderFormatOpenGL.cpp
--- a/Engine/Source/Developer/ShaderFormatOpenGL/Private/ShaderFormatOpenGL.cpp
+++ b/Engine/Source/Developer/ShaderFormatOpenGL/Private/ShaderFormatOpenGL.cpp
@@ -42,13 +42,7 @@ public:
 		uint32 Version = ((HLSLCC_VersionMinor & 0xff) << 8) | (GLSLVersion & 0xff);
 
 	#if UE_OPENGL_SHADER_COMPILER_ALLOW_DEAD_CODE_REMOVAL
-		{
-			static const auto CVar = IConsoleManager::Get().FindConsoleVariable(TEXT("r.Shader.RemoveDeadCode"));
-			if (CVar && CVar->GetInt() != 0)
-			{
-				Version = HashCombine(Version, 0x75E2FE85);
-			}
-		}
+		Version = HashCombine(Version, 0x75E2FE85);
 	#endif // UE_OPENGL_SHADER_COMPILER_ALLOW_DEAD_CODE_REMOVAL
 
 		return Version;
diff --git a/Engine/Source/Developer/VulkanShaderFormat/Private/VulkanShaderFormat.cpp b/Engine/Source/Developer/VulkanShaderFormat/Private/VulkanShaderFormat.cpp
--- a/Engine/Source/Developer/VulkanShaderFormat/Private/VulkanShaderFormat.cpp
+++ b/Engine/Source/Developer/VulkanShaderFormat/Private/VulkanShaderFormat.cpp
@@ -61,13 +61,7 @@ public:
 	#endif
 
 	#if UE_VULKAN_SHADER_COMPILER_ALLOW_DEAD_CODE_REMOVAL
-		{
-			static const auto CVar = IConsoleManager::Get().FindConsoleVariable(TEXT("r.Shader.RemoveDeadCode"));
-			if (CVar && CVar->GetInt() != 0)
-			{
-				Version = HashCombine(Version, 0x75E2FE85);
-			}
-		}
+		Version = HashCombine(Version, 0x75E2FE85);
 	#endif // UE_VULKAN_SHADER_COMPILER_ALLOW_DEAD_CODE_REMOVAL
 
 		return Version;
diff --git a/Engine/Source/Developer/Windows/ShaderFormatD3D/Private/ShaderFormatD3D.cpp b/Engine/Source/Developer/Windows/ShaderFormatD3D/Private/ShaderFormatD3D.cpp
--- a/Engine/Source/Developer/Windows/ShaderFormatD3D/Private/ShaderFormatD3D.cpp
+++ b/Engine/Source/Developer/Windows/ShaderFormatD3D/Private/ShaderFormatD3D.cpp
@@ -39,13 +39,7 @@ public:
 		uint32 VersionHash = GetTypeHash(InVersion);
 
 	#if UE_D3D_SHADER_COMPILER_ALLOW_DEAD_CODE_REMOVAL
-		{
-			static const auto CVar = IConsoleManager::Get().FindConsoleVariable(TEXT("r.Shader.RemoveDeadCode"));
-			if (CVar && CVar->GetInt() != 0)
-			{
-				VersionHash = HashCombine(VersionHash, 0x75E2FE85);
-			}
-		}
+		VersionHash = HashCombine(VersionHash, 0x75E2FE85);
 	#endif // UE_D3D_SHADER_COMPILER_ALLOW_DEAD_CODE_REMOVAL
 
 		return HashCombine(BaseHash, VersionHash);
