Restricting PLATFORM_SUPPORTS_STATIC_SAMPLERS to SM6 and RayTracing shaders and making sure that PC DX12 SM6 requires Resource Binding Tier 2 or higher.
We don't want to restrict other shaders to resource binding tier 2 or higher even for CFLAG_ForceDXC or CFLAG_WaveOperations.
Invalidating PCD3D_SM5 and PCD3D_SM6 shaders.

diff --git a/Engine/Source/Developer/Windows/ShaderFormatD3D/Private/D3DShaderCompiler.cpp b/Engine/Source/Developer/Windows/ShaderFormatD3D/Private/D3DShaderCompiler.cpp
--- a/Engine/Source/Developer/Windows/ShaderFormatD3D/Private/D3DShaderCompiler.cpp
+++ b/Engine/Source/Developer/Windows/ShaderFormatD3D/Private/D3DShaderCompiler.cpp
@@ -1143,7 +1143,10 @@ bool PreprocessD3DShader(
 	{
 		AdditionalDefines.SetDefine(TEXT("PLATFORM_SUPPORTS_CALLABLE_SHADERS"), 1);
 		AdditionalDefines.SetDefine(TEXT("PLATFORM_SUPPORTS_SM6_0_WAVE_OPERATIONS"), 1);
-		AdditionalDefines.SetDefine(TEXT("PLATFORM_SUPPORTS_STATIC_SAMPLERS"), 1);
+		if (Language == ELanguage::SM6 || Input.IsRayTracingShader())
+		{
+			AdditionalDefines.SetDefine(TEXT("PLATFORM_SUPPORTS_STATIC_SAMPLERS"), 1);
+		}
 		AdditionalDefines.SetDefine(TEXT("PLATFORM_SUPPORTS_DIAGNOSTIC_BUFFER"), 1);
 		AdditionalDefines.SetDefine(TEXT("COMPILER_SUPPORTS_NOINLINE"), 1);
 
diff --git a/Engine/Source/Developer/Windows/ShaderFormatD3D/Private/ShaderFormatD3D.cpp b/Engine/Source/Developer/Windows/ShaderFormatD3D/Private/ShaderFormatD3D.cpp
--- a/Engine/Source/Developer/Windows/ShaderFormatD3D/Private/ShaderFormatD3D.cpp
+++ b/Engine/Source/Developer/Windows/ShaderFormatD3D/Private/ShaderFormatD3D.cpp
@@ -19,8 +19,8 @@ class FShaderFormatD3D : public UE::ShaderCompilerCommon::FBaseShaderFormat
 		UE_SHADER_PCD3D_SHARED_VER = 5,
 
 		/** Version for shader format, this becomes part of the DDC key. */
-		UE_SHADER_PCD3D_SM6_VER = 7,
-		UE_SHADER_PCD3D_SM5_VER = 12,
+		UE_SHADER_PCD3D_SM6_VER = 8,
+		UE_SHADER_PCD3D_SM5_VER = 13,
 		UE_SHADER_PCD3D_ES3_1_VER = 8,
 	};
 
diff --git a/Engine/Source/Runtime/D3D12RHI/Private/Windows/WindowsD3D12Device.cpp b/Engine/Source/Runtime/D3D12RHI/Private/Windows/WindowsD3D12Device.cpp
--- a/Engine/Source/Runtime/D3D12RHI/Private/Windows/WindowsD3D12Device.cpp
+++ b/Engine/Source/Runtime/D3D12RHI/Private/Windows/WindowsD3D12Device.cpp
@@ -422,6 +422,10 @@ inline ERHIFeatureLevel::Type FindMaxRHIFeatureLevel(D3D_FEATURE_LEVEL InMaxFeat
 		{
 			bHighEnoughBindingTier = ResourceBindingTier >= D3D12_RESOURCE_BINDING_TIER_3;
 		}
+		else
+		{
+			bHighEnoughBindingTier = ResourceBindingTier >= D3D12_RESOURCE_BINDING_TIER_2;
+		}
 
 		if (bSupportsWaveOps && bHighEnoughBindingTier && bSupportsAtomic64)
 		{
