[shaders]
- add new fields to FShaderCompilerOutput which should be populated with the shader source and entry point name if modified by a backend compile process (in a format implementing the independent preprocessing API). this is intended to supercede "OptionalFinalShaderSource" for such shader formats, as well as being used as the source for preprocessed debug dumps when enabled.
- modify the PCD3D shader format to use these fields as intended (and so fix the debug output for pipelines which remove unused interpolators to match the final modified source).

diff --git a/Engine/Source/Developer/ShaderCompilerCommon/Private/ShaderCompilerCommon.cpp b/Engine/Source/Developer/ShaderCompilerCommon/Private/ShaderCompilerCommon.cpp
--- a/Engine/Source/Developer/ShaderCompilerCommon/Private/ShaderCompilerCommon.cpp
+++ b/Engine/Source/Developer/ShaderCompilerCommon/Private/ShaderCompilerCommon.cpp
@@ -1403,7 +1403,9 @@ namespace UE::ShaderCompilerCommon
 {
 	void FBaseShaderFormat::OutputDebugData(const FString& InputHash, const FShaderCompilerInput& Input, const FShaderPreprocessOutput& PreprocessOutput, const FShaderCompilerOutput& Output) const
 	{
-		DumpDebugShaderData(Input, Input.bCachePreprocessed ? PreprocessOutput.GetOriginalSource() : PreprocessOutput.GetSource());
+		const FString& SourceToDump = Output.ModifiedShaderSource.IsEmpty() ? PreprocessOutput.GetUnstrippedSource() : Output.ModifiedShaderSource;
+		DumpDebugShaderData(Input, SourceToDump);
+		
 		FString OutputHashFileName = FPaths::Combine(Input.DumpDebugInfoPath, TEXT("OutputHash.txt"));
 		FFileHelper::SaveStringToFile(Output.OutputHash.ToString(), *OutputHashFileName, FFileHelper::EEncodingOptions::ForceAnsi);
 
diff --git a/Engine/Source/Developer/Windows/ShaderFormatD3D/Private/D3DShaderCompiler.cpp b/Engine/Source/Developer/Windows/ShaderFormatD3D/Private/D3DShaderCompiler.cpp
--- a/Engine/Source/Developer/Windows/ShaderFormatD3D/Private/D3DShaderCompiler.cpp
+++ b/Engine/Source/Developer/Windows/ShaderFormatD3D/Private/D3DShaderCompiler.cpp
@@ -549,16 +549,17 @@ static bool CompileErrorsContainInternalError(ID3DBlob* Errors)
 
 // Generate the dumped usf file; call the D3D compiler, gather reflection information and generate the output data
 static bool CompileAndProcessD3DShaderFXCExt(
-	const FString& PreprocessedShaderSource,
+	const FShaderPreprocessOutput& PreprocessOutput,
 	uint32 CompileFlags,
 	const FShaderCompilerInput& Input,
 	const FShaderParameterParser& ShaderParameterParser,
-	const FString& EntryPointName,
 	const TCHAR* ShaderProfile, bool bSecondPassAferUnusedInputRemoval,
 	TArray<FString>& FilteredErrors, FShaderCompilerOutput& Output)
 {
 	TRACE_CPUPROFILER_EVENT_SCOPE(CompileAndProcessD3DShaderFXCExt);
 
+	const FString& PreprocessedShaderSource = Output.ModifiedShaderSource.IsEmpty() ? PreprocessOutput.GetSource() : Output.ModifiedShaderSource;
+	const FString& EntryPointName = Output.ModifiedEntryPointName.IsEmpty() ? Input.EntryPointName : Output.ModifiedEntryPointName;
 	auto AnsiSourceFile = StringCast<ANSICHAR>(*PreprocessedShaderSource);
 
 	bool bDumpDebugInfo = Input.DumpDebugInfoEnabled();
@@ -851,19 +852,19 @@ static bool CompileAndProcessD3DShaderFXCExt(
 					// Since unused inputs are passed to the next stage, that will cause us to generate a vertex shader that does not output B, but our pixel shader will still be expecting B on input,
 					// as it was rewritten based on the pass1 results.
 
-					FString ModifiedSource = PreprocessedShaderSource;
-					FString ModifiedEntryPointName = EntryPointName;
 					FShaderCompilerOutput OriginalOutput = Output;
 					const int kMaxReasonableAttempts = 64;
 					for (int32 Attempt = 0; Attempt < kMaxReasonableAttempts; ++Attempt)
 					{
 						TArray<FString> RemoveErrors;
-						ModifiedSource = PreprocessedShaderSource;
-						ModifiedEntryPointName = EntryPointName;
-						if (RemoveUnusedInputs(ModifiedSource, ShaderInputs, ModifiedEntryPointName, RemoveErrors))
+						FString ModifiedShaderSource = PreprocessOutput.GetSource();
+						FString ModifiedEntryPointName = Input.EntryPointName;
+						if (RemoveUnusedInputs(ModifiedShaderSource, ShaderInputs, ModifiedEntryPointName, RemoveErrors))
 						{
 							Output = OriginalOutput;
-							if (!CompileAndProcessD3DShaderFXCExt(ModifiedSource, CompileFlags, Input, ShaderParameterParser, ModifiedEntryPointName, ShaderProfile, true, FilteredErrors, Output))
+							Output.ModifiedShaderSource = MoveTemp(ModifiedShaderSource);
+							Output.ModifiedEntryPointName = MoveTemp(ModifiedEntryPointName);
+							if (!CompileAndProcessD3DShaderFXCExt(PreprocessOutput, CompileFlags, Input, ShaderParameterParser, ShaderProfile, true, FilteredErrors, Output))
 							{
 								// if we failed to compile the shader, propagate the error up
 								return false;
@@ -922,11 +923,6 @@ static bool CompileAndProcessD3DShaderFXCExt(
 							break;
 						}
 					}
-
-					if (Input.ExtraSettings.bExtractShaderSource)
-					{
-						Output.OptionalFinalShaderSource = ModifiedSource;
-					}
 				}
 			}
 
@@ -1040,10 +1036,9 @@ static bool CompileAndProcessD3DShaderFXCExt(
 	return SUCCEEDED(Result);
 }
 
-bool CompileAndProcessD3DShaderFXC(const FString& PreprocessedShaderSource,
+bool CompileAndProcessD3DShaderFXC(const FShaderPreprocessOutput& PreprocessOutput,
 	const FShaderCompilerInput& Input,
 	const FShaderParameterParser& ShaderParameterParser,
-	const FString& EntryPointName,
 	const TCHAR* ShaderProfile,
 	bool bSecondPassAferUnusedInputRemoval,
 	FShaderCompilerOutput& Output)
@@ -1080,7 +1075,7 @@ bool CompileAndProcessD3DShaderFXC(const FString& PreprocessedShaderSource,
 		});
 
 	TArray<FString> FilteredErrors;
-	const bool bSuccess = CompileAndProcessD3DShaderFXCExt(PreprocessedShaderSource, CompileFlags, Input, ShaderParameterParser, EntryPointName, ShaderProfile, false, FilteredErrors, Output);
+	const bool bSuccess = CompileAndProcessD3DShaderFXCExt(PreprocessOutput, CompileFlags, Input, ShaderParameterParser, ShaderProfile, false, FilteredErrors, Output);
 
 	// Process errors
 	for (int32 ErrorIndex = 0; ErrorIndex < FilteredErrors.Num(); ErrorIndex++)
@@ -1224,11 +1219,6 @@ void CompileD3DShader(const FShaderCompilerInput& Input, const FShaderPreprocess
 		return;
 	}
 
-	const FString* PreprocessedSource = &PreprocessOutput.GetSource();
-	const FString* EntryPoint = &Input.EntryPointName;
-	
-	FString ModifiedSource;
-	FString ModifiedEntryPoint;
 	if (Input.Environment.CompilerFlags.Contains(CFLAG_ForceRemoveUnusedInterpolators) && Input.Target.Frequency == SF_Vertex && Input.bCompilingForShaderPipeline)
 	{
 		// Always add SV_Position
@@ -1259,11 +1249,11 @@ void CompileD3DShader(const FShaderCompilerInput& Input, const FShaderPreprocess
 		Exceptions.AddUnique(TEXT("SV_CullDistance6"));
 		Exceptions.AddUnique(TEXT("SV_CullDistance7"));
 
-		ModifiedSource = *PreprocessedSource;
-		ModifiedEntryPoint = *EntryPoint;
+		Output.ModifiedShaderSource = PreprocessOutput.GetSource();
+		Output.ModifiedEntryPointName = Input.EntryPointName;
 
 		TArray<FString> Errors;
-		if (!RemoveUnusedOutputs(ModifiedSource, UsedOutputs, Exceptions, ModifiedEntryPoint, Errors))
+		if (!RemoveUnusedOutputs(Output.ModifiedShaderSource, UsedOutputs, Exceptions, Output.ModifiedEntryPointName, Errors))
 		{
 			UE_LOG(LogD3D11ShaderCompiler, Warning, TEXT("Failed to remove unused outputs from shader: %s"), *Input.GenerateShaderName());
 			for (const FString& ErrorReport : Errors)
@@ -1275,18 +1265,13 @@ void CompileD3DShader(const FShaderCompilerInput& Input, const FShaderPreprocess
 				Output.Errors.Add(NewError);
 			}
 		}
-		else
-		{
-			PreprocessedSource = &ModifiedSource;
-			EntryPoint = &ModifiedEntryPoint;
-		}
 	}
 
 	const FShaderParameterParser& ShaderParameterParser = PreprocessOutput.GetParameterParser();
 
 	const bool bSuccess = bUseDXC
-		? CompileAndProcessD3DShaderDXC(*PreprocessedSource, Input, ShaderParameterParser, *EntryPoint, ShaderProfile, Language, false, Output)
-		: CompileAndProcessD3DShaderFXC(*PreprocessedSource, Input, ShaderParameterParser, *EntryPoint, ShaderProfile, false, Output);
+		? CompileAndProcessD3DShaderDXC(PreprocessOutput, Input, ShaderParameterParser, ShaderProfile, Language, false, Output)
+		: CompileAndProcessD3DShaderFXC(PreprocessOutput, Input, ShaderParameterParser, ShaderProfile, false, Output);
 
 	if (!bSuccess && !Output.Errors.Num())
 	{
@@ -1303,11 +1288,4 @@ void CompileD3DShader(const FShaderCompilerInput& Input, const FShaderPreprocess
 			FPlatformMisc::LowLevelOutputDebugStringf(TEXT("%s\n"), *Error.GetErrorStringWithLineMarker());
 		}
 	}
-
-	// Only set OptionalFinalShaderSource if it wasn't already set by the compile function 
-	// (FXC path can further modify the preprocessed source during unused output removal)
-	if (Input.ExtraSettings.bExtractShaderSource && Output.OptionalFinalShaderSource.IsEmpty())
-	{
-		Output.OptionalFinalShaderSource = *PreprocessedSource;
-	}
 }
diff --git a/Engine/Source/Developer/Windows/ShaderFormatD3D/Private/D3DShaderCompilerDXC.cpp b/Engine/Source/Developer/Windows/ShaderFormatD3D/Private/D3DShaderCompilerDXC.cpp
--- a/Engine/Source/Developer/Windows/ShaderFormatD3D/Private/D3DShaderCompilerDXC.cpp
+++ b/Engine/Source/Developer/Windows/ShaderFormatD3D/Private/D3DShaderCompilerDXC.cpp
@@ -9,6 +9,7 @@
 #include "Misc/FileHelper.h"
 #include "HAL/FileManager.h"
 #include "Serialization/MemoryWriter.h"
+#include "ShaderPreprocessTypes.h"
 #include "RayTracingDefinitions.h"
 
 DEFINE_LOG_CATEGORY_STATIC(LogD3D12ShaderCompiler, Log, All);
@@ -866,10 +867,9 @@ static bool DXCRewriteWrapper(
 }
 
 // Generate the dumped usf file; call the D3D compiler, gather reflection information and generate the output data
-bool CompileAndProcessD3DShaderDXC(const FString& PreprocessedShaderSource,
+bool CompileAndProcessD3DShaderDXC(const FShaderPreprocessOutput& PreprocessOutput,
 	const FShaderCompilerInput& Input,
 	const FShaderParameterParser& ShaderParameterParser,
-	const FString& EntryPointName,
 	const TCHAR* ShaderProfile,
 	ELanguage Language,
 	bool bProcessingSecondTime,
@@ -877,6 +877,9 @@ bool CompileAndProcessD3DShaderDXC(const FString& PreprocessedShaderSource,
 {
 	TRACE_CPUPROFILER_EVENT_SCOPE(CompileAndProcessD3DShaderDXC);
 
+	const FString& PreprocessedShaderSource = Output.ModifiedShaderSource.IsEmpty() ? PreprocessOutput.GetSource() : Output.ModifiedShaderSource;
+	const FString& EntryPointName = Output.ModifiedEntryPointName.IsEmpty() ? Input.EntryPointName : Output.ModifiedEntryPointName;
+
 	auto AnsiSourceFile = StringCast<ANSICHAR>(*PreprocessedShaderSource);
 
 	const bool bIsRayTracingShader = Input.IsRayTracingShader();
diff --git a/Engine/Source/Developer/Windows/ShaderFormatD3D/Private/ShaderFormatD3D.h b/Engine/Source/Developer/Windows/ShaderFormatD3D/Private/ShaderFormatD3D.h
--- a/Engine/Source/Developer/Windows/ShaderFormatD3D/Private/ShaderFormatD3D.h
+++ b/Engine/Source/Developer/Windows/ShaderFormatD3D/Private/ShaderFormatD3D.h
@@ -39,18 +39,18 @@ void CompileD3DShader(
 /**
  * @param bSecondPassAferUnusedInputRemoval whether we're compiling the shader second time, after having removed the unused inputs discovered in the first pass
  */
-bool CompileAndProcessD3DShaderFXC(const FString& PreprocessedShaderSource,
+bool CompileAndProcessD3DShaderFXC(
+	const FShaderPreprocessOutput& PreprocessOutput,
 	const FShaderCompilerInput& Input,
 	const FShaderParameterParser& ShaderParameterParser,
-	const FString& EntryPointName,
 	const TCHAR* ShaderProfile,
 	bool bSecondPassAferUnusedInputRemoval,
 	FShaderCompilerOutput& Output);
 
-bool CompileAndProcessD3DShaderDXC(const FString& PreprocessedShaderSource,
+bool CompileAndProcessD3DShaderDXC(
+	const FShaderPreprocessOutput& PreprocessOutput,
 	const FShaderCompilerInput& Input,
 	const FShaderParameterParser& ShaderParameterParser,
-	const FString& EntryPointName,
 	const TCHAR* ShaderProfile,
 	ELanguage Language,
 	bool bProcessingSecondTime,
diff --git a/Engine/Source/Runtime/Engine/Private/Materials/MaterialShader.cpp b/Engine/Source/Runtime/Engine/Private/Materials/MaterialShader.cpp
--- a/Engine/Source/Runtime/Engine/Private/Materials/MaterialShader.cpp
+++ b/Engine/Source/Runtime/Engine/Private/Materials/MaterialShader.cpp
@@ -2335,7 +2335,7 @@ FShader* FMaterialShaderMap::ProcessCompilationResultsForSingleJob(FShaderCompil
 		const int32 Index = Algo::LowerBoundBy(GetMutableContent()->ShaderProcessedSource, Key, [](const FMaterialProcessedSource& Value) { return Value.Name; });
 		if (Index >= GetMutableContent()->ShaderProcessedSource.Num() || GetMutableContent()->ShaderProcessedSource[Index].Name != Key)
 		{
-			GetMutableContent()->ShaderProcessedSource.EmplaceAt(Index, Key, *CurrentJob.Output.OptionalFinalShaderSource);
+			GetMutableContent()->ShaderProcessedSource.EmplaceAt(Index, Key, *CurrentJob.GetFinalSource());
 		}
 	}
 
diff --git a/Engine/Source/Runtime/Engine/Private/ShaderCompiler/ShaderCompiler.cpp b/Engine/Source/Runtime/Engine/Private/ShaderCompiler/ShaderCompiler.cpp
--- a/Engine/Source/Runtime/Engine/Private/ShaderCompiler/ShaderCompiler.cpp
+++ b/Engine/Source/Runtime/Engine/Private/ShaderCompiler/ShaderCompiler.cpp
@@ -6824,9 +6824,13 @@ void GlobalBeginCompileShader(
 	checkf(Format, TEXT("Shader format %s cannot be found"), *ShaderFormatName.ToString());
 	Format->ModifyShaderCompilerInput(Input);
 
-	if (Format->SupportsIndependentPreprocessing() && CVarPreprocessedJobCache.GetValueOnAnyThread())
+	if (Format->SupportsIndependentPreprocessing())
 	{
-		Input.bCachePreprocessed = true;
+		Input.bIndependentPreprocessed = true;
+		if (CVarPreprocessedJobCache.GetValueOnAnyThread())
+		{
+			Input.bCachePreprocessed = true;
+		}
 	}
 
 	// Allow the GBuffer and other shader defines to cause dependend environment changes, but minimizing the #ifdef magic in the shaders, which
diff --git a/Engine/Source/Runtime/RenderCore/Private/ShaderCore.cpp b/Engine/Source/Runtime/RenderCore/Private/ShaderCore.cpp
--- a/Engine/Source/Runtime/RenderCore/Private/ShaderCore.cpp
+++ b/Engine/Source/Runtime/RenderCore/Private/ShaderCore.cpp
@@ -2899,6 +2899,7 @@ FArchive& operator<<(FArchive& Ar, FShaderCompilerInput& Input)
 	Ar << Input.bCompilingForShaderPipeline;
 	Ar << Input.bIncludeUsedOutputs;
 	Ar << Input.bCachePreprocessed;
+	Ar << Input.bIndependentPreprocessed;
 	Ar << Input.UsedOutputs;
 	Ar << Input.DumpDebugInfoRootPath;
 	Ar << Input.DumpDebugInfoPath;
@@ -3173,6 +3174,19 @@ void FShaderCompileJob::OnComplete()
 
 void FShaderCompileJob::SerializeWorkerOutput(FArchive& Ar)
 {
+	// Only serialize the modified source/entry point if either:
+	// (a) this job is using the preprocessed cache and requires debug output (since this is done in the cooker process the modified source is needed outside of SCW)
+	// or (b) the "extract shader source" setting is enabled (i.e. something upstream explicitly wants the final source passed to the compiler).
+	Output.bSerializeModifiedSource = (Input.bCachePreprocessed && Input.DumpDebugInfoEnabled()) || Input.ExtraSettings.bExtractShaderSource;
+
+	// edge case for backends which have implemented independent preprocessing API when the preprocessed cache is not enabled.
+	// if no modifications have occurred as part of the compile step, we still need a copy of the source back in the cooker
+	// if bExtractShaderSource is set, so explicitly serialize just that portion of the preprocess output struct here.
+	if (Input.ExtraSettings.bExtractShaderSource && Input.bIndependentPreprocessed && !Input.bCachePreprocessed && Output.ModifiedShaderSource.IsEmpty())
+	{
+		Ar << PreprocessOutput.EditSource();
+	}
+
 	Ar << Output;
 	bool bSucceededTemp = (bool)bSucceeded;
 	Ar << bSucceededTemp;
@@ -3189,6 +3203,23 @@ void FShaderCompileJob::SerializeWorkerInput(FArchive& Ar)
 	}
 }
 
+const FString& FShaderCompileJob::GetFinalSource() const
+{
+	 // if the backend supports independent preprocessing, any modifications to the source
+	// done as part of the compile step will be written to the "ModifiedShaderSource" field
+	if (Input.bIndependentPreprocessed)
+	{
+		// if there are no such modifications, return the "unstripped" version of the source code (with comments & line directives maintained),
+		// otherwise return whatever the final modified source is as input to the compiler by the backend.
+		return Output.ModifiedShaderSource.IsEmpty() ? PreprocessOutput.GetUnstrippedSource() : Output.ModifiedShaderSource;
+	}
+	else
+	{
+		// backends that do not implement the independent preprocessing API populate this field based on the bExtractShaderSource setting.
+		return Output.OptionalFinalShaderSource;
+	}
+}
+
 FShaderPipelineCompileJob::FShaderPipelineCompileJob(int32 NumStages)
 	: FShaderCommonCompileJob(Type, 0u, 0u, EShaderCompileJobPriority::Num)
 {
diff --git a/Engine/Source/Runtime/RenderCore/Public/ShaderCompilerCore.h b/Engine/Source/Runtime/RenderCore/Public/ShaderCompilerCore.h
--- a/Engine/Source/Runtime/RenderCore/Public/ShaderCompilerCore.h
+++ b/Engine/Source/Runtime/RenderCore/Public/ShaderCompilerCore.h
@@ -27,7 +27,7 @@ class FShaderPipelineCompileJob;
 // this is for the protocol, not the data, bump if FShaderCompilerInput or ProcessInputFromArchive changes.
 inline const int32 ShaderCompileWorkerInputVersion = 19;
 // this is for the protocol, not the data, bump if FShaderCompilerOutput or WriteToOutputArchive changes.
-inline const int32 ShaderCompileWorkerOutputVersion = 13;
+inline const int32 ShaderCompileWorkerOutputVersion = 14;
 // this is for the protocol, not the data.
 inline const int32 ShaderCompileWorkerSingleJobHeader = 'S';
 // this is for the protocol, not the data.
@@ -222,6 +222,9 @@ struct FShaderCompilerInput
 	// Indicates which additional debug outputs should be written for this compile job.
 	EShaderDebugInfoFlags DebugInfoFlags;
 
+	// True if the backend for this job implements the independent preprocessing API.
+	bool bIndependentPreprocessed;
+
 	// True if the cache key for this job should be based on preprocessed source. If so,
 	// preprocessing will be executed in the cook process independent of compilation (and
 	// as such this will only ever be set for jobs whose shader format supports independent
@@ -268,6 +271,7 @@ struct FShaderCompilerInput
 		Target(SF_NumFrequencies, SP_NumPlatforms),
 		bSkipPreprocessedCache(false),
 		DebugInfoFlags(EShaderDebugInfoFlags::Default),
+		bIndependentPreprocessed(false),
 		bCachePreprocessed(false),
 		bCompilingForShaderPipeline(false),
 		bIncludeUsedOutputs(false)
@@ -527,6 +531,7 @@ struct FShaderCompilerOutput
 	,	bSucceeded(false)
 	,	bSupportsQueryingUsedAttributes(false)
 	,	bUsedHLSLccCompiler(false)
+	,	bSerializeModifiedSource(false)
 	{
 	}
 
@@ -545,6 +550,7 @@ struct FShaderCompilerOutput
 	bool bFailedRemovingUnused;
 	bool bSupportsQueryingUsedAttributes;
 	bool bUsedHLSLccCompiler;
+	bool bSerializeModifiedSource;
 	TArray<FString> UsedAttributes;
 
 	PRAGMA_DISABLE_DEPRECATION_WARNINGS
@@ -559,8 +565,25 @@ struct FShaderCompilerOutput
 
 	TArray<FShaderCodeValidationStride> ParametersStrideToValidate;
 
+	/** This field should be set by backends which do not implement the independent preprocessing API to contain the "final" shader source as 
+	 * passed to the platform compiler. For backends that do implement this API this is superceded by ModifiedShaderSource (and will eventually
+	 * be deprecated).
+	 */
 	FString OptionalFinalShaderSource;
 
+	/** Use this field to store the shader source code if it's modified as part of the shader format's compilation process. This field is only 
+	 * currently required for shader formats which implement the independent preprocessing API and should only be set when additional manipulation 
+	 * is required that is not part of the implementation of PreprocessShader. This version of the source, if set, will be what is written as part
+	 * of the debug dumps of preprocessed source, as well as used in place of OptionalFinalShaderSource for upstream code which explicitly requests
+	 * the final source code for other purposes (i.e. when ExtraSettings.bExtractShaderSource is set on the FShaderCompilerInput struct)
+	 */
+	FString ModifiedShaderSource;
+
+	/** Use this field to store the entry point name if it's modified as part of the shader format's compilation process. This field is only 
+	 * currently required for shader formats which implement the independent preprocessing API And should only be set when compilation requires
+	 * a different entry point than was set on the FShaderCompilerInput struct. */
+	FString ModifiedEntryPointName;
+
 	TArray<uint8> PlatformDebugData;
 
 	/** Generates OutputHash from the compiler output. */
@@ -580,6 +603,12 @@ struct FShaderCompilerOutput
 		Ar << Output.CompileTime;
 		Ar << Output.PreprocessTime;
 		Ar << Output.OptionalFinalShaderSource;
+		Ar << Output.bSerializeModifiedSource;
+		if (Output.bSerializeModifiedSource)
+		{
+			Ar << Output.ModifiedShaderSource;
+			Ar << Output.ModifiedEntryPointName;
+		}
 		Ar << Output.PlatformDebugData;
 
 		return Ar;
diff --git a/Engine/Source/Runtime/RenderCore/Public/ShaderCompilerJobTypes.h b/Engine/Source/Runtime/RenderCore/Public/ShaderCompilerJobTypes.h
--- a/Engine/Source/Runtime/RenderCore/Public/ShaderCompilerJobTypes.h
+++ b/Engine/Source/Runtime/RenderCore/Public/ShaderCompilerJobTypes.h
@@ -217,6 +217,8 @@ public:
 	// Serializes only the subset of data written by ShaderCompiler and read from SCW when using worker processes.
 	RENDERCORE_API void SerializeWorkerInput(FArchive& Ar);
 
+	RENDERCORE_API const FString& GetFinalSource() const;
+
 	RENDERCORE_API FShaderCompileJob() : FShaderCommonCompileJob(Type, 0u, 0u, EShaderCompileJobPriority::Num)
 	{}
 
diff --git a/Engine/Source/Runtime/RenderCore/Public/ShaderPreprocessTypes.h b/Engine/Source/Runtime/RenderCore/Public/ShaderPreprocessTypes.h
--- a/Engine/Source/Runtime/RenderCore/Public/ShaderPreprocessTypes.h
+++ b/Engine/Source/Runtime/RenderCore/Public/ShaderPreprocessTypes.h
@@ -64,9 +64,11 @@ public:
 		return PreprocessedSource;
 	}
 
-	const FString& GetOriginalSource() const
+	const FString& GetUnstrippedSource() const
 	{
-		return OriginalPreprocessedSource;
+		// if the unstripped source is requested, check if the "original source" field has been populated
+		// if not then stripping hasn't occurred so there's only one preprocessed source; return it
+		return OriginalPreprocessedSource.IsEmpty() ? PreprocessedSource : OriginalPreprocessedSource;
 	}
 
 	FString& EditSource()
