Fixing CompileFXC.bat issues.
- Setting FXC="..." is different from setting "FXC=...", the former will break uses of just %FXC% when the path has spaces in it.
- Switching to all "FXC=..." and "%FXC%" to make the variable consistent.
- Tweaking the bat file formatting inside of D3DShaderCompiler to make it more readable.

diff --git a/Engine/Source/Developer/Windows/ShaderFormatD3D/Private/D3DShaderCompiler.cpp b/Engine/Source/Developer/Windows/ShaderFormatD3D/Private/D3DShaderCompiler.cpp
--- a/Engine/Source/Developer/Windows/ShaderFormatD3D/Private/D3DShaderCompiler.cpp
+++ b/Engine/Source/Developer/Windows/ShaderFormatD3D/Private/D3DShaderCompiler.cpp
@@ -171,7 +171,7 @@ static FString D3D11CreateShaderCompileCommandLine(
 	)
 {
 	// fxc is our command line compiler
-	FString FXCCommandline = FString(TEXT("%FXC% ")) + ShaderPath;
+	FString FXCCommandline = FString(TEXT("\"%FXC%\" ")) + ShaderPath;
 
 	// add the entry point reference
 	FXCCommandline += FString(TEXT(" /E ")) + EntryFunction;
@@ -287,13 +287,13 @@ static FString D3D11CreateShaderCompileCommandLine(
 	*/
 	const FString BatchFileHeader = TEXT(
 		"@ECHO OFF\n"\
-		"IF \"%FXC%\" == \"\" SET FXC=\"C:\\Program Files (x86)\\Windows Kits\\10\\bin\\x64\\fxc.exe\"\n"\
-		"IF NOT EXIST %FXC% (\n"\
-		"\tECHO Couldn't find Windows 10 SDK, falling back to DXSDK...\n"\
-		"\tSET FXC=\"%DXSDK_DIR%\\Utilities\\bin\\x86\\fxc.exe\"\n"\
-		"\tIF NOT EXIST % FXC % (\n"\
-		"\t\tECHO Couldn't find DXSDK! Exiting...\n"\
-		"\t\tGOTO END\n"\
+		"IF \"%FXC%\" == \"\" SET \"FXC=C:\\Program Files (x86)\\Windows Kits\\10\\bin\\x64\\fxc.exe\"\n"\
+		"IF NOT EXIST \"%FXC%\" (\n"\
+		"\t" "ECHO Couldn't find Windows 10 SDK, falling back to DXSDK...\n"\
+		"\t" "SET \"FXC=%DXSDK_DIR%\\Utilities\\bin\\x86\\fxc.exe\"\n"\
+		"\t" "IF NOT EXIST \"%FXC%\" (\n"\
+		"\t" "\t" "ECHO Couldn't find DXSDK! Exiting...\n"\
+		"\t" "\t" "GOTO END\n"\
 		"\t)\n"\
 		")\n"
 	);
