[shaders] change d3d backend to not fail on minification failure; ensure preprocess warnings are appended for independent preprocessing shader formats regardless of whether preprocessing fails

diff --git a/Engine/Source/Developer/Windows/ShaderFormatD3D/Private/D3DShaderCompiler.cpp b/Engine/Source/Developer/Windows/ShaderFormatD3D/Private/D3DShaderCompiler.cpp
--- a/Engine/Source/Developer/Windows/ShaderFormatD3D/Private/D3DShaderCompiler.cpp
+++ b/Engine/Source/Developer/Windows/ShaderFormatD3D/Private/D3DShaderCompiler.cpp
@@ -563,7 +563,6 @@ static bool CompileAndProcessD3DShaderFXCExt(
 	auto AnsiSourceFile = StringCast<ANSICHAR>(*PreprocessedShaderSource);
 
 	bool bDumpDebugInfo = Input.DumpDebugInfoEnabled();
-	FString DisasmFilename;
 	if (bDumpDebugInfo)
 	{
 		FString BatchFileContents;
@@ -577,8 +576,6 @@ static bool CompileAndProcessD3DShaderFXCExt(
 		}
 
 		FFileHelper::SaveStringToFile(BatchFileContents, *(Input.DumpDebugInfoPath / TEXT("CompileFXC.bat")));
-
-		DisasmFilename = *(Input.DumpDebugInfoPath / TEXT("Output.d3dasm"));
 	}
 
 	TRefCountPtr<ID3DBlob> Shader;
@@ -1196,10 +1193,7 @@ bool PreprocessD3DShader(
 	#if UE_D3D_SHADER_COMPILER_ALLOW_DEAD_CODE_REMOVAL
 	if (Environment.CompilerFlags.Contains(CFLAG_RemoveDeadCode))
 	{
-		if (!UE::ShaderCompilerCommon::RemoveDeadCode(PreprocessedSource, Input.EntryPointName, Output.EditErrors()))
-		{
-			return false;
-		}
+		UE::ShaderCompilerCommon::RemoveDeadCode(PreprocessedSource, Input.EntryPointName, Output.EditErrors());
 	}
 	#endif // UE_D3D_SHADER_COMPILER_ALLOW_DEAD_CODE_REMOVAL
 	return true;
