Fix to live coding shutdown to eliminate a somewhat rare crash on exit.

diff --git a/Engine/Source/Developer/Windows/LiveCoding/Private/LiveCodingModule.cpp b/Engine/Source/Developer/Windows/LiveCoding/Private/LiveCodingModule.cpp
--- a/Engine/Source/Developer/Windows/LiveCoding/Private/LiveCodingModule.cpp
+++ b/Engine/Source/Developer/Windows/LiveCoding/Private/LiveCodingModule.cpp
@@ -1165,11 +1165,12 @@ void FLiveCodingModule::UpdateModules(bool bAllowStarting)
 		TArray<uintptr_t> Pages;
 		FGraphEventRef Task;
 		{
-			FScopeLock lock(&ModuleChangeCs);
-			Swap(Changes, ModuleChanges);
-			Swap(Pages, ReservedPages);
-			Task = ReservePagesTaskRef;
-			LastReservePagesModuleCount = 0;
+			ReservePagesGlobalData& Global = GetReservePagesGlobalData();
+			FScopeLock lock(&Global.ModuleChangeCs);
+			Swap(Changes, Global.ModuleChanges);
+			Swap(Pages, Global.ReservedPages);
+			Task = Global.ReservePagesTaskRef;
+			Global.LastReservePagesModuleCount = 0;
 		}
 
 		if (Task.IsValid())
@@ -1254,21 +1255,28 @@ void FLiveCodingModule::UpdateModules(bool bAllowStarting)
 	}
 }
 
+FLiveCodingModule::ReservePagesGlobalData& FLiveCodingModule::GetReservePagesGlobalData()
+{
+	static ReservePagesGlobalData Singleton;
+	return Singleton;
+}
+
 void FLiveCodingModule::ReservePagesTask()
 {
 	for (;;)
 	{
+		ReservePagesGlobalData& Global = GetReservePagesGlobalData();
 		TArray<ModuleChange> Changes;
 		{
-			FScopeLock lock(&ModuleChangeCs);
-			uint32 NewModules = ModuleChanges.Num() - LastReservePagesModuleCount;
+			FScopeLock lock(&Global.ModuleChangeCs);
+			uint32 NewModules = Global.ModuleChanges.Num() - Global.LastReservePagesModuleCount;
 			if (NewModules < NewModuleCountThreshhold)
 			{
-				ReservePagesTaskRef = FGraphEventRef();
+				Global.ReservePagesTaskRef = FGraphEventRef();
 				return;
 			}
-			Changes.Append(ModuleChanges.GetData() + LastReservePagesModuleCount, NewModules);
-			LastReservePagesModuleCount = ModuleChanges.Num();
+			Changes.Append(Global.ModuleChanges.GetData() + Global.LastReservePagesModuleCount, NewModules);
+			Global.LastReservePagesModuleCount = Global.ModuleChanges.Num();
 		}
 
 		TRACE_CPUPROFILER_EVENT_SCOPE(LiveCodingReservePages);
@@ -1288,8 +1296,8 @@ void FLiveCodingModule::ReservePagesTask()
 		ReservePages(Pages, ModuleBases, DefaultPadding, DefaultPageAlignment);
 
 		{
-			FScopeLock lock(&ModuleChangeCs);
-			ReservedPages.Append(Pages);
+			FScopeLock lock(&Global.ModuleChangeCs);
+			Global.ReservedPages.Append(Pages);
 		}
 	}
 }
@@ -1436,11 +1444,12 @@ void FLiveCodingModule::OnDllLoaded(const FString& FullPath)
 {
 	if (IsUEDll(FullPath))
 	{
-		FScopeLock lock(&ModuleChangeCs);
-		ModuleChanges.Emplace(ModuleChange{ FName(FullPath), true });
-		if (!ReservePagesTaskRef.IsValid() && ModuleChanges.Num() - LastReservePagesModuleCount >= NewModuleCountThreshhold)
+		ReservePagesGlobalData& Global = GetReservePagesGlobalData();
+		FScopeLock lock(&Global.ModuleChangeCs);
+		Global.ModuleChanges.Emplace(ModuleChange{ FName(FullPath), true });
+		if (!Global.ReservePagesTaskRef.IsValid() && Global.ModuleChanges.Num() - Global.LastReservePagesModuleCount >= NewModuleCountThreshhold)
 		{
-			ReservePagesTaskRef = FFunctionGraphTask::CreateAndDispatchWhenReady([this]() { ReservePagesTask(); });
+			Global.ReservePagesTaskRef = FFunctionGraphTask::CreateAndDispatchWhenReady([this]() { ReservePagesTask(); });
 		}
 	}
 }
@@ -1449,8 +1458,9 @@ void FLiveCodingModule::OnDllUnloaded(const FString& FullPath)
 {
 	if (IsUEDll(FullPath))
 	{
-		FScopeLock lock(&ModuleChangeCs);
-		ModuleChanges.Emplace(ModuleChange{ FName(FullPath), false });
+		ReservePagesGlobalData& Global = GetReservePagesGlobalData();
+		FScopeLock lock(&Global.ModuleChangeCs);
+		Global.ModuleChanges.Emplace(ModuleChange{ FName(FullPath), false });
 	}
 }
 
diff --git a/Engine/Source/Developer/Windows/LiveCoding/Private/LiveCodingModule.h b/Engine/Source/Developer/Windows/LiveCoding/Private/LiveCodingModule.h
--- a/Engine/Source/Developer/Windows/LiveCoding/Private/LiveCodingModule.h
+++ b/Engine/Source/Developer/Windows/LiveCoding/Private/LiveCodingModule.h
@@ -64,7 +64,6 @@ private:
 	bool StartLiveCoding(ELiveCodingStartupMode StartupMode);
 	void OnModulesChanged(FName ModuleName, EModuleChangeReason Reason);
 	void UpdateModules(bool bAllowStarting);
-	void ReservePagesTask();
 	bool ShouldPreloadModule(const TSet<FName>& PreloadedFileNames, const FString& FullFilePath) const;
 	bool IsReinstancingEnabled() const;
 	void WaitForStartup();
@@ -72,6 +71,22 @@ private:
 	void ShowConsole(bool bAllowStarting);
 	void SetBuildArguments();
 
+	struct ModuleChange
+	{
+		FName FullName;
+		bool bLoaded;
+	};
+	struct ReservePagesGlobalData
+	{
+		FCriticalSection ModuleChangeCs;
+		TArray<ModuleChange> ModuleChanges;
+		TArray<uintptr_t> ReservedPages;
+		int LastReservePagesModuleCount = 0;
+		FGraphEventRef ReservePagesTaskRef;
+	};
+	static ReservePagesGlobalData& GetReservePagesGlobalData();
+	static void ReservePagesTask();
+
 #if WITH_EDITOR
 	void ShowNotification(bool Success, const FText& Title, const FText* SubText);
 #endif
@@ -115,17 +130,7 @@ private:
 	FDelegateHandle ModulesChangedDelegateHandle;
 	FOnPatchCompleteDelegate OnPatchCompleteDelegate;
 
-	struct ModuleChange
-	{
-		FName FullName;
-		bool bLoaded;
-	};
 	FCriticalSection SetBuildArgumentsCs;
-	FCriticalSection ModuleChangeCs;
-	TArray<ModuleChange> ModuleChanges;
-	TArray<uintptr_t> ReservedPages;
-	int LastReservePagesModuleCount = 0;
-	FGraphEventRef ReservePagesTaskRef;
 
 #if WITH_EDITOR
 	TUniquePtr<FReload> Reload;
