PR #10502: D3DShaderCompiler Memleak Fix

#rnx

diff --git a/Engine/Source/Developer/Windows/ShaderFormatD3D/Private/D3DShaderCompiler.cpp b/Engine/Source/Developer/Windows/ShaderFormatD3D/Private/D3DShaderCompiler.cpp
--- a/Engine/Source/Developer/Windows/ShaderFormatD3D/Private/D3DShaderCompiler.cpp
+++ b/Engine/Source/Developer/Windows/ShaderFormatD3D/Private/D3DShaderCompiler.cpp
@@ -787,9 +787,9 @@ static bool CompileAndProcessD3DShaderFXCExt(
 		if (D3DReflectFunc)
 		{
 			Output.bSucceeded = true;
-			ID3D11ShaderReflection* Reflector = NULL;
-
-			Result = D3DReflectFunc(Shader->GetBufferPointer(), Shader->GetBufferSize(), IID_ID3D11ShaderReflectionForCurrentCompiler, (void**)&Reflector);
+			TRefCountPtr<ID3D11ShaderReflection> Reflector;
+			
+			Result = D3DReflectFunc(Shader->GetBufferPointer(), Shader->GetBufferSize(), IID_ID3D11ShaderReflectionForCurrentCompiler, (void**)Reflector.GetInitReference());
 			if (FAILED(Result))
 			{
 				UE_LOG(LogD3D11ShaderCompiler, Fatal, TEXT("D3DReflect failed: Result=%08x"), Result);
@@ -935,9 +935,6 @@ static bool CompileAndProcessD3DShaderFXCExt(
 					Output, UniformBufferNames, UsedUniformBufferSlots, VendorExtensions);
 
 			NumInstructions = ShaderDesc.InstructionCount;
-
-			// Reflector is a com interface, so it needs to be released.
-			Reflector->Release();
 		}
 		else
 		{
