[shaders] fix debug compile batchfiles for PC DXC

diff --git a/Engine/Source/Developer/Windows/ShaderFormatD3D/Private/D3DShaderCompilerDXC.cpp b/Engine/Source/Developer/Windows/ShaderFormatD3D/Private/D3DShaderCompilerDXC.cpp
--- a/Engine/Source/Developer/Windows/ShaderFormatD3D/Private/D3DShaderCompilerDXC.cpp
+++ b/Engine/Source/Developer/Windows/ShaderFormatD3D/Private/D3DShaderCompilerDXC.cpp
@@ -777,7 +777,7 @@ static FString D3DCreateDXCCompileBatchFile(const FDxcArguments& Args)
 			"\tECHO Couldn't find dxc.exe under \"%s\"\n"
 			"\tGOTO :END\n"
 			")\n"
-			"%%DXC%%%s %s\n"
+			"%%DXC%%%s %s.usf\n"
 			":END\n"
 			"PAUSE\n"
 		),
