Add [loop] attribute to Substrate shader to fix HLSL cross-compilation error on D3D11.

#rnx

diff --git a/Engine/Shaders/Private/Strata/StrataEvaluation.ush b/Engine/Shaders/Private/Strata/StrataEvaluation.ush
--- a/Engine/Shaders/Private/Strata/StrataEvaluation.ush
+++ b/Engine/Shaders/Private/Strata/StrataEvaluation.ush
@@ -59,7 +59,11 @@
 //		...
 // }
 #if STRATA_FASTPATH == 0 && STRATA_SINGLEPATH == 0
-	#define Strata_for(X,Y,Z) for(X;Y;Z)
+	#if COMPILER_FXC
+		#define Strata_for(X,Y,Z) [loop] for(X;Y;Z)
+	#else
+		#define Strata_for(X,Y,Z) for(X;Y;Z)
+	#endif
 #else
 	#define Strata_for(X,Y,Z) X; if(Y)
 #endif
diff --git a/Engine/Shaders/Public/Platform.ush b/Engine/Shaders/Public/Platform.ush
--- a/Engine/Shaders/Public/Platform.ush
+++ b/Engine/Shaders/Public/Platform.ush
@@ -84,6 +84,10 @@
 #define COMPILER_DXC 0
 #endif
 
+#ifndef COMPILER_FXC
+#define COMPILER_FXC 0
+#endif
+
 #ifndef COMPILER_HLSL
 #define COMPILER_HLSL 0
 #endif
diff --git a/Engine/Source/Developer/Windows/ShaderFormatD3D/Private/D3DShaderCompiler.cpp b/Engine/Source/Developer/Windows/ShaderFormatD3D/Private/D3DShaderCompiler.cpp
--- a/Engine/Source/Developer/Windows/ShaderFormatD3D/Private/D3DShaderCompiler.cpp
+++ b/Engine/Source/Developer/Windows/ShaderFormatD3D/Private/D3DShaderCompiler.cpp
@@ -1163,6 +1163,10 @@ bool PreprocessD3DShader(
 			AdditionalDefines.SetDefine(TEXT("PLATFORM_SUPPORTS_REAL_TYPES"), 1);
 		}
 	}
+	else
+	{
+		AdditionalDefines.SetDefine(TEXT("COMPILER_FXC"), 1);
+	}
 
 	PRAGMA_ENABLE_DEPRECATION_WARNINGS
 	
