Fix several more Shipping compile errors

#rnx

diff --git a/Engine/Source/Developer/Windows/LiveCodingServer/Private/LiveCodingServerModule.cpp b/Engine/Source/Developer/Windows/LiveCodingServer/Private/LiveCodingServerModule.cpp
--- a/Engine/Source/Developer/Windows/LiveCodingServer/Private/LiveCodingServerModule.cpp
+++ b/Engine/Source/Developer/Windows/LiveCodingServer/Private/LiveCodingServerModule.cpp
@@ -13,6 +13,7 @@ DEFINE_LOG_CATEGORY_STATIC(LogLiveCodingServer, Display, All);
 
 static void ServerLogOutput(Logging::Channel::Enum Channel, Logging::Type::Enum Type, const wchar_t* const Text, int Count)
 {
+#if !NO_LOGGING
 	ELogVerbosity::Type Verbosity = ELogVerbosity::Error;
 
 	switch (Type)
@@ -44,6 +45,7 @@ static void ServerLogOutput(Logging::Channel::Enum Channel, Logging::Type::Enum
 	{
 		FMsg::Logf(__FILE__, __LINE__, LogLiveCodingServer.GetCategoryName(), Verbosity, TEXT("%s (repeated %d more times)"), Text, Count);
 	}
+#endif
 }
 
 static void ServerOutputHandler(Logging::Channel::Enum Channel, Logging::Type::Enum Type, const wchar_t* const Text)
