D3D Shader Compiler changes
- PLATFORM_SUPPORTS_UB_STRUCT was set for RT shaders, but RemoveUniformBuffersFromSource was still being called.
- New IsUsingSM66 function to keep consistent usage of ELanguage::SM6 since not all RT shaders use that 'language'
- Removing USE_SHADER_MODEL_6_6 since SM6.6 is locked in
- Changing FShaderFormatD3D versions to be Guids instead of numbers to better handle resolving changes in different branches.

diff --git a/Engine/Source/Developer/Windows/ShaderFormatD3D/Private/D3DShaderCompiler.cpp b/Engine/Source/Developer/Windows/ShaderFormatD3D/Private/D3DShaderCompiler.cpp
--- a/Engine/Source/Developer/Windows/ShaderFormatD3D/Private/D3DShaderCompiler.cpp
+++ b/Engine/Source/Developer/Windows/ShaderFormatD3D/Private/D3DShaderCompiler.cpp
@@ -82,10 +82,9 @@ static void D3D11FilterShaderCompileWarnings(const FString& CompileWarnings, TAr
 }
 
 // @return 0 if not recognized
-static const TCHAR* GetShaderProfileName(ELanguage Language, uint32 Frequency, bool bForceSM6)
+static const TCHAR* GetShaderProfileName(ELanguage Language, EShaderFrequency Frequency, bool bForceSM6)
 {
-
-	if (Language == ELanguage::SM6 || IsRayTracingShaderFrequency(EShaderFrequency(Frequency)))
+	if (IsUsingSM66(Language, Frequency))
 	{
 		switch (Frequency)
 		{
@@ -93,22 +92,22 @@ static const TCHAR* GetShaderProfileName(ELanguage Language, uint32 Frequency, b
 			checkfSlow(false, TEXT("Unexpected shader frequency"));
 			return nullptr;
 		case SF_Pixel:
-			return USE_SHADER_MODEL_6_6 ? TEXT("ps_6_6") : TEXT("ps_6_5");
+			return TEXT("ps_6_6");
 		case SF_Vertex:
-			return USE_SHADER_MODEL_6_6 ? TEXT("vs_6_6") : TEXT("vs_6_5");
+			return TEXT("vs_6_6");
 		case SF_Mesh:
-			return USE_SHADER_MODEL_6_6 ? TEXT("ms_6_6") : TEXT("ms_6_5");
+			return TEXT("ms_6_6");
 		case SF_Amplification:
-			return USE_SHADER_MODEL_6_6 ? TEXT("as_6_6") : TEXT("as_6_5");
+			return TEXT("as_6_6");
 		case SF_Geometry:
-			return USE_SHADER_MODEL_6_6 ? TEXT("gs_6_6") : TEXT("gs_6_5");
+			return TEXT("gs_6_6");
 		case SF_Compute:
-			return USE_SHADER_MODEL_6_6 ? TEXT("cs_6_6") : TEXT("cs_6_5");
+			return TEXT("cs_6_6");
 		case SF_RayGen:
 		case SF_RayMiss:
 		case SF_RayHitGroup:
 		case SF_RayCallable:
-			return USE_SHADER_MODEL_6_6 ? TEXT("lib_6_6") : TEXT("lib_6_5");
+			return TEXT("lib_6_6");
 		}
 	}
 	else if(Language == ELanguage::SM5)
@@ -1106,8 +1105,7 @@ bool CompileAndProcessD3DShaderFXC(const FShaderPreprocessOutput& PreprocessOutp
 
 inline bool ShouldUseDXC(const FShaderCompilerInput& Input, ELanguage Language)
 {
-	return Language == ELanguage::SM6
-		|| Input.IsRayTracingShader()
+	return IsUsingSM66(Input, Language)
 		|| Input.Environment.CompilerFlags.Contains(CFLAG_WaveOperations)
 		|| Input.Environment.CompilerFlags.Contains(CFLAG_ForceDXC)
 		|| Input.Environment.CompilerFlags.Contains(CFLAG_InlineRayTracing)
@@ -1122,7 +1120,7 @@ bool PreprocessD3DShader(
 {
 	PRAGMA_DISABLE_DEPRECATION_WARNINGS		// FShaderCompilerDefinitions will be made internal in the future, marked deprecated until then
 	FShaderCompilerDefinitions AdditionalDefines;
-	if (Language == ELanguage::SM6)
+	if (IsUsingSM66(Input, Language))
 	{
 		AdditionalDefines.SetDefine(TEXT("SM6_PROFILE"), 1);
 	}
@@ -1147,7 +1145,7 @@ bool PreprocessD3DShader(
 	{
 		AdditionalDefines.SetDefine(TEXT("PLATFORM_SUPPORTS_CALLABLE_SHADERS"), 1);
 		AdditionalDefines.SetDefine(TEXT("PLATFORM_SUPPORTS_SM6_0_WAVE_OPERATIONS"), 1);
-		if (Language == ELanguage::SM6 || Input.IsRayTracingShader())
+		if (IsUsingSM66(Input, Language))
 		{
 			AdditionalDefines.SetDefine(TEXT("PLATFORM_SUPPORTS_STATIC_SAMPLERS"), 1);
 		}
@@ -1197,7 +1195,7 @@ bool PreprocessD3DShader(
 	}
 
 	// Only use UniformBuffer structs on SM6 until we can fully vet SM5
-	if (Language != ELanguage::SM6)
+	if (!IsUsingSM66(Input, Language))
 	{
 		RemoveUniformBuffersFromSource(Environment, PreprocessedSource);
 	}
@@ -1222,7 +1220,7 @@ void CompileD3DShader(const FShaderCompilerInput& Input, const FShaderPreprocess
 	TRACE_CPUPROFILER_EVENT_SCOPE(CompileD3DShader);
 
 	const bool bUseDXC = ShouldUseDXC(Input, Language);
-	const TCHAR* ShaderProfile = GetShaderProfileName(Language, Input.Target.Frequency, bUseDXC);
+	const TCHAR* ShaderProfile = GetShaderProfileName(Language, Input.Target.GetFrequency(), bUseDXC);
 
 	if(!ShaderProfile)
 	{
diff --git a/Engine/Source/Developer/Windows/ShaderFormatD3D/Private/ShaderFormatD3D.cpp b/Engine/Source/Developer/Windows/ShaderFormatD3D/Private/ShaderFormatD3D.cpp
--- a/Engine/Source/Developer/Windows/ShaderFormatD3D/Private/ShaderFormatD3D.cpp
+++ b/Engine/Source/Developer/Windows/ShaderFormatD3D/Private/ShaderFormatD3D.cpp
@@ -12,18 +12,13 @@ static FName NAME_PCD3D_SM6(TEXT("PCD3D_SM6"));
 static FName NAME_PCD3D_SM5(TEXT("PCD3D_SM5"));
 static FName NAME_PCD3D_ES3_1(TEXT("PCD3D_ES31"));
 
+static const FGuid UE_SHADER_PCD3D_SHARED_VER = FGuid("dd4e6e76-4b48-4097-9ece-0f21118b7177");
+static const FGuid UE_SHADER_PCD3D_SM6_VER    = FGuid("c12d7716-7d6a-4be7-920f-9d38c54d1d36");
+static const FGuid UE_SHADER_PCD3D_SM5_VER    = FGuid("07cd0fa4-6d18-43e0-915d-e6c125effafd");
+static const FGuid UE_SHADER_PCD3D_ES3_1_VER  = FGuid("75466d2b-e169-40d8-bac5-1e2f9d43e0bb");
+
 class FShaderFormatD3D : public UE::ShaderCompilerCommon::FBaseShaderFormat 
 {
-	enum EVersion
-	{
-		UE_SHADER_PCD3D_SHARED_VER = 5,
-
-		/** Version for shader format, this becomes part of the DDC key. */
-		UE_SHADER_PCD3D_SM6_VER = 8,
-		UE_SHADER_PCD3D_SM5_VER = 13,
-		UE_SHADER_PCD3D_ES3_1_VER = 8,
-	};
-
 	uint32 DxcVersionHash = 0;
 
 public:
@@ -33,7 +28,7 @@ public:
 	{
 	}
 
-	inline uint32 GetVersionHash(EVersion InVersion) const
+	inline uint32 GetVersionHash(const FGuid& InVersion) const
 	{
 		const uint32 BaseHash = GetTypeHash(UE_SHADER_PCD3D_SHARED_VER);
 		uint32 VersionHash = GetTypeHash(InVersion);
@@ -51,10 +46,6 @@ public:
 		{
 			uint32 ShaderModelHash = GetVersionHash(UE_SHADER_PCD3D_SM6_VER);
 
-		#if USE_SHADER_MODEL_6_6
-			ShaderModelHash ^= 0x96ED7F56;
-		#endif
-
 			// Make sure we recompile if EShaderCodeFeatures gets bigger
 			ShaderModelHash = HashCombine(ShaderModelHash, GetTypeHash(sizeof(EShaderCodeFeatures)));
 
@@ -136,22 +127,17 @@ public:
 
 	void ModifyShaderCompilerInput(FShaderCompilerInput& Input) const final
 	{
-		if (Input.ShaderFormat == NAME_PCD3D_SM6 || Input.IsRayTracingShader())
+		const ELanguage Language = FormatToLanguage(Input.ShaderFormat);
+
+		if (IsUsingSM66(Input, Language))
 		{
 			Input.Environment.SetDefine(TEXT("SM6_PROFILE"), 1);
 			Input.Environment.SetDefine(TEXT("COMPILER_DXC"), 1);
 			Input.Environment.SetDefine(TEXT("PLATFORM_SUPPORTS_UB_STRUCT"), 1);
 
-			if (USE_SHADER_MODEL_6_6)
-			{
-				AddShaderTargetDefines(Input, 6, 6);
-			}
-			else
-			{
-				AddShaderTargetDefines(Input, 6, 5);
-			}
+			AddShaderTargetDefines(Input, 6, 6);
 		}
-		else if (Input.ShaderFormat == NAME_PCD3D_SM5)
+		else if (Language == ELanguage::SM5)
 		{
 			Input.Environment.SetDefine(TEXT("SM5_PROFILE"), 1);
 			const bool bUseDXC =
@@ -168,7 +154,7 @@ public:
 				AddShaderTargetDefines(Input, 5, 0);
 			}
 		}
-		else if (Input.ShaderFormat == NAME_PCD3D_ES3_1)
+		else if (Language == ELanguage::ES3_1)
 		{
 			Input.Environment.SetDefine(TEXT("ES3_1_PROFILE"), 1);
 			Input.Environment.SetDefine(TEXT("COMPILER_DXC"), 0);
diff --git a/Engine/Source/Developer/Windows/ShaderFormatD3D/Private/ShaderFormatD3D.h b/Engine/Source/Developer/Windows/ShaderFormatD3D/Private/ShaderFormatD3D.h
--- a/Engine/Source/Developer/Windows/ShaderFormatD3D/Private/ShaderFormatD3D.h
+++ b/Engine/Source/Developer/Windows/ShaderFormatD3D/Private/ShaderFormatD3D.h
@@ -5,9 +5,6 @@
 #include "Windows/WindowsHWrapper.h"
 #include "ShaderCompilerCommon.h"
 
-// TODO: Lock D3D12 to SM 6.6 min spec
-#define USE_SHADER_MODEL_6_6 1
-
 // Controls whether r.Shaders.RemoveDeadCode should be honored
 #ifndef UE_D3D_SHADER_COMPILER_ALLOW_DEAD_CODE_REMOVAL
 #define UE_D3D_SHADER_COMPILER_ALLOW_DEAD_CODE_REMOVAL 1
@@ -23,6 +20,16 @@ enum class ELanguage
 	Invalid,
 };
 
+inline bool IsUsingSM66(ELanguage Language, EShaderFrequency Frequency)
+{
+	return Language == ELanguage::SM6 || IsRayTracingShaderFrequency(Frequency);
+}
+
+inline bool IsUsingSM66(const FShaderCompilerInput& Input, ELanguage Language)
+{
+	return IsUsingSM66(Language, Input.Target.GetFrequency());
+}
+
 bool PreprocessD3DShader(
 	const struct FShaderCompilerInput& Input,
 	const struct FShaderCompilerEnvironment& MergedEnvironment,
