Do post-compile signing of DXIL containers to work around some validation issues.

diff --git a/Engine/Shaders/Public/Platform/D3D/D3DCommon.ush b/Engine/Shaders/Public/Platform/D3D/D3DCommon.ush
--- a/Engine/Shaders/Public/Platform/D3D/D3DCommon.ush
+++ b/Engine/Shaders/Public/Platform/D3D/D3DCommon.ush
@@ -5,7 +5,7 @@ D3DCommon.ush: Common defines for the D3D shader compiler
 =============================================================================*/
 
 // Update this GUID to invalidate D3D shaders and force a shader recompilation for D3D RHIs.
-// GUID: A44B0FDA-31ED-4E3B-9E04-A80DFDD5B015
+// GUID: EE1A1608-B3B7-4239-9B61-4DF491A9B6E3
 
 #include "/Engine/Shared/HLSLReservedSpaces.h"
 
diff --git a/Engine/Source/Developer/Windows/ShaderFormatD3D/Private/D3DShaderCompilerDXC.cpp b/Engine/Source/Developer/Windows/ShaderFormatD3D/Private/D3DShaderCompilerDXC.cpp
--- a/Engine/Source/Developer/Windows/ShaderFormatD3D/Private/D3DShaderCompilerDXC.cpp
+++ b/Engine/Source/Developer/Windows/ShaderFormatD3D/Private/D3DShaderCompilerDXC.cpp
@@ -701,6 +701,33 @@ static HRESULT D3DCompileToDxil(const char* SourceText, const FDxcArguments& Arg
 		checkf(CompileResult->HasOutput(DXC_OUT_OBJECT), TEXT("No object code found!"));
 		VERIFYHRESULT(CompileResult->GetOutput(DXC_OUT_OBJECT, IID_PPV_ARGS(OutDxilBlob.GetInitReference()), ObjectCodeNameBlob.GetInitReference()));
 
+		const bool bPostCompileSign = true;
+		if (bPostCompileSign)
+		{
+			// https://www.wihlidal.com/blog/pipeline/2018-09-16-dxil-signing-post-compile/
+			TRefCountPtr<IDxcValidator> Validator;
+			VERIFYHRESULT(DxcDllHelper.CreateInstance2(GetDxcMalloc(), CLSID_DxcValidator, Validator.GetInitReference()));
+
+		#if 0
+			struct FDxilMinimalHeader
+			{
+				uint32 FourCC;
+				uint32 HashDigest[4];
+			};
+
+			FDxilMinimalHeader BeforeSignHeader = *reinterpret_cast<FDxilMinimalHeader*>(OutDxilBlob->GetBufferPointer());
+			(void)BeforeSignHeader;
+		#endif
+
+			TRefCountPtr<IDxcOperationResult> ValidateResult;
+			VERIFYHRESULT(Validator->Validate(OutDxilBlob.GetReference(), DxcValidatorFlags_InPlaceEdit, ValidateResult.GetInitReference()));
+
+		#if 0
+			FDxilMinimalHeader AfterSignHeader = *reinterpret_cast<FDxilMinimalHeader*>(OutDxilBlob->GetBufferPointer());
+			(void)AfterSignHeader;
+		#endif
+		}
+
 		TRefCountPtr<IDxcBlobUtf16> ReflectionNameBlob; // Dummy name blob to silence static analysis warning
 		checkf(CompileResult->HasOutput(DXC_OUT_REFLECTION), TEXT("No reflection found!"));
 		VERIFYHRESULT(CompileResult->GetOutput(DXC_OUT_REFLECTION, IID_PPV_ARGS(OutReflectionBlob.GetInitReference()), ReflectionNameBlob.GetInitReference()));
