Refactored DXC RemoveContainerReflection into RemoveContainerParts that can remove additional parts. Also disabled post compile signing step for now.

diff --git a/Engine/Source/Developer/Windows/ShaderFormatD3D/Private/D3DShaderCompilerDXC.cpp b/Engine/Source/Developer/Windows/ShaderFormatD3D/Private/D3DShaderCompilerDXC.cpp
--- a/Engine/Source/Developer/Windows/ShaderFormatD3D/Private/D3DShaderCompilerDXC.cpp
+++ b/Engine/Source/Developer/Windows/ShaderFormatD3D/Private/D3DShaderCompilerDXC.cpp
@@ -645,19 +645,32 @@ static void DumpFourCCParts(dxc::DxcDllSupport& DxcDllHelper, TRefCountPtr<IDxcB
 #endif
 }
 
-static bool RemoveContainerReflection(dxc::DxcDllSupport& DxcDllHelper, TRefCountPtr<IDxcBlob>& Dxil, bool bRemovePDB)
+static bool RemoveContainerParts(const TConstArrayView<uint32> PartCodes, dxc::DxcDllSupport& DxcDllHelper, TRefCountPtr<IDxcBlob>& Dxil)
 {
+	if (PartCodes.Num() == 0)
+	{
+		return false;
+	}
+
 	TRefCountPtr<IDxcOperationResult> Result;
 	TRefCountPtr<IDxcContainerBuilder> Builder;
 	TRefCountPtr<IDxcBlob> StrippedDxil;
 
 	VERIFYHRESULT(DxcDllHelper.CreateInstance2(GetDxcMalloc(), CLSID_DxcContainerBuilder, Builder.GetInitReference()));
 	VERIFYHRESULT(Builder->Load(Dxil));
-	
-	// Try and remove both the PDB & Reflection Data
-	bool bPDBRemoved = bRemovePDB && SUCCEEDED(Builder->RemovePart(DXC_PART_PDB));
-	bool bReflectionDataRemoved = bRemovePDB && SUCCEEDED(Builder->RemovePart(DXC_PART_REFLECTION_DATA));
-	if (bPDBRemoved || bReflectionDataRemoved)
+
+	bool bSuccess = true;
+
+	for (uint32 PartCode : PartCodes)
+	{
+		if (FAILED(Builder->RemovePart(PartCode)))
+		{
+			bSuccess = false;
+			break;
+		}
+	}
+
+	if (bSuccess)
 	{
 		VERIFYHRESULT(Builder->SerializeContainer(Result.GetInitReference()));
 		if (SUCCEEDED(Result->GetResult(StrippedDxil.GetInitReference())))
@@ -668,8 +681,8 @@ static bool RemoveContainerReflection(dxc::DxcDllSupport& DxcDllHelper, TRefCoun
 		}
 	}
 
-	return false;
-};
+	return bSuccess;
+}
 
 static HRESULT D3DCompileToDxil(const char* SourceText, const FDxcArguments& Arguments,
 	TRefCountPtr<IDxcBlob>& OutDxilBlob, TRefCountPtr<IDxcBlob>& OutReflectionBlob, TRefCountPtr<IDxcBlobEncoding>& OutErrorBlob)
@@ -701,7 +714,7 @@ static HRESULT D3DCompileToDxil(const char* SourceText, const FDxcArguments& Arg
 		checkf(CompileResult->HasOutput(DXC_OUT_OBJECT), TEXT("No object code found!"));
 		VERIFYHRESULT(CompileResult->GetOutput(DXC_OUT_OBJECT, IID_PPV_ARGS(OutDxilBlob.GetInitReference()), ObjectCodeNameBlob.GetInitReference()));
 
-		const bool bPostCompileSign = true;
+		const bool bPostCompileSign = false;
 		if (bPostCompileSign)
 		{
 			// https://www.wihlidal.com/blog/pipeline/2018-09-16-dxil-signing-post-compile/
@@ -733,6 +746,15 @@ static HRESULT D3DCompileToDxil(const char* SourceText, const FDxcArguments& Arg
 		VERIFYHRESULT(CompileResult->GetOutput(DXC_OUT_REFLECTION, IID_PPV_ARGS(OutReflectionBlob.GetInitReference()), ReflectionNameBlob.GetInitReference()));
 
 		const bool bHasOutputPDB = CompileResult->HasOutput(DXC_OUT_PDB);
+		const bool bRemovePDB = bHasOutputPDB && !Arguments.ShouldKeepEmbeddedPDB();
+
+		TArray<uint32, TInlineAllocator<4>> PartsToRemove;
+		if (bRemovePDB)
+		{
+			// Try and remove both the PDB & Reflection Data
+			PartsToRemove.Add(DXC_PART_PDB);
+			PartsToRemove.Add(DXC_PART_REFLECTION_DATA);
+		}
 
 		if (Arguments.ShouldDump())
 		{
@@ -760,7 +782,7 @@ static HRESULT D3DCompileToDxil(const char* SourceText, const FDxcArguments& Arg
 		}
 
 		DumpFourCCParts(DxcDllHelper, OutDxilBlob);
-		if (RemoveContainerReflection(DxcDllHelper, OutDxilBlob, bHasOutputPDB && !Arguments.ShouldKeepEmbeddedPDB()))
+		if (RemoveContainerParts(PartsToRemove, DxcDllHelper, OutDxilBlob))
 		{
 			DumpFourCCParts(DxcDllHelper, OutDxilBlob);
 		}
