Bumping the SM6 version to force all D3D RT shaders to rebuild.

diff --git a/Engine/Source/Developer/Windows/ShaderFormatD3D/Private/ShaderFormatD3D.cpp b/Engine/Source/Developer/Windows/ShaderFormatD3D/Private/ShaderFormatD3D.cpp
--- a/Engine/Source/Developer/Windows/ShaderFormatD3D/Private/ShaderFormatD3D.cpp
+++ b/Engine/Source/Developer/Windows/ShaderFormatD3D/Private/ShaderFormatD3D.cpp
@@ -13,7 +13,7 @@ static FName NAME_PCD3D_SM5(TEXT("PCD3D_SM5"));
 static FName NAME_PCD3D_ES3_1(TEXT("PCD3D_ES31"));
 
 static const FGuid UE_SHADER_PCD3D_SHARED_VER = FGuid("dd4e6e76-4b48-4097-9ece-0f21118b7177");
-static const FGuid UE_SHADER_PCD3D_SM6_VER    = FGuid("c12d7716-7d6a-4be7-920f-9d38c54d1d36");
+static const FGuid UE_SHADER_PCD3D_SM6_VER    = FGuid("96a53626-2b11-4663-8b85-01e30625b76e");
 static const FGuid UE_SHADER_PCD3D_SM5_VER    = FGuid("07cd0fa4-6d18-43e0-915d-e6c125effafd");
 static const FGuid UE_SHADER_PCD3D_ES3_1_VER  = FGuid("75466d2b-e169-40d8-bac5-1e2f9d43e0bb");
 
