Fix compile issue with cl-clang

#rnx

diff --git a/Engine/Source/Developer/Windows/LiveCoding/Private/External/LC_API.h b/Engine/Source/Developer/Windows/LiveCoding/Private/External/LC_API.h
--- a/Engine/Source/Developer/Windows/LiveCoding/Private/External/LC_API.h
+++ b/Engine/Source/Developer/Windows/LiveCoding/Private/External/LC_API.h
@@ -3,7 +3,8 @@
 #pragma once
 
 // BEGIN EPIC MOD
-#include "stddef.h"
+#include <stddef.h>
+#include <stdint.h>
 // END EPIC MOD
 #include "LPP_API.h"
 
