Fix invalid virtual file paths from intermediate HLSL shader in D3D11 backend.

#rnx

diff --git a/Engine/Source/Developer/Windows/ShaderFormatD3D/Private/D3DShaderCompiler.cpp b/Engine/Source/Developer/Windows/ShaderFormatD3D/Private/D3DShaderCompiler.cpp
--- a/Engine/Source/Developer/Windows/ShaderFormatD3D/Private/D3DShaderCompiler.cpp
+++ b/Engine/Source/Developer/Windows/ShaderFormatD3D/Private/D3DShaderCompiler.cpp
@@ -58,6 +58,29 @@ static uint32 TranslateCompilerFlagD3D11(ECompilerFlags CompilerFlag)
 	};
 }
 
+/*
+ * Turns invalid absolute paths that FXC generated back into a virtual file paths,
+ * e.g. "D:\\Engine\\Private\\Common.ush" into "/Engine/Private/Common.ush"
+ */
+static void D3D11SanitizeErrorVirtualFilePath(FString& ErrorLine)
+{
+	if (ErrorLine.Len() > 3 && ErrorLine[1] == TEXT(':') && ErrorLine[2] == TEXT('\\'))
+	{
+		const int32 EndOfFilePath = ErrorLine.Find(TEXT(":"), ESearchCase::CaseSensitive, ESearchDir::FromStart, 3);
+		if (EndOfFilePath != INDEX_NONE)
+		{
+			for (int32 ErrorLineStringPosition = 2; ErrorLineStringPosition < EndOfFilePath; ++ErrorLineStringPosition)
+			{
+				if (ErrorLine[ErrorLineStringPosition] == TEXT('\\'))
+				{
+					ErrorLine[ErrorLineStringPosition] = TEXT('/');
+				}
+			}
+			ErrorLine.RightChopInline(2);
+		}
+	}
+}
+
 /**
  * Filters out unwanted shader compile warnings
  */
@@ -76,6 +99,7 @@ static void D3D11FilterShaderCompileWarnings(const FString& CompileWarnings, TAr
 			// Gets spammed when converting from float to half
 			&& !WarningArray[WarningIndex].Contains(TEXT("X3205")))
 		{
+			D3D11SanitizeErrorVirtualFilePath(WarningArray[WarningIndex]);
 			FilteredWarnings.AddUnique(WarningArray[WarningIndex]);
 		}
 	}
@@ -680,9 +704,17 @@ static bool CompileAndProcessD3DShaderFXCExt(
 				DumpDebugShaderText(Input, CrossCompiledSource.GetData(), CrossCompiledSource.Num() - 1, TEXT("intermediate.hlsl"));
 			}
 
+			// Generates an virtual source file path with the ".intermediate." suffix injected.
+			auto MakeIntermediateVirtualSourceFilePath = [](const FString& InVirtualSourceFilePath) -> FString
+			{
+				FString PathPart, FilenamePart, ExtensionPart;
+				FPaths::Split(InVirtualSourceFilePath, PathPart, FilenamePart, ExtensionPart);
+				return FPaths::Combine(PathPart, FilenamePart) + TEXT(".intermediate.") + ExtensionPart;
+			};
+
 			// Compile again with FXC:
 			// SPIRV-Cross will have generated the new shader with "main" as the new entry point.
-			const FString CrossCompiledSourceFilename = Input.VirtualSourceFilePath + TEXT(".intermediate.hlsl");
+			const FString CrossCompiledSourceFilename = MakeIntermediateVirtualSourceFilePath(Input.VirtualSourceFilePath);
 			Result = D3DCompileWrapper(
 				D3DCompileFunc,
 				CrossCompiledSource.GetData(),
