Shader Parameter Parser changes
- Moving more uses of ShaderParameterParse from Preprocessing to Compile step.
- Adding a "DidModifyShader" flag to ShaderParameterParser to determine if Output.ModifiedShaderSource needs to be updated.

diff --git a/Engine/Shaders/Public/ShaderVersion.ush b/Engine/Shaders/Public/ShaderVersion.ush
--- a/Engine/Shaders/Public/ShaderVersion.ush
+++ b/Engine/Shaders/Public/ShaderVersion.ush
@@ -2,4 +2,4 @@
 // Each time the console command is executed it generates a new GUID. As this file is included
 // in Platform.ush (which should be included in any shader) it allows to invalidate the shader DDC.
 // 
-#pragma message("UESHADERMETADATA_VERSION 4037B637-C512-4BA2-B763-29C04DF6E818")
+#pragma message("UESHADERMETADATA_VERSION BC047A48-A9B5-4CF5-BD87-B5E9107C5711")
diff --git a/Engine/Source/Developer/Windows/ShaderFormatD3D/Private/D3DShaderCompiler.cpp b/Engine/Source/Developer/Windows/ShaderFormatD3D/Private/D3DShaderCompiler.cpp
--- a/Engine/Source/Developer/Windows/ShaderFormatD3D/Private/D3DShaderCompiler.cpp
+++ b/Engine/Source/Developer/Windows/ShaderFormatD3D/Private/D3DShaderCompiler.cpp
@@ -573,17 +573,16 @@ static bool CompileErrorsContainInternalError(ID3DBlob* Errors)
 
 // Generate the dumped usf file; call the D3D compiler, gather reflection information and generate the output data
 static bool CompileAndProcessD3DShaderFXCExt(
-	const FShaderPreprocessOutput& PreprocessOutput,
 	uint32 CompileFlags,
 	const FShaderCompilerInput& Input,
+	const FString& PreprocessedShaderSource,
+	const FString& EntryPointName,
 	const FShaderParameterParser& ShaderParameterParser,
 	const TCHAR* ShaderProfile, bool bSecondPassAferUnusedInputRemoval,
 	TArray<FString>& FilteredErrors, FShaderCompilerOutput& Output)
 {
 	TRACE_CPUPROFILER_EVENT_SCOPE(CompileAndProcessD3DShaderFXCExt);
 
-	const FString& PreprocessedShaderSource = Output.ModifiedShaderSource.IsEmpty() ? PreprocessOutput.GetSource() : Output.ModifiedShaderSource;
-	const FString& EntryPointName = Output.ModifiedEntryPointName.IsEmpty() ? Input.EntryPointName : Output.ModifiedEntryPointName;
 	auto AnsiSourceFile = StringCast<ANSICHAR>(*PreprocessedShaderSource);
 
 	bool bDumpDebugInfo = Input.DumpDebugInfoEnabled();
@@ -889,14 +888,12 @@ static bool CompileAndProcessD3DShaderFXCExt(
 					for (int32 Attempt = 0; Attempt < kMaxReasonableAttempts; ++Attempt)
 					{
 						TArray<FString> RemoveErrors;
-						FString ModifiedShaderSource = PreprocessOutput.GetSource();
+						FString ModifiedShaderSource = PreprocessedShaderSource;
 						FString ModifiedEntryPointName = Input.EntryPointName;
 						if (RemoveUnusedInputs(ModifiedShaderSource, ShaderInputs, ModifiedEntryPointName, RemoveErrors))
 						{
 							Output = OriginalOutput;
-							Output.ModifiedShaderSource = MoveTemp(ModifiedShaderSource);
-							Output.ModifiedEntryPointName = MoveTemp(ModifiedEntryPointName);
-							if (!CompileAndProcessD3DShaderFXCExt(PreprocessOutput, CompileFlags, Input, ShaderParameterParser, ShaderProfile, true, FilteredErrors, Output))
+							if (!CompileAndProcessD3DShaderFXCExt(CompileFlags, Input, ModifiedShaderSource, ModifiedEntryPointName, ShaderParameterParser, ShaderProfile, true, FilteredErrors, Output))
 							{
 								// if we failed to compile the shader, propagate the error up
 								return false;
@@ -905,6 +902,9 @@ static bool CompileAndProcessD3DShaderFXCExt(
 							// check if the ShaderInputs changed - if not, we're done here
 							if (Output.UsedAttributes.Num() == ShaderInputs.Num())
 							{
+								Output.ModifiedShaderSource = MoveTemp(ModifiedShaderSource);
+								Output.ModifiedEntryPointName = MoveTemp(ModifiedEntryPointName);
+
 								return true;
 							}
 
@@ -1065,8 +1065,10 @@ static bool CompileAndProcessD3DShaderFXCExt(
 	return SUCCEEDED(Result);
 }
 
-bool CompileAndProcessD3DShaderFXC(const FShaderPreprocessOutput& PreprocessOutput,
+bool CompileAndProcessD3DShaderFXC(
 	const FShaderCompilerInput& Input,
+	const FString& InPreprocessedSource,
+	const FString& InEntryPointName,
 	const FShaderParameterParser& ShaderParameterParser,
 	const TCHAR* ShaderProfile,
 	bool bSecondPassAferUnusedInputRemoval,
@@ -1104,7 +1106,7 @@ bool CompileAndProcessD3DShaderFXC(const FShaderPreprocessOutput& PreprocessOutp
 		});
 
 	TArray<FString> FilteredErrors;
-	const bool bSuccess = CompileAndProcessD3DShaderFXCExt(PreprocessOutput, CompileFlags, Input, ShaderParameterParser, ShaderProfile, false, FilteredErrors, Output);
+	const bool bSuccess = CompileAndProcessD3DShaderFXCExt(CompileFlags, Input, InPreprocessedSource, InEntryPointName, ShaderParameterParser, ShaderProfile, false, FilteredErrors, Output);
 
 	// Process errors
 	for (int32 ErrorIndex = 0; ErrorIndex < FilteredErrors.Num(); ErrorIndex++)
@@ -1220,18 +1222,11 @@ bool PreprocessD3DShader(
 		}
 	}
 
-	if (!Output.ParseAndModify(Input, Environment, TEXT("cbuffer")))
-	{
-		// The FShaderParameterParser will add any relevant errors.
-		return false;
-	}
-
 	CleanupUniformBufferCode(Input.Environment, PreprocessedSource);
 
 	// Process TEXT macro.
 	TransformStringIntoCharacterArray(PreprocessedSource, &Output.EditDiagnosticDatas());
 
-	TArray<FString> FilteredErrors;
 	// Run the shader minifier
 	#if UE_D3D_SHADER_COMPILER_ALLOW_DEAD_CODE_REMOVAL
 	if (Environment.CompilerFlags.Contains(CFLAG_RemoveDeadCode))
@@ -1239,11 +1234,12 @@ bool PreprocessD3DShader(
 		UE::ShaderCompilerCommon::RemoveDeadCode(PreprocessedSource, Input.EntryPointName, Output.EditErrors());
 	}
 	#endif // UE_D3D_SHADER_COMPILER_ALLOW_DEAD_CODE_REMOVAL
+
 	return true;
 }
 
 
-void CompileD3DShader(const FShaderCompilerInput& Input, const FShaderPreprocessOutput& PreprocessOutput, FShaderCompilerOutput& Output, const FString& WorkingDirectory, ELanguage Language)
+void CompileD3DShader(const FShaderCompilerInput& Input, const FString& InPreprocessedSource, FShaderCompilerOutput& Output, const FString& WorkingDirectory, ELanguage Language)
 {
 	TRACE_CPUPROFILER_EVENT_SCOPE(CompileD3DShader);
 
@@ -1256,41 +1252,66 @@ void CompileD3DShader(const FShaderCompilerInput& Input, const FShaderPreprocess
 		return;
 	}
 
+	FString EntryPointName = Input.EntryPointName;
+	FString PreprocessedSource = InPreprocessedSource;
+
+	FShaderParameterParser ShaderParameterParser(Input.Environment.CompilerFlags, TEXT("cbuffer"), {}, {});
+	if (!ShaderParameterParser.ParseAndModify(Input, Output.Errors, PreprocessedSource, EBindlessParameterMode::Default))
+	{
+		// The FShaderParameterParser will add any relevant errors.
+		return;
+	}
+
+	if (ShaderParameterParser.DidModifyShader())
+	{
+		Output.ModifiedShaderSource = PreprocessedSource;
+	}
+
 	if (Input.Environment.CompilerFlags.Contains(CFLAG_ForceRemoveUnusedInterpolators) && Input.Target.Frequency == SF_Vertex && Input.bCompilingForShaderPipeline)
 	{
 		// Always add SV_Position
-		TArray<FString> UsedOutputs = Input.UsedOutputs;
-		UsedOutputs.AddUnique(TEXT("SV_POSITION"));
-		UsedOutputs.AddUnique(TEXT("SV_ViewPortArrayIndex"));
+		TArray<FStringView> UsedOutputs;
+		for (const FString& UsedOutput : Input.UsedOutputs)
+		{
+			UsedOutputs.Emplace(UsedOutput);
+		}
+		UsedOutputs.Emplace(TEXTVIEW("SV_POSITION"));
+		UsedOutputs.Emplace(TEXTVIEW("SV_ViewPortArrayIndex"));
 
 		// We can't remove any of the output-only system semantics
 		//@todo - there are a bunch of tessellation ones as well
-		TArray<FString> Exceptions;
-		Exceptions.AddUnique(TEXT("SV_ClipDistance"));
-		Exceptions.AddUnique(TEXT("SV_ClipDistance0"));
-		Exceptions.AddUnique(TEXT("SV_ClipDistance1"));
-		Exceptions.AddUnique(TEXT("SV_ClipDistance2"));
-		Exceptions.AddUnique(TEXT("SV_ClipDistance3"));
-		Exceptions.AddUnique(TEXT("SV_ClipDistance4"));
-		Exceptions.AddUnique(TEXT("SV_ClipDistance5"));
-		Exceptions.AddUnique(TEXT("SV_ClipDistance6"));
-		Exceptions.AddUnique(TEXT("SV_ClipDistance7"));
-
-		Exceptions.AddUnique(TEXT("SV_CullDistance"));
-		Exceptions.AddUnique(TEXT("SV_CullDistance0"));
-		Exceptions.AddUnique(TEXT("SV_CullDistance1"));
-		Exceptions.AddUnique(TEXT("SV_CullDistance2"));
-		Exceptions.AddUnique(TEXT("SV_CullDistance3"));
-		Exceptions.AddUnique(TEXT("SV_CullDistance4"));
-		Exceptions.AddUnique(TEXT("SV_CullDistance5"));
-		Exceptions.AddUnique(TEXT("SV_CullDistance6"));
-		Exceptions.AddUnique(TEXT("SV_CullDistance7"));
-
-		Output.ModifiedShaderSource = PreprocessOutput.GetSource();
-		Output.ModifiedEntryPointName = Input.EntryPointName;
+		const FStringView Exceptions[] =
+		{
+			TEXTVIEW("SV_ClipDistance"),
+			TEXTVIEW("SV_ClipDistance0"),
+			TEXTVIEW("SV_ClipDistance1"),
+			TEXTVIEW("SV_ClipDistance2"),
+			TEXTVIEW("SV_ClipDistance3"),
+			TEXTVIEW("SV_ClipDistance4"),
+			TEXTVIEW("SV_ClipDistance5"),
+			TEXTVIEW("SV_ClipDistance6"),
+			TEXTVIEW("SV_ClipDistance7"),
+
+			TEXTVIEW("SV_CullDistance"),
+			TEXTVIEW("SV_CullDistance0"),
+			TEXTVIEW("SV_CullDistance1"),
+			TEXTVIEW("SV_CullDistance2"),
+			TEXTVIEW("SV_CullDistance3"),
+			TEXTVIEW("SV_CullDistance4"),
+			TEXTVIEW("SV_CullDistance5"),
+			TEXTVIEW("SV_CullDistance6"),
+			TEXTVIEW("SV_CullDistance7"),
+		};
+
+		TArray<FScopedDeclarations> ScopedDeclarations;
+		const FStringView GlobalSymbols[] =
+		{
+			TEXTVIEW("RayDesc"),
+		};
+		ScopedDeclarations.Emplace(TConstArrayView<FStringView>(), GlobalSymbols);
 
 		TArray<FString> Errors;
-		if (!RemoveUnusedOutputs(Output.ModifiedShaderSource, UsedOutputs, Exceptions, Output.ModifiedEntryPointName, Errors))
+		if (!RemoveUnusedOutputs(PreprocessedSource, UsedOutputs, Exceptions, ScopedDeclarations, EntryPointName, Errors))
 		{
 			UE_LOG(LogD3D11ShaderCompiler, Warning, TEXT("Failed to remove unused outputs from shader: %s"), *Input.GenerateShaderName());
 			for (const FString& ErrorReport : Errors)
@@ -1302,13 +1323,16 @@ void CompileD3DShader(const FShaderCompilerInput& Input, const FShaderPreprocess
 				Output.Errors.Add(NewError);
 			}
 		}
+		else
+		{
+			Output.ModifiedEntryPointName = EntryPointName;
+			Output.ModifiedShaderSource = PreprocessedSource;
+		}
 	}
 
-	const FShaderParameterParser& ShaderParameterParser = PreprocessOutput.GetParameterParser();
-
 	const bool bSuccess = bUseDXC
-		? CompileAndProcessD3DShaderDXC(PreprocessOutput, Input, ShaderParameterParser, ShaderProfile, Language, false, Output)
-		: CompileAndProcessD3DShaderFXC(PreprocessOutput, Input, ShaderParameterParser, ShaderProfile, false, Output);
+		? CompileAndProcessD3DShaderDXC(Input, PreprocessedSource, EntryPointName, ShaderParameterParser, ShaderProfile, Language, false, Output)
+		: CompileAndProcessD3DShaderFXC(Input, PreprocessedSource, EntryPointName, ShaderParameterParser, ShaderProfile, false, Output);
 
 	if (!bSuccess && !Output.Errors.Num())
 	{
@@ -1325,6 +1349,4 @@ void CompileD3DShader(const FShaderCompilerInput& Input, const FShaderPreprocess
 			FPlatformMisc::LowLevelOutputDebugStringf(TEXT("%s\n"), *Error.GetErrorStringWithLineMarker());
 		}
 	}
-
-	Output.ShaderDiagnosticDatas = PreprocessOutput.GetDiagnosticDatas();
 }
diff --git a/Engine/Source/Developer/Windows/ShaderFormatD3D/Private/D3DShaderCompilerDXC.cpp b/Engine/Source/Developer/Windows/ShaderFormatD3D/Private/D3DShaderCompilerDXC.cpp
--- a/Engine/Source/Developer/Windows/ShaderFormatD3D/Private/D3DShaderCompilerDXC.cpp
+++ b/Engine/Source/Developer/Windows/ShaderFormatD3D/Private/D3DShaderCompilerDXC.cpp
@@ -922,8 +922,10 @@ static bool DXCRewriteWrapper(
 }
 
 // Generate the dumped usf file; call the D3D compiler, gather reflection information and generate the output data
-bool CompileAndProcessD3DShaderDXC(const FShaderPreprocessOutput& PreprocessOutput,
+bool CompileAndProcessD3DShaderDXC(
 	const FShaderCompilerInput& Input,
+	const FString& PreprocessedShaderSource,
+	const FString& EntryPointName,
 	const FShaderParameterParser& ShaderParameterParser,
 	const TCHAR* ShaderProfile,
 	ELanguage Language,
@@ -932,9 +934,6 @@ bool CompileAndProcessD3DShaderDXC(const FShaderPreprocessOutput& PreprocessOutp
 {
 	TRACE_CPUPROFILER_EVENT_SCOPE(CompileAndProcessD3DShaderDXC);
 
-	const FString& PreprocessedShaderSource = Output.ModifiedShaderSource.IsEmpty() ? PreprocessOutput.GetSource() : Output.ModifiedShaderSource;
-	const FString& EntryPointName = Output.ModifiedEntryPointName.IsEmpty() ? Input.EntryPointName : Output.ModifiedEntryPointName;
-
 	auto AnsiSourceFile = StringCast<ANSICHAR>(*PreprocessedShaderSource);
 
 	const bool bIsRayTracingShader = Input.IsRayTracingShader();
diff --git a/Engine/Source/Developer/Windows/ShaderFormatD3D/Private/ShaderFormatD3D.cpp b/Engine/Source/Developer/Windows/ShaderFormatD3D/Private/ShaderFormatD3D.cpp
--- a/Engine/Source/Developer/Windows/ShaderFormatD3D/Private/ShaderFormatD3D.cpp
+++ b/Engine/Source/Developer/Windows/ShaderFormatD3D/Private/ShaderFormatD3D.cpp
@@ -110,7 +110,9 @@ public:
 		struct FShaderCompilerOutput& Output,
 		const FString& WorkingDirectory) const
 	{
-		CompileD3DShader(Input, PreprocessOutput, Output, WorkingDirectory, FormatToLanguage(Input.ShaderFormat));
+		CompileD3DShader(Input, PreprocessOutput.GetSource(), Output, WorkingDirectory, FormatToLanguage(Input.ShaderFormat));
+
+		Output.ShaderDiagnosticDatas = PreprocessOutput.GetDiagnosticDatas();
 	}
 
 	virtual bool SupportsIndependentPreprocessing() const
diff --git a/Engine/Source/Developer/Windows/ShaderFormatD3D/Private/ShaderFormatD3D.h b/Engine/Source/Developer/Windows/ShaderFormatD3D/Private/ShaderFormatD3D.h
--- a/Engine/Source/Developer/Windows/ShaderFormatD3D/Private/ShaderFormatD3D.h
+++ b/Engine/Source/Developer/Windows/ShaderFormatD3D/Private/ShaderFormatD3D.h
@@ -31,32 +31,34 @@ inline bool IsUsingSM66(const FShaderCompilerInput& Input, ELanguage Language)
 }
 
 bool PreprocessD3DShader(
-	const struct FShaderCompilerInput& Input,
-	const struct FShaderCompilerEnvironment& MergedEnvironment,
-	class FShaderPreprocessOutput& PreprocessOutput,
+	const FShaderCompilerInput& Input,
+	const FShaderCompilerEnvironment& MergedEnvironment,
+	FShaderPreprocessOutput& PreprocessOutput,
 	ELanguage Language);
 
 void CompileD3DShader(
-	const struct FShaderCompilerInput& Input,
-	const class FShaderPreprocessOutput& PreprocessOutput,
-	struct FShaderCompilerOutput& Output,
-	const class FString& WorkingDirectory,
+	const FShaderCompilerInput& Input,
+	const FString& InPreprocessedSource,
+	FShaderCompilerOutput& Output,
+	const FString& WorkingDirectory,
 	ELanguage Language);
 
 /**
  * @param bSecondPassAferUnusedInputRemoval whether we're compiling the shader second time, after having removed the unused inputs discovered in the first pass
  */
 bool CompileAndProcessD3DShaderFXC(
-	const FShaderPreprocessOutput& PreprocessOutput,
 	const FShaderCompilerInput& Input,
+	const FString& InPreprocessedSource,
+	const FString& InEntryPointName,
 	const FShaderParameterParser& ShaderParameterParser,
 	const TCHAR* ShaderProfile,
 	bool bSecondPassAferUnusedInputRemoval,
 	FShaderCompilerOutput& Output);
 
 bool CompileAndProcessD3DShaderDXC(
-	const FShaderPreprocessOutput& PreprocessOutput,
 	const FShaderCompilerInput& Input,
+	const FString& InPreprocessedSource,
+	const FString& InEntryPointName,
 	const FShaderParameterParser& ShaderParameterParser,
 	const TCHAR* ShaderProfile,
 	ELanguage Language,
diff --git a/Engine/Source/Runtime/RenderCore/Private/ShaderParameterParser.cpp b/Engine/Source/Runtime/RenderCore/Private/ShaderParameterParser.cpp
--- a/Engine/Source/Runtime/RenderCore/Private/ShaderParameterParser.cpp
+++ b/Engine/Source/Runtime/RenderCore/Private/ShaderParameterParser.cpp
@@ -1022,6 +1022,8 @@ void FShaderParameterParser::ApplyBindlessModifications(FString& PreprocessedSha
 
 			// Commit all modifications to caller
 			PreprocessedShaderSource = NewShaderCode;
+
+			bModifiedShader = true;
 		}
 	}
 }
@@ -1122,6 +1124,7 @@ bool FShaderParameterParser::MoveShaderParametersToRootConstantBuffer(
 		PreprocessedShaderSource = MoveTemp(NewShaderCode);
 
 		bMovedLoosedParametersToRootConstantBuffer = true;
+		bModifiedShader = true;
 	} // if (CompilerInput.RootParametersStructure && bNeedToMoveToRootConstantBuffer)
 
 	return bSuccess;
@@ -1167,10 +1170,26 @@ bool FShaderParameterParser::ParseAndModify(
 	RemoveMovingParametersFromSource(PreprocessedShaderSource);
 	ApplyBindlessModifications(PreprocessedShaderSource);
 
+	bool bResult = true;
+
 	if (bNeedToMoveToRootConstantBuffer)
 	{
-		return MoveShaderParametersToRootConstantBuffer(CompilerInput.RootParametersStructure, PreprocessedShaderSource);
+		bResult = MoveShaderParametersToRootConstantBuffer(CompilerInput.RootParametersStructure, PreprocessedShaderSource);
+	}
+
+#if DO_GUARD_SLOW
+	if (bResult)
+	{
+		if (DidModifyShader())
+		{
+			checkSlow(PreprocessedShaderSource != OriginalParsedShader);
+		}
+		else
+		{
+			checkSlow(PreprocessedShaderSource == OriginalParsedShader);
+		}
 	}
+#endif
 	
 	return true;
 }
@@ -1408,6 +1427,7 @@ void SerializeParam(FArchive& Ar, FShaderParameterParser::FParsedShaderParameter
 FArchive& operator<<(FArchive& Ar, FShaderParameterParser& Parser)
 {
 	Ar << Parser.bMovedLoosedParametersToRootConstantBuffer;
+	Ar << Parser.bModifiedShader;
 	Ar << Parser.OriginalParsedShader;
 	if (Ar.IsSaving())
 	{
diff --git a/Engine/Source/Runtime/RenderCore/Public/ShaderParameterParser.h b/Engine/Source/Runtime/RenderCore/Public/ShaderParameterParser.h
--- a/Engine/Source/Runtime/RenderCore/Public/ShaderParameterParser.h
+++ b/Engine/Source/Runtime/RenderCore/Public/ShaderParameterParser.h
@@ -182,6 +182,8 @@ public:
 		return ExtractFileAndLine(ParsedParameter.ParsedPragmaLineOffset, ParsedParameter.ParsedLineOffset, OutFile, OutLine);
 	}
 
+	bool DidModifyShader() const { return bModifiedShader; }
+
 	friend FArchive& operator<<(FArchive& Ar, FShaderParameterParser& Parser);
 
 protected:
@@ -232,4 +234,7 @@ protected:
 
 	/** Indicates that parameters were actually moved to the root constant buffer. */
 	bool bMovedLoosedParametersToRootConstantBuffer = false;
+
+	/** Indicates that the shader source was actually modified. */
+	bool bModifiedShader = false;
 };
