Always set the HLSL version and don't assume the default is 2018.

diff --git a/Engine/Source/Developer/ShaderCompilerCommon/Private/ShaderConductorContext.cpp b/Engine/Source/Developer/ShaderCompilerCommon/Private/ShaderConductorContext.cpp
--- a/Engine/Source/Developer/ShaderCompilerCommon/Private/ShaderConductorContext.cpp
+++ b/Engine/Source/Developer/ShaderCompilerCommon/Private/ShaderConductorContext.cpp
@@ -431,7 +431,8 @@ namespace CrossCompiler
 			DxcArgRefs.Add("2017");
 			break;
 		case 2018:
-			// Default
+			DxcArgRefs.Add("-HV");
+			DxcArgRefs.Add("2018");
 			break;
 		case 2021:
 			DxcArgRefs.Add("-HV");
diff --git a/Engine/Source/Developer/Windows/ShaderFormatD3D/Private/D3DShaderCompilerDXC.cpp b/Engine/Source/Developer/Windows/ShaderFormatD3D/Private/D3DShaderCompilerDXC.cpp
--- a/Engine/Source/Developer/Windows/ShaderFormatD3D/Private/D3DShaderCompilerDXC.cpp
+++ b/Engine/Source/Developer/Windows/ShaderFormatD3D/Private/D3DShaderCompilerDXC.cpp
@@ -198,31 +198,15 @@ public:
 			|| Input.IsRayTracingShader();
 
 		const bool bHlslVersion2021 = Input.Environment.CompilerFlags.Contains(CFLAG_HLSL2021);
-		const uint32 HlslVersion = (bHlslVersion2021 ? 2021 : 2018);
-
-		switch (HlslVersion)
+		if (bHlslVersion2021)
 		{
-		case 2015:
-			ExtraArguments.Add(TEXT("-HV"));
-			ExtraArguments.Add(TEXT("2015"));
-			break;
-		case 2016:
-			ExtraArguments.Add(TEXT("-HV"));
-			ExtraArguments.Add(TEXT("2016"));
-			break;
-		case 2017:
-			ExtraArguments.Add(TEXT("-HV"));
-			ExtraArguments.Add(TEXT("2017"));
-			break;
-		case 2018:
-			break; // Default
-		case 2021:
 			ExtraArguments.Add(TEXT("-HV"));
 			ExtraArguments.Add(TEXT("2021"));
-			break;
-		default:
-			checkf(false, TEXT("Invalid HLSL version: expected 2015, 2016, 2017, 2018, or 2021 but %u was specified"), HlslVersion);
-			break;
+		}
+		else
+		{
+			ExtraArguments.Add(TEXT("-HV"));
+			ExtraArguments.Add(TEXT("2018"));
 		}
 
 		// Unpack uniform matrices as row-major to match the CPU layout.
