Removed dead code from D3D12 shader backend.

#rnx

diff --git a/Engine/Source/Developer/Windows/ShaderFormatD3D/Private/D3DShaderCompilerDXC.cpp b/Engine/Source/Developer/Windows/ShaderFormatD3D/Private/D3DShaderCompilerDXC.cpp
--- a/Engine/Source/Developer/Windows/ShaderFormatD3D/Private/D3DShaderCompilerDXC.cpp
+++ b/Engine/Source/Developer/Windows/ShaderFormatD3D/Private/D3DShaderCompilerDXC.cpp
@@ -845,66 +845,6 @@ inline bool IsCompatibleBinding(const D3D12_SHADER_INPUT_BIND_DESC& BindDesc, ui
 	return bIsCompatibleBinding;
 }
 
-static ShaderConductor::Compiler::ShaderModel ToDXCShaderModel(ELanguage Language)
-{
-	switch (Language)
-	{
-	case ELanguage::ES3_1:
-	case ELanguage::SM5:
-		return { 5, 0 };
-	default:
-		UE_LOG(LogD3D12ShaderCompiler, Error, TEXT("Invalid input shader target for enum ELanguage (%d)."), (int32)Language);
-	}
-	return { 6,0 };
-}
-
-static ShaderConductor::ShaderStage ToDXCShaderStage(EShaderFrequency Frequency)
-{
-	check(Frequency >= SF_Vertex && Frequency <= SF_Compute);
-	switch (Frequency)
-	{
-	case SF_Vertex:		return ShaderConductor::ShaderStage::VertexShader;
-	case SF_Pixel:		return ShaderConductor::ShaderStage::PixelShader;
-	case SF_Geometry:	return ShaderConductor::ShaderStage::GeometryShader;
-	case SF_Compute:	return ShaderConductor::ShaderStage::ComputeShader;
-	default:			return ShaderConductor::ShaderStage::NumShaderStages;
-	}
-}
-
-// Inner wrapper function is required here because '__try'-statement cannot be used with function that requires object unwinding
-static void InnerScRewriteWrapper(
-	const ShaderConductor::Compiler::SourceDesc& InDesc,
-	const ShaderConductor::Compiler::Options& InOptions,
-	ShaderConductor::Compiler::ResultDesc& OutResultDesc)
-{
-	OutResultDesc = ShaderConductor::Compiler::Rewrite(InDesc, InOptions);
-}
-
-static bool DXCRewriteWrapper(
-	const ShaderConductor::Compiler::SourceDesc& InDesc,
-	const ShaderConductor::Compiler::Options& InOptions,
-	ShaderConductor::Compiler::ResultDesc& OutResultDesc,
-	bool& bOutException)
-{
-	bOutException = false;
-#if !PLATFORM_SEH_EXCEPTIONS_DISABLED
-	__try
-#endif
-	{
-		InnerScRewriteWrapper(InDesc, InOptions, OutResultDesc);
-		return true;
-	}
-#if !PLATFORM_SEH_EXCEPTIONS_DISABLED
-	__except (EXCEPTION_EXECUTE_HANDLER)
-	{
-		FSCWErrorCode::Report(FSCWErrorCode::CrashInsidePlatformCompiler);
-		FMemory::Memzero(OutResultDesc);
-		bOutException = true;
-		return false;
-	}
-#endif
-}
-
 // Generate the dumped usf file; call the D3D compiler, gather reflection information and generate the output data
 bool CompileAndProcessD3DShaderDXC(
 	const FShaderCompilerInput& Input,
