Removed usage of include file macros from LiveCoding source since it was causing issues with SN-DBS

diff --git a/Engine/Source/Developer/Windows/LiveCoding/Private/External/LC_Assert.h b/Engine/Source/Developer/Windows/LiveCoding/Private/External/LC_Assert.h
--- a/Engine/Source/Developer/Windows/LiveCoding/Private/External/LC_Assert.h
+++ b/Engine/Source/Developer/Windows/LiveCoding/Private/External/LC_Assert.h
@@ -4,9 +4,7 @@
 
 // BEGIN EPIC MOD
 #include "LC_Platform.h"
-// END EPIC MOD
-#include LC_PLATFORM_INCLUDE(LC_Debugger)
-// BEGIN EPIC MOD
+#include "LC_Debugger_Windows.h"
 #include "LC_Logging.h"
 // END EPIC MOD
 
diff --git a/Engine/Source/Developer/Windows/LiveCoding/Private/External/LC_ClientUserCommandThread.cpp b/Engine/Source/Developer/Windows/LiveCoding/Private/External/LC_ClientUserCommandThread.cpp
--- a/Engine/Source/Developer/Windows/LiveCoding/Private/External/LC_ClientUserCommandThread.cpp
+++ b/Engine/Source/Developer/Windows/LiveCoding/Private/External/LC_ClientUserCommandThread.cpp
@@ -18,7 +18,7 @@
 #include "LC_MemoryStream.h"
 // BEGIN EPIC MOD
 #include "LC_Platform.h"
-#include LC_PLATFORM_INCLUDE(LC_Foundation)
+#include "LC_Foundation_Windows.h"
 #include "LC_Thread.h"
 #include "LC_Logging.h"
 #include <deque>
diff --git a/Engine/Source/Developer/Windows/LiveCoding/Private/External/LC_CriticalSection.h b/Engine/Source/Developer/Windows/LiveCoding/Private/External/LC_CriticalSection.h
--- a/Engine/Source/Developer/Windows/LiveCoding/Private/External/LC_CriticalSection.h
+++ b/Engine/Source/Developer/Windows/LiveCoding/Private/External/LC_CriticalSection.h
@@ -6,7 +6,7 @@
 #include "CoreTypes.h"
 #include "Windows/MinimalWindowsApi.h"
 #include "LC_Platform.h"
-#include LC_PLATFORM_INCLUDE(LC_Foundation)
+#include "LC_Foundation_Windows.h"
 // END EPIC MOD
 
 class CriticalSection
diff --git a/Engine/Source/Developer/Windows/LiveCoding/Private/External/LC_FixedSizeString.h b/Engine/Source/Developer/Windows/LiveCoding/Private/External/LC_FixedSizeString.h
--- a/Engine/Source/Developer/Windows/LiveCoding/Private/External/LC_FixedSizeString.h
+++ b/Engine/Source/Developer/Windows/LiveCoding/Private/External/LC_FixedSizeString.h
@@ -4,7 +4,7 @@
 
 // BEGIN EPIC MOD
 #include "LC_Platform.h"
-#include LC_PLATFORM_INCLUDE(LC_Foundation)
+#include "LC_Foundation_Windows.h"
 // END EPIC MOD
 
 template <typename T, unsigned int N>
diff --git a/Engine/Source/Developer/Windows/LiveCoding/Private/External/LC_Hashing.h b/Engine/Source/Developer/Windows/LiveCoding/Private/External/LC_Hashing.h
--- a/Engine/Source/Developer/Windows/LiveCoding/Private/External/LC_Hashing.h
+++ b/Engine/Source/Developer/Windows/LiveCoding/Private/External/LC_Hashing.h
@@ -5,7 +5,7 @@
 #include "xxhash.h"
 // BEGIN EPIC MOD
 #include "LC_Platform.h"
-#include LC_PLATFORM_INCLUDE(LC_Foundation)
+#include "LC_Foundation_Windows.h"
 // END EPIC MOD
 
 namespace Hashing
diff --git a/Engine/Source/Developer/Windows/LiveCoding/Private/External/LC_InterprocessMutex.h b/Engine/Source/Developer/Windows/LiveCoding/Private/External/LC_InterprocessMutex.h
--- a/Engine/Source/Developer/Windows/LiveCoding/Private/External/LC_InterprocessMutex.h
+++ b/Engine/Source/Developer/Windows/LiveCoding/Private/External/LC_InterprocessMutex.h
@@ -5,7 +5,7 @@
 // BEGIN EPIC MOD
 #include "Windows/MinimalWindowsAPI.h"
 #include "LC_Platform.h"
-#include LC_PLATFORM_INCLUDE(LC_Foundation)
+#include "LC_Foundation_Windows.h"
 // END EPIC MOD
 
 class InterprocessMutex
diff --git a/Engine/Source/Developer/Windows/LiveCoding/Private/External/LC_MemoryStream.h b/Engine/Source/Developer/Windows/LiveCoding/Private/External/LC_MemoryStream.h
--- a/Engine/Source/Developer/Windows/LiveCoding/Private/External/LC_MemoryStream.h
+++ b/Engine/Source/Developer/Windows/LiveCoding/Private/External/LC_MemoryStream.h
@@ -5,7 +5,7 @@
 // BEGIN EPIC MOD
 #include "CoreTypes.h"
 #include "LC_Platform.h"
-#include LC_PLATFORM_INCLUDE(LC_Foundation)
+#include "LC_Foundation_Windows.h"
 // END EPIC MOD
 
 namespace memoryStream
diff --git a/Engine/Source/Developer/Windows/LiveCoding/Private/External/LC_Platform.h b/Engine/Source/Developer/Windows/LiveCoding/Private/External/LC_Platform.h
--- a/Engine/Source/Developer/Windows/LiveCoding/Private/External/LC_Platform.h
+++ b/Engine/Source/Developer/Windows/LiveCoding/Private/External/LC_Platform.h
@@ -31,8 +31,9 @@
 	#pragma error "Unknown platform"
 #endif
 
-
-#define LC_PLATFORM_INCLUDE(_header) LC_PP_STRINGIFY(LC_PP_JOIN(_header, LC_PLATFORM_SUFFIX))
+// BEGIN EPIC MOD
+//#define LC_PLATFORM_INCLUDE(_header) LC_PP_STRINGIFY(LC_PP_JOIN(_header, LC_PLATFORM_SUFFIX))
+// END EPIC MOD
 
 // BEGIN EPIC MOD
 // convenience macro for referring to symbol identifiers.
diff --git a/Engine/Source/Developer/Windows/LiveCodingServer/Private/External/LC_Assembler.cpp b/Engine/Source/Developer/Windows/LiveCodingServer/Private/External/LC_Assembler.cpp
--- a/Engine/Source/Developer/Windows/LiveCodingServer/Private/External/LC_Assembler.cpp
+++ b/Engine/Source/Developer/Windows/LiveCodingServer/Private/External/LC_Assembler.cpp
@@ -6,7 +6,7 @@
 #include "LC_Assembler.h"
 // BEGIN EPIC MOD
 #include "LC_Platform.h"
-#include LC_PLATFORM_INCLUDE(LC_Foundation)
+#include "LC_Foundation_Windows.h"
 #include "LC_Assert.h"
 #include <cstdlib>
 // END EPIC MOD
diff --git a/Engine/Source/Developer/Windows/LiveCodingServer/Private/External/LC_COMThread.cpp b/Engine/Source/Developer/Windows/LiveCodingServer/Private/External/LC_COMThread.cpp
--- a/Engine/Source/Developer/Windows/LiveCodingServer/Private/External/LC_COMThread.cpp
+++ b/Engine/Source/Developer/Windows/LiveCodingServer/Private/External/LC_COMThread.cpp
@@ -8,7 +8,7 @@
 // BEGIN EPIC MOD
 #include "LC_Logging.h"
 #include "LC_Platform.h"
-#include LC_PLATFORM_INCLUDE(LC_Foundation)
+#include "LC_Foundation_Windows.h"
 #include "Windows/MinimalWindowsAPI.h"
 #include <objidl.h>
 
diff --git a/Engine/Source/Developer/Windows/LiveCodingServer/Private/External/LC_DiaSymbolName.h b/Engine/Source/Developer/Windows/LiveCodingServer/Private/External/LC_DiaSymbolName.h
--- a/Engine/Source/Developer/Windows/LiveCodingServer/Private/External/LC_DiaSymbolName.h
+++ b/Engine/Source/Developer/Windows/LiveCodingServer/Private/External/LC_DiaSymbolName.h
@@ -5,7 +5,7 @@
 // BEGIN EPIC MOD
 #include "CoreTypes.h"
 #include "LC_Platform.h"
-#include LC_PLATFORM_INCLUDE(LC_Foundation)
+#include "LC_Foundation_Windows.h"
 #include "Windows/WindowsHWrapper.h"
 #include <dia2.h>
 // END EPIC MOD
diff --git a/Engine/Source/Developer/Windows/LiveCodingServer/Private/External/LC_DiaVariant.h b/Engine/Source/Developer/Windows/LiveCodingServer/Private/External/LC_DiaVariant.h
--- a/Engine/Source/Developer/Windows/LiveCodingServer/Private/External/LC_DiaVariant.h
+++ b/Engine/Source/Developer/Windows/LiveCodingServer/Private/External/LC_DiaVariant.h
@@ -5,7 +5,7 @@
 // BEGIN EPIC MOD
 #include "CoreTypes.h"
 #include "LC_Platform.h"
-#include LC_PLATFORM_INCLUDE(LC_Foundation)
+#include "LC_Foundation_Windows.h"
 #include "Windows/WindowsHWrapper.h"
 #include <dia2.h>
 // END EPIC MOD
diff --git a/Engine/Source/Developer/Windows/LiveCodingServer/Private/External/LC_GrowingMemoryBlock.h b/Engine/Source/Developer/Windows/LiveCodingServer/Private/External/LC_GrowingMemoryBlock.h
--- a/Engine/Source/Developer/Windows/LiveCodingServer/Private/External/LC_GrowingMemoryBlock.h
+++ b/Engine/Source/Developer/Windows/LiveCodingServer/Private/External/LC_GrowingMemoryBlock.h
@@ -5,7 +5,7 @@
 // BEGIN EPIC MOD
 #include "CoreTypes.h"
 #include "LC_Platform.h"
-#include LC_PLATFORM_INCLUDE(LC_Foundation)
+#include "LC_Foundation_Windows.h"
 // END EPIC MOD
 
 class GrowingMemoryBlock
diff --git a/Engine/Source/Developer/Windows/LiveCodingServer/Private/External/LC_ImmutableString.h b/Engine/Source/Developer/Windows/LiveCodingServer/Private/External/LC_ImmutableString.h
--- a/Engine/Source/Developer/Windows/LiveCodingServer/Private/External/LC_ImmutableString.h
+++ b/Engine/Source/Developer/Windows/LiveCodingServer/Private/External/LC_ImmutableString.h
@@ -8,7 +8,7 @@
 #include "LC_Memory.h"
 // BEGIN EPIC MOD
 #include "LC_Platform.h"
-#include LC_PLATFORM_INCLUDE(LC_Foundation)
+#include "LC_Foundation_Windows.h"
 #include <string>
 // END EPIC MOD
 
diff --git a/Engine/Source/Developer/Windows/LiveCodingServer/Private/External/LC_Mallocator.cpp b/Engine/Source/Developer/Windows/LiveCodingServer/Private/External/LC_Mallocator.cpp
--- a/Engine/Source/Developer/Windows/LiveCodingServer/Private/External/LC_Mallocator.cpp
+++ b/Engine/Source/Developer/Windows/LiveCodingServer/Private/External/LC_Mallocator.cpp
@@ -7,7 +7,7 @@
 // BEGIN EPIC MOD
 #include "LC_Assert.h"
 #include "LC_Platform.h"
-#include LC_PLATFORM_INCLUDE(LC_Foundation)
+#include "LC_Foundation_Windows.h"
 #include "LC_Logging.h"
 #include "Windows/WindowsHWrapper.h"
 // END EPIC MOD
diff --git a/Engine/Source/Developer/Windows/LiveCodingServer/Private/External/LC_MemoryBlock.h b/Engine/Source/Developer/Windows/LiveCodingServer/Private/External/LC_MemoryBlock.h
--- a/Engine/Source/Developer/Windows/LiveCodingServer/Private/External/LC_MemoryBlock.h
+++ b/Engine/Source/Developer/Windows/LiveCodingServer/Private/External/LC_MemoryBlock.h
@@ -5,7 +5,7 @@
 // BEGIN EPIC MOD
 #include "CoreTypes.h"
 #include "LC_Platform.h"
-#include LC_PLATFORM_INCLUDE(LC_Foundation)
+#include "LC_Foundation_Windows.h"
 // END EPIC MOD
 
 class MemoryBlock
diff --git a/Engine/Source/Developer/Windows/LiveCodingServer/Private/External/LC_SymbolInfo.h b/Engine/Source/Developer/Windows/LiveCodingServer/Private/External/LC_SymbolInfo.h
--- a/Engine/Source/Developer/Windows/LiveCodingServer/Private/External/LC_SymbolInfo.h
+++ b/Engine/Source/Developer/Windows/LiveCodingServer/Private/External/LC_SymbolInfo.h
@@ -5,7 +5,7 @@
 // BEGIN EPIC MOD
 #include "CoreTypes.h"
 #include "LC_Platform.h"
-#include LC_PLATFORM_INCLUDE(LC_Foundation)
+#include "LC_Foundation_Windows.h"
 // END EPIC MOD
 
 class SymbolInfo
