Workaround FXC validation errors for cross-compiled HLSL:
- Compile a 2nd time internally defined rule disabled.
- Compile a 3rd time without optimizations as a last resort.

#rnx

diff --git a/Engine/Source/Developer/Windows/ShaderFormatD3D/Private/D3DShaderCompiler.cpp b/Engine/Source/Developer/Windows/ShaderFormatD3D/Private/D3DShaderCompiler.cpp
--- a/Engine/Source/Developer/Windows/ShaderFormatD3D/Private/D3DShaderCompiler.cpp
+++ b/Engine/Source/Developer/Windows/ShaderFormatD3D/Private/D3DShaderCompiler.cpp
@@ -683,6 +683,16 @@ static bool CompileErrorsContainInternalError(ID3DBlob* Errors)
 	return false;
 }
 
+static bool D3DCompileErrorContainsValidationErrors(ID3DBlob* ErrorBlob)
+{
+	if (ErrorBlob != nullptr)
+	{
+		const FAnsiStringView ErrorString((const ANSICHAR*)ErrorBlob->GetBufferPointer(), (int32)ErrorBlob->GetBufferSize());
+		return (ErrorString.Find(ANSITEXTVIEW("error X8000: Validation Error:")) != INDEX_NONE);
+	}
+	return false;
+}
+
 // Generate the dumped usf file; call the D3D compiler, gather reflection information and generate the output data
 static bool CompileAndProcessD3DShaderFXCExt(
 	uint32 CompileFlags,
@@ -823,23 +833,55 @@ static bool CompileAndProcessD3DShaderFXCExt(
 				return FPaths::Combine(PathPart, FilenamePart) + TEXT(".intermediate.") + ExtensionPart;
 			};
 
-			// Compile again with FXC:
-			// SPIRV-Cross will have generated the new shader with "main" as the new entry point.
 			const FString CrossCompiledSourceFilename = MakeIntermediateVirtualSourceFilePath(Input.VirtualSourceFilePath);
-			Result = D3DCompileWrapper(
-				D3DCompileFunc,
-				CrossCompiledSource.GetData(),
-				CrossCompiledSource.Num() - 1,
-				TCHAR_TO_ANSI(*CrossCompiledSourceFilename),
-				/*pDefines=*/ NULL,
-				/*pInclude=*/ NULL,
-				"main",
-				TCHAR_TO_ANSI(ShaderProfile),
-				CompileFlags & (~D3DCOMPILE_WARNINGS_ARE_ERRORS),
-				0,
-				Shader.GetInitReference(),
-				Errors.GetInitReference()
-			);
+			auto ShaderProfileAnsi = StringCast<ANSICHAR>(ShaderProfile);
+			auto CrossCompiledSourceFilenameAnsi = StringCast<ANSICHAR>(*CrossCompiledSourceFilename);
+
+			// SPIRV-Cross will have generated the new shader with "main" as the new entry point.
+			auto CompileCrossCompiledHlsl = [&D3DCompileFunc, &CrossCompiledSourceFilenameAnsi, &Shader, &Errors, &ShaderProfileAnsi](const TArray<ANSICHAR>& Source, uint32 CompileFlags, const ANSICHAR* EntryPoint = "main") -> HRESULT
+			{
+				checkf(Source.Num() > 0, TEXT("TArray<ANSICHAR> of cross-compiled HLSL source must have at least one element including the NUL-terminator"));
+				return D3DCompileWrapper(
+					D3DCompileFunc,
+					Source.GetData(),
+					static_cast<SIZE_T>(Source.Num() - 1),
+					CrossCompiledSourceFilenameAnsi.Get(),
+					/*pDefines=*/ NULL,
+					/*pInclude=*/ NULL,
+					EntryPoint,
+					ShaderProfileAnsi.Get(),
+					CompileFlags,
+					0,
+					Shader.GetInitReference(),
+					Errors.GetInitReference()
+				);
+			};
+
+			// Compile again with FXC - 1st try
+			const uint32 CompileFlagsNoWarningsAsErrors = CompileFlags & (~D3DCOMPILE_WARNINGS_ARE_ERRORS);
+			Result = CompileCrossCompiledHlsl(CrossCompiledSource, CompileFlagsNoWarningsAsErrors);
+
+			// If FXC compilation failed with a validation error, assume bug in FXC's optimization passes
+			// Compile again with FXC and disable special compiler rule to simplify control flow - 2nd try
+			if (Result == E_FAIL && D3DCompileErrorContainsValidationErrors(Errors.GetReference()))
+			{
+				Output.Errors.Add(FShaderCompilerError(TEXT("Validation error in FXC encountered: Compiling intermediate HLSL a second time with simplified control flow")));
+
+				// Rule 0x08024065 is described as "simplify flow control that writes the same value in each flow control path"
+				const FAnsiStringView PragmaDirectiveCode = "#pragma ruledisable 0x08024065\n";
+				CrossCompiledSource.Insert(PragmaDirectiveCode.GetData(), PragmaDirectiveCode.Len(), 0);
+
+				Result = CompileCrossCompiledHlsl(CrossCompiledSource, CompileFlagsNoWarningsAsErrors);
+
+				// If FXC compilation still fails with a validation error, compile again and skip optimizations entirely as a last resort - 3rd try
+				if (Result == E_FAIL && D3DCompileErrorContainsValidationErrors(Errors.GetReference()))
+				{
+					Output.Errors.Add(FShaderCompilerError(TEXT("Validation error in FXC encountered: Compiling intermediate HLSL a third time without optimization (D3DCOMPILE_SKIP_OPTIMIZATION)")));
+
+					const uint32 CompileFlagsSkipOptimizations = CompileFlagsNoWarningsAsErrors | D3DCOMPILE_SKIP_OPTIMIZATION;
+					Result = CompileCrossCompiledHlsl(CrossCompiledSource, CompileFlagsSkipOptimizations);
+				}
+			}
 
 			if (!bPrecompileWithDXC && SUCCEEDED(Result))
 			{
