Reworking D3D Shader Format support checks
- ELanguage and Input flags were being used inconsistently to determine defines and shader profiles to use.
- ED3DShaderModel now communicates the exact shader model to use to the FXC and DXC backends.
- Moved all D3D custom definitions to ModifyShaderCompilerInput
- D3D definitions have been fixed to be based on specific needs of input flags. ie- "PLATFORM_MAX_SAMPLERS" doesn't imply that wave operations are wanted/supported.

diff --git a/Engine/Source/Developer/Windows/ShaderFormatD3D/Private/D3DShaderCompiler.cpp b/Engine/Source/Developer/Windows/ShaderFormatD3D/Private/D3DShaderCompiler.cpp
--- a/Engine/Source/Developer/Windows/ShaderFormatD3D/Private/D3DShaderCompiler.cpp
+++ b/Engine/Source/Developer/Windows/ShaderFormatD3D/Private/D3DShaderCompiler.cpp
@@ -106,79 +106,48 @@ static void D3D11FilterShaderCompileWarnings(const FString& CompileWarnings, TAr
 }
 
 // @return 0 if not recognized
-static const TCHAR* GetShaderProfileName(ELanguage Language, EShaderFrequency Frequency, bool bForceSM6)
+static const TCHAR* GetShaderProfileName(const FShaderCompilerInput& Input, ED3DShaderModel ShaderModel)
 {
-	if (IsUsingSM66(Language, Frequency))
+	if (ShaderModel == ED3DShaderModel::SM6_6)
 	{
-		switch (Frequency)
+		switch (Input.Target.GetFrequency())
 		{
-		default:
-			checkfSlow(false, TEXT("Unexpected shader frequency"));
-			return nullptr;
-		case SF_Pixel:
-			return TEXT("ps_6_6");
-		case SF_Vertex:
-			return TEXT("vs_6_6");
-		case SF_Mesh:
-			return TEXT("ms_6_6");
-		case SF_Amplification:
-			return TEXT("as_6_6");
-		case SF_Geometry:
-			return TEXT("gs_6_6");
-		case SF_Compute:
-			return TEXT("cs_6_6");
+		case SF_Pixel:         return TEXT("ps_6_6");
+		case SF_Vertex:        return TEXT("vs_6_6");
+		case SF_Mesh:          return TEXT("ms_6_6");
+		case SF_Amplification: return TEXT("as_6_6");
+		case SF_Geometry:      return TEXT("gs_6_6");
+		case SF_Compute:       return TEXT("cs_6_6");
 		case SF_RayGen:
 		case SF_RayMiss:
 		case SF_RayHitGroup:
-		case SF_RayCallable:
-			return TEXT("lib_6_6");
+		case SF_RayCallable:   return TEXT("lib_6_6");
 		}
 	}
-	else if(Language == ELanguage::SM5)
+	else if (ShaderModel == ED3DShaderModel::SM6_0)
 	{
 		//set defines and profiles for the appropriate shader paths
-		switch(Frequency)
+		switch (Input.Target.GetFrequency())
 		{
-		default:
-			checkfSlow(false, TEXT("Unexpected shader frequency"));
-			return nullptr;
-		case SF_Pixel:
-			return bForceSM6 ? TEXT("ps_6_0") : TEXT("ps_5_0");
-		case SF_Vertex:
-			return bForceSM6 ? TEXT("vs_6_0") : TEXT("vs_5_0");
-		case SF_Geometry:
-			return bForceSM6 ? TEXT("gs_6_0") : TEXT("gs_5_0");
-		case SF_Compute:
-			return bForceSM6 ? TEXT("cs_6_0") : TEXT("cs_5_0");
-		case SF_RayGen:
-		case SF_RayMiss:
-		case SF_RayHitGroup:
-		case SF_RayCallable:
-			checkNoEntry();
-			return TEXT("lib_6_5");
+		case SF_Pixel:    return TEXT("ps_6_0");
+		case SF_Vertex:   return TEXT("vs_6_0");
+		case SF_Geometry: return TEXT("gs_6_0");
+		case SF_Compute:  return TEXT("cs_6_0");
 		}
 	}
-	else if (Language == ELanguage::ES3_1)
+	else
 	{
-		checkSlow(Frequency == SF_Vertex ||
-			Frequency == SF_Pixel ||
-			Frequency == SF_Geometry ||
-			Frequency == SF_Compute);
-
 		//set defines and profiles for the appropriate shader paths
-		switch(Frequency)
+		switch (Input.Target.GetFrequency())
 		{
-		case SF_Pixel:
-			return bForceSM6 ? TEXT("ps_6_0") : TEXT("ps_5_0");
-		case SF_Vertex:
-			return bForceSM6 ? TEXT("vs_6_0") : TEXT("vs_5_0");
-		case SF_Geometry:
-			return bForceSM6 ? TEXT("gs_6_0") : TEXT("gs_5_0");
-		case SF_Compute:
-			return bForceSM6 ? TEXT("cs_6_0") : TEXT("cs_5_0");
+		case SF_Pixel:    return TEXT("ps_5_0");
+		case SF_Vertex:   return TEXT("vs_5_0");
+		case SF_Geometry: return TEXT("gs_5_0");
+		case SF_Compute:  return TEXT("cs_5_0");
 		}
 	}
 
+	checkfSlow(false, TEXT("Unexpected shader frequency"));
 	return nullptr;
 }
 
@@ -1291,71 +1260,16 @@ bool CompileAndProcessD3DShaderFXC(
 	return bSuccess;
 }
 
-inline bool ShouldUseDXC(const FShaderCompilerInput& Input, ELanguage Language)
-{
-	return IsUsingSM66(Input, Language)
-		|| Input.Environment.CompilerFlags.Contains(CFLAG_WaveOperations)
-		|| Input.Environment.CompilerFlags.Contains(CFLAG_ForceDXC)
-		|| Input.Environment.CompilerFlags.Contains(CFLAG_InlineRayTracing)
-		|| Input.Environment.GetIntegerValue(TEXT("PLATFORM_MAX_SAMPLERS")) > 16;
-}
-
 bool PreprocessD3DShader(
 	const FShaderCompilerInput& Input,
 	const FShaderCompilerEnvironment& Environment,
-	FShaderPreprocessOutput& Output,
-	ELanguage Language)
+	FShaderPreprocessOutput& Output)
 {
 	PRAGMA_DISABLE_DEPRECATION_WARNINGS		// FShaderCompilerDefinitions will be made internal in the future, marked deprecated until then
 	FShaderCompilerDefinitions AdditionalDefines;
-	if (IsUsingSM66(Input, Language))
-	{
-		AdditionalDefines.SetDefine(TEXT("SM6_PROFILE"), 1);
-	}
-	else if (Language == ELanguage::SM5)
-	{
-		AdditionalDefines.SetDefine(TEXT("SM5_PROFILE"), 1);
-	}
-	else if (Language == ELanguage::ES3_1)
-	{
-		AdditionalDefines.SetDefine(TEXT("ES3_1_PROFILE"), 1);
-	}
+	PRAGMA_ENABLE_DEPRECATION_WARNINGS
 
 	FString& PreprocessedSource = Output.EditSource();
-	const bool bUseDXC = ShouldUseDXC(Input, Language);
-
-
-	// Set additional defines.
-	AdditionalDefines.SetDefine(TEXT("COMPILER_HLSL"), 1);
-	AdditionalDefines.SetDefine(TEXT("PLATFORM_SUPPORTS_ROV"), 1); // Assume min. spec HW supports with DX12/SM5
-
-	if (bUseDXC)
-	{
-		AdditionalDefines.SetDefine(TEXT("PLATFORM_SUPPORTS_CALLABLE_SHADERS"), 1);
-		AdditionalDefines.SetDefine(TEXT("PLATFORM_SUPPORTS_SM6_0_WAVE_OPERATIONS"), 1);
-		if (IsUsingSM66(Input, Language))
-		{
-			AdditionalDefines.SetDefine(TEXT("PLATFORM_SUPPORTS_STATIC_SAMPLERS"), 1);
-		}
-		AdditionalDefines.SetDefine(TEXT("PLATFORM_SUPPORTS_DIAGNOSTIC_BUFFER"), 1);
-		AdditionalDefines.SetDefine(TEXT("COMPILER_SUPPORTS_NOINLINE"), 1);
-
-		if (Environment.CompilerFlags.Contains(CFLAG_InlineRayTracing))
-		{
-			AdditionalDefines.SetDefine(TEXT("PLATFORM_SUPPORTS_INLINE_RAY_TRACING"), 1);
-		}
-
-		if (Language == ELanguage::SM6 && Environment.CompilerFlags.Contains(CFLAG_AllowRealTypes))
-		{
-			AdditionalDefines.SetDefine(TEXT("PLATFORM_SUPPORTS_REAL_TYPES"), 1);
-		}
-	}
-	else
-	{
-		AdditionalDefines.SetDefine(TEXT("COMPILER_FXC"), 1);
-	}
-
-	PRAGMA_ENABLE_DEPRECATION_WARNINGS
 	
 	if (Input.bSkipPreprocessedCache)
 	{
@@ -1410,14 +1324,13 @@ struct FD3DShaderParameterParserPlatformConfiguration : public FShaderParameterP
 	}
 };
 
-void CompileD3DShader(const FShaderCompilerInput& Input, const FString& InPreprocessedSource, FShaderCompilerOutput& Output, const FString& WorkingDirectory, ELanguage Language)
+void CompileD3DShader(const FShaderCompilerInput& Input, const FString& InPreprocessedSource, FShaderCompilerOutput& Output, const FString& WorkingDirectory, ED3DShaderModel ShaderModel)
 {
 	TRACE_CPUPROFILER_EVENT_SCOPE(CompileD3DShader);
 
-	const bool bUseDXC = ShouldUseDXC(Input, Language);
-	const TCHAR* ShaderProfile = GetShaderProfileName(Language, Input.Target.GetFrequency(), bUseDXC);
+	const TCHAR* ShaderProfile = GetShaderProfileName(Input, ShaderModel);
 
-	if(!ShaderProfile)
+	if (!ShaderProfile)
 	{
 		Output.Errors.Add(FShaderCompilerError(*FString::Printf(TEXT("Unrecognized shader frequency %s"), GetShaderFrequencyString((EShaderFrequency)Input.Target.Frequency))));
 		return;
@@ -1502,8 +1415,8 @@ void CompileD3DShader(const FShaderCompilerInput& Input, const FString& InPrepro
 		}
 	}
 
-	const bool bSuccess = bUseDXC
-		? CompileAndProcessD3DShaderDXC(Input, PreprocessedSource, EntryPointName, ShaderParameterParser, ShaderProfile, Language, false, Output)
+	const bool bSuccess = DoesShaderModelRequireDXC(ShaderModel)
+		? CompileAndProcessD3DShaderDXC(Input, PreprocessedSource, EntryPointName, ShaderParameterParser, ShaderProfile, ShaderModel, false, Output)
 		: CompileAndProcessD3DShaderFXC(Input, PreprocessedSource, EntryPointName, ShaderParameterParser, ShaderProfile, false, Output);
 
 	if (!bSuccess && !Output.Errors.Num())
diff --git a/Engine/Source/Developer/Windows/ShaderFormatD3D/Private/D3DShaderCompilerDXC.cpp b/Engine/Source/Developer/Windows/ShaderFormatD3D/Private/D3DShaderCompilerDXC.cpp
--- a/Engine/Source/Developer/Windows/ShaderFormatD3D/Private/D3DShaderCompilerDXC.cpp
+++ b/Engine/Source/Developer/Windows/ShaderFormatD3D/Private/D3DShaderCompilerDXC.cpp
@@ -175,7 +175,7 @@ public:
 		const FShaderCompilerInput& Input,
 		const FString& InEntryPoint,
 		const TCHAR* InShaderProfile,
-		ELanguage Language,
+		ED3DShaderModel ShaderModel,
 		const FString& InExports
 	)
 		: ShaderProfile(InShaderProfile)
@@ -191,8 +191,8 @@ public:
 		}
 
 		const bool bEnable16BitTypes =
-			// 16bit types are SM6.2 whereas Language == ELanguage::SM6 is SM6.6, so their support at runtime is guarented.
-			(Language == ELanguage::SM6 && Input.Environment.CompilerFlags.Contains(CFLAG_AllowRealTypes))
+			// 16bit types are SM6.2, so their support at runtime is guaranteed in SM6.6.
+			(ShaderModel >= ED3DShaderModel::SM6_6 && Input.Environment.CompilerFlags.Contains(CFLAG_AllowRealTypes))
 
 			// Enable 16bit_types to reduce DXIL size (compiler bug - will be fixed)
 			|| Input.IsRayTracingShader();
@@ -852,7 +852,7 @@ bool CompileAndProcessD3DShaderDXC(
 	const FString& EntryPointName,
 	const FShaderParameterParser& ShaderParameterParser,
 	const TCHAR* ShaderProfile,
-	ELanguage Language,
+	ED3DShaderModel ShaderModel,
 	bool bProcessingSecondTime,
 	FShaderCompilerOutput& Output)
 {
@@ -893,7 +893,7 @@ bool CompileAndProcessD3DShaderDXC(
 		Input,
 		EntryPointName,
 		ShaderProfile,
-		Language,
+		ShaderModel,
 		RayTracingExports
 	);
 
@@ -1233,7 +1233,7 @@ bool CompileAndProcessD3DShaderDXC(
 				// We only need this to appear when using a DXC shader
 				ShaderCode.AddOptionalData<FShaderCodeFeatures>(CodeFeatures);
 
-				if (Language != ELanguage::SM6)
+				if (ShaderModel >= ED3DShaderModel::SM6_0)
 				{
 					uint8 IsSM6 = 1;
 					ShaderCode.AddOptionalData('6', &IsSM6, 1);
diff --git a/Engine/Source/Developer/Windows/ShaderFormatD3D/Private/ShaderFormatD3D.cpp b/Engine/Source/Developer/Windows/ShaderFormatD3D/Private/ShaderFormatD3D.cpp
--- a/Engine/Source/Developer/Windows/ShaderFormatD3D/Private/ShaderFormatD3D.cpp
+++ b/Engine/Source/Developer/Windows/ShaderFormatD3D/Private/ShaderFormatD3D.cpp
@@ -79,44 +79,95 @@ public:
 		OutFormats.Add(NAME_PCD3D_ES3_1);
 	}
 
-	ELanguage FormatToLanguage(FName Format) const
+	enum class ELanguage
+	{
+		SM5,
+		SM6,
+		ES3_1,
+		Invalid,
+	};
+
+	static ELanguage LanguageFromFormat(FName Format)
 	{
 		if (Format == NAME_PCD3D_SM6)
 		{
 			return ELanguage::SM6;
 		}
-		else if (Format == NAME_PCD3D_SM5)
+		
+		if (Format == NAME_PCD3D_SM5)
 		{
 			return ELanguage::SM5;
 		}
-		else if (Format == NAME_PCD3D_ES3_1)
+		
+		if (Format == NAME_PCD3D_ES3_1)
 		{
 			return ELanguage::ES3_1;
 		}
-		else
+
+		checkf(0, TEXT("Unknown format %s"), *Format.ToString());
+		return ELanguage::Invalid;
+	}
+
+	static bool IsSM66(const FShaderCompilerInput& Input, ELanguage Language)
+	{
+		return Language == ELanguage::SM6 || IsRayTracingShaderFrequency(Input.Target.GetFrequency());
+	}
+
+	// Do we need any SM6.0 features?
+	static bool RequiresSM6Features(const FShaderCompilerInput& Input, ELanguage Language)
+	{
+		return IsSM66(Input, Language)
+			|| Input.Environment.CompilerFlags.Contains(CFLAG_WaveOperations)
+			// TODO: Forcing DXC should not change platform flags, this needs to be moved to IsSM60 once existing uses are accounted for.
+			|| Input.Environment.CompilerFlags.Contains(CFLAG_ForceDXC)
+			;
+	}
+
+	// Do we need to compile with SM6.0 at all? We can compile with 6.0 but not require the language features
+	static bool IsSM60(const FShaderCompilerInput& Input, ELanguage Language)
+	{
+		return RequiresSM6Features(Input, Language)
+			|| Input.Environment.GetIntegerValue(TEXT("PLATFORM_MAX_SAMPLERS")) > 16
+			;
+	}
+
+	static ED3DShaderModel DetermineShaderModel(const FShaderCompilerInput& Input, ELanguage Language)
+	{
+		if (IsSM66(Input, Language))
+		{
+			return ED3DShaderModel::SM6_6;
+		}
+
+		if (IsSM60(Input, Language))
 		{
-			checkf(0, TEXT("Unknown format %s"), *Format.ToString());
-			return ELanguage::Invalid;
+			return ED3DShaderModel::SM6_0;
 		}
+
+		return ED3DShaderModel::SM5_0;
+	}
+
+	static ED3DShaderModel DetermineShaderModel(const FShaderCompilerInput& Input)
+	{
+		return DetermineShaderModel(Input, LanguageFromFormat(Input.ShaderFormat));
 	}
 
-	virtual bool PreprocessShader(const struct FShaderCompilerInput& Input, const struct FShaderCompilerEnvironment& MergedEnvironment, class FShaderPreprocessOutput& PreprocessOutput) const
+	virtual bool PreprocessShader(const FShaderCompilerInput& Input, const FShaderCompilerEnvironment& MergedEnvironment, FShaderPreprocessOutput& PreprocessOutput) const
 	{
-		return PreprocessD3DShader(Input, MergedEnvironment, PreprocessOutput, FormatToLanguage(Input.ShaderFormat));
+		return PreprocessD3DShader(Input, MergedEnvironment, PreprocessOutput);
 	}
 
 	virtual void CompilePreprocessedShader(
-		const struct FShaderCompilerInput& Input, 
-		const class FShaderPreprocessOutput& PreprocessOutput, 
-		struct FShaderCompilerOutput& Output,
+		const FShaderCompilerInput& Input, 
+		const FShaderPreprocessOutput& PreprocessOutput, 
+		FShaderCompilerOutput& Output,
 		const FString& WorkingDirectory) const
 	{
-		CompileD3DShader(Input, PreprocessOutput.GetSource(), Output, WorkingDirectory, FormatToLanguage(Input.ShaderFormat));
+		CompileD3DShader(Input, PreprocessOutput.GetSource(), Output, WorkingDirectory, DetermineShaderModel(Input));
 
 		Output.ShaderDiagnosticDatas = PreprocessOutput.GetDiagnosticDatas();
 	}
 
-	void AddShaderTargetDefines(FShaderCompilerInput& Input, uint32 ShaderTargetMajor, uint32 ShaderTargetMinor) const
+	static void AddShaderTargetDefines(FShaderCompilerInput& Input, uint32 ShaderTargetMajor, uint32 ShaderTargetMinor)
 	{
 		// Inserting our own versions of these defines since we preprocess our shader source before we actually use something that defines them.
 		Input.Environment.SetDefine(TEXT("__SHADER_TARGET_MAJOR"), ShaderTargetMajor);
@@ -125,44 +176,56 @@ public:
 
 	void ModifyShaderCompilerInput(FShaderCompilerInput& Input) const final
 	{
-		const ELanguage Language = FormatToLanguage(Input.ShaderFormat);
+		const ELanguage Language = LanguageFromFormat(Input.ShaderFormat);
+		const ED3DShaderModel ShaderModel = DetermineShaderModel(Input, Language);
 
-		if (IsUsingSM66(Input, Language))
-		{
-			Input.Environment.SetDefine(TEXT("SM6_PROFILE"), 1);
-			Input.Environment.SetDefine(TEXT("COMPILER_DXC"), 1);
-			Input.Environment.SetDefine(TEXT("PLATFORM_SUPPORTS_CONSTANTBUFFER_OBJECT"), 1);
+		// Our compilers only support HLSL
+		Input.Environment.SetDefine(TEXT("COMPILER_HLSL"), true);
 
-			AddShaderTargetDefines(Input, 6, 6);
-		}
-		else if (Language == ELanguage::SM5)
-		{
-			Input.Environment.SetDefine(TEXT("SM5_PROFILE"), 1);
-			const bool bUseDXC =
-				Input.Environment.CompilerFlags.Contains(CFLAG_WaveOperations)
-				|| Input.Environment.CompilerFlags.Contains(CFLAG_ForceDXC);
-			Input.Environment.SetDefine(TEXT("COMPILER_DXC"), bUseDXC);
-			Input.Environment.SetDefine(TEXT("PLATFORM_SUPPORTS_CONSTANTBUFFER_OBJECT"), bUseDXC);
-
-			if (bUseDXC)
-			{
-				AddShaderTargetDefines(Input, 6, 0);
-			}
-			else
-			{
-				AddShaderTargetDefines(Input, 5, 0);
-			}
-		}
-		else if (Language == ELanguage::ES3_1)
+		// Assume min. spec HW supports with DX12/SM5
+		Input.Environment.SetDefine(TEXT("PLATFORM_SUPPORTS_ROV"), true);
+
+		const bool bSM66        = IsSM66(Input, Language);
+		const bool bSM6Features = RequiresSM6Features(Input, Language);
+		const bool bDXC         = DoesShaderModelRequireDXC(ShaderModel);
+
+		// Compiler specific defines
+		Input.Environment.SetDefine(TEXT("COMPILER_FXC"), !bDXC);
+		Input.Environment.SetDefine(TEXT("COMPILER_DXC"),  bDXC);
+
+		// Do we need SM6.0+ features enabled? This is intentionally disconnected from the ED3DShaderModel to allow SM6.0 to be used without new language features.
+		Input.Environment.SetDefine(TEXT("PLATFORM_SUPPORTS_CONSTANTBUFFER_OBJECT"), bSM6Features);
+		Input.Environment.SetDefine(TEXT("PLATFORM_SUPPORTS_SM6_0_WAVE_OPERATIONS"), bSM6Features);
+		Input.Environment.SetDefine(TEXT("PLATFORM_SUPPORTS_DIAGNOSTIC_BUFFER"),     bSM6Features);
+
+		// "profiles" are almost analogous to ERHIFeatureLevel but RT shaders are forcing themselves to SM6
+		Input.Environment.SetDefine(TEXT("SM6_PROFILE"),   bSM66);
+		Input.Environment.SetDefine(TEXT("SM5_PROFILE"),   (Language == ELanguage::SM5));
+		Input.Environment.SetDefine(TEXT("ES3_1_PROFILE"), (Language == ELanguage::ES3_1));
+
+		// Add SM6.6+ specific defines. None of these are intended to be enabled in lower SM's
+		if (bSM66)
 		{
-			Input.Environment.SetDefine(TEXT("ES3_1_PROFILE"), 1);
-			Input.Environment.SetDefine(TEXT("COMPILER_DXC"), 0);
-			Input.Environment.SetDefine(TEXT("__SHADER_TARGET_MAJOR"), 5);
-			Input.Environment.SetDefine(TEXT("__SHADER_TARGET_MINOR"), 0);
+			Input.Environment.SetDefine(TEXT("PLATFORM_SUPPORTS_REAL_TYPES"),         Input.Environment.CompilerFlags.Contains(CFLAG_AllowRealTypes));
+			Input.Environment.SetDefine(TEXT("PLATFORM_SUPPORTS_INLINE_RAY_TRACING"), Input.Environment.CompilerFlags.Contains(CFLAG_InlineRayTracing));
+
+			Input.Environment.SetDefine(TEXT("PLATFORM_SUPPORTS_STATIC_SAMPLERS"),  true);
+			Input.Environment.SetDefine(TEXT("PLATFORM_SUPPORTS_CALLABLE_SHADERS"), true);
+			Input.Environment.SetDefine(TEXT("COMPILER_SUPPORTS_NOINLINE"),         true);
 		}
-		else
+
+		// Add defines for our specific shader model (5.0, 6.0, 6.6)
+		switch (ShaderModel)
 		{
-			checkf(0, TEXT("Unknown format %s"), *Input.ShaderFormat.ToString());
+		case ED3DShaderModel::SM5_0:
+			AddShaderTargetDefines(Input, 5, 0);
+			break;
+		case ED3DShaderModel::SM6_0:
+			AddShaderTargetDefines(Input, 6, 0);
+			break;
+		case ED3DShaderModel::SM6_6:
+			AddShaderTargetDefines(Input, 6, 6);
+			break;
 		}
 	}
 
diff --git a/Engine/Source/Developer/Windows/ShaderFormatD3D/Private/ShaderFormatD3D.h b/Engine/Source/Developer/Windows/ShaderFormatD3D/Private/ShaderFormatD3D.h
--- a/Engine/Source/Developer/Windows/ShaderFormatD3D/Private/ShaderFormatD3D.h
+++ b/Engine/Source/Developer/Windows/ShaderFormatD3D/Private/ShaderFormatD3D.h
@@ -12,36 +12,30 @@
 
 struct FShaderTarget;
 
-enum class ELanguage
+enum class ED3DShaderModel
 {
-	SM5,
-	SM6,
-	ES3_1,
 	Invalid,
+	SM5_0,
+	SM6_0,
+	SM6_6,
 };
 
-inline bool IsUsingSM66(ELanguage Language, EShaderFrequency Frequency)
+inline bool DoesShaderModelRequireDXC(ED3DShaderModel ShaderModel)
 {
-	return Language == ELanguage::SM6 || IsRayTracingShaderFrequency(Frequency);
-}
-
-inline bool IsUsingSM66(const FShaderCompilerInput& Input, ELanguage Language)
-{
-	return IsUsingSM66(Language, Input.Target.GetFrequency());
+	return ShaderModel >= ED3DShaderModel::SM6_0;
 }
 
 bool PreprocessD3DShader(
 	const FShaderCompilerInput& Input,
 	const FShaderCompilerEnvironment& MergedEnvironment,
-	FShaderPreprocessOutput& PreprocessOutput,
-	ELanguage Language);
+	FShaderPreprocessOutput& PreprocessOutput);
 
 void CompileD3DShader(
 	const FShaderCompilerInput& Input,
 	const FString& InPreprocessedSource,
 	FShaderCompilerOutput& Output,
 	const FString& WorkingDirectory,
-	ELanguage Language);
+	ED3DShaderModel ShaderModel);
 
 /**
  * @param bSecondPassAferUnusedInputRemoval whether we're compiling the shader second time, after having removed the unused inputs discovered in the first pass
@@ -61,7 +55,7 @@ bool CompileAndProcessD3DShaderDXC(
 	const FString& InEntryPointName,
 	const FShaderParameterParser& ShaderParameterParser,
 	const TCHAR* ShaderProfile,
-	ELanguage Language,
+	ED3DShaderModel ShaderModel,
 	bool bProcessingSecondTime,
 	FShaderCompilerOutput& Output);
 
