Support exporting shader symbols info for PC D3D_SM6 shaders.

The PDB name is either retrieved from the PDB blob part of DXC's output, or is generated using the shader hash that DXC reports even when no PDB is requested. We don't export PDBs because it would result in duplicate data as we use embedded PDBs.

diff --git a/Engine/Source/Developer/ShaderCompilerCommon/Private/ShaderSymbolExport.cpp b/Engine/Source/Developer/ShaderCompilerCommon/Private/ShaderSymbolExport.cpp
--- a/Engine/Source/Developer/ShaderCompilerCommon/Private/ShaderSymbolExport.cpp
+++ b/Engine/Source/Developer/ShaderCompilerCommon/Private/ShaderSymbolExport.cpp
@@ -273,6 +273,7 @@ void FShaderSymbolExport::NotifyShaderCompilersShutdown()
 			{
 				Writer->Serialize(Output.GetData(), Output.Num());
 				Writer->Close();
+				UE_LOG(LogShaderSymbolExport, Display, TEXT("Wrote %d records into shader symbols info output file \"%s\"."), ShaderInfos.Num(), *InfoFilePath);
 			}
 			else
 			{
diff --git a/Engine/Source/Developer/ShaderCompilerCommon/Public/DXCUtils.inl b/Engine/Source/Developer/ShaderCompilerCommon/Public/DXCUtils.inl
--- a/Engine/Source/Developer/ShaderCompilerCommon/Public/DXCUtils.inl
+++ b/Engine/Source/Developer/ShaderCompilerCommon/Public/DXCUtils.inl
@@ -77,4 +77,32 @@ static void DumpDebugBlobDetail(IDxcBlob* Blob, const TCHAR* BlobName, TRefCount
 	FPlatformMisc::LowLevelOutputDebugString(TEXT("\n\n"));
 }
 
-
+HRESULT RetrieveDebugNameAndBlob(TRefCountPtr<IDxcResult>& CompileResult, FString& OutDebugBlobName, IDxcBlob** OutDebugBlob)
+{
+	HRESULT Result = S_OK;
+	if (OutDebugBlob && CompileResult->HasOutput(DXC_OUT_PDB))
+	{
+		// If the PDB blob exists, we get its name at the same time.
+		TRefCountPtr<IDxcBlobUtf16> DebugBlobName;
+		Result = CompileResult->GetOutput(DXC_OUT_PDB, __uuidof(IDxcBlob), (void**)OutDebugBlob, DebugBlobName.GetInitReference());
+		if (SUCCEEDED(Result))
+		{
+			OutDebugBlobName = DebugBlobName->GetStringPointer();
+		}
+	}
+	else if (CompileResult->HasOutput(DXC_OUT_SHADER_HASH))
+	{
+		// If the PDB blob is not available, or if we don't need to return it, we can reconstruct the name 
+		// the PDB would have used from the Shader Hash, using only 16 bytes if there is more.
+		TRefCountPtr<IDxcBlob> HashBlob;
+		TRefCountPtr<IDxcBlobUtf16> OutputNameUnused;
+		Result = CompileResult->GetOutput(DXC_OUT_SHADER_HASH, __uuidof(IDxcBlob), (void**)HashBlob.GetInitReference(), OutputNameUnused.GetInitReference());
+		if (SUCCEEDED(Result))
+		{
+			size_t BlobSize = HashBlob->GetBufferSize();
+			uint32 Offset = BlobSize <= 16 ? 0 : BlobSize - 16;
+			OutDebugBlobName = FString(BytesToHex(reinterpret_cast<uint8*>(HashBlob->GetBufferPointer()) + Offset, HashBlob->GetBufferSize() - Offset)) + TEXT(".pdb");
+		}
+	}
+	return Result;
+}
diff --git a/Engine/Source/Developer/Windows/ShaderFormatD3D/Private/D3DShaderCompilerDXC.cpp b/Engine/Source/Developer/Windows/ShaderFormatD3D/Private/D3DShaderCompilerDXC.cpp
--- a/Engine/Source/Developer/Windows/ShaderFormatD3D/Private/D3DShaderCompilerDXC.cpp
+++ b/Engine/Source/Developer/Windows/ShaderFormatD3D/Private/D3DShaderCompilerDXC.cpp
@@ -45,6 +45,7 @@ THIRD_PARTY_INCLUDES_START
 	#include "ShaderConductor/ShaderConductor.hpp"
 THIRD_PARTY_INCLUDES_END
 
+#include "DXCUtils.inl"
 #include "D3DShaderCompiler.inl"
 
 FORCENOINLINE static void DXCFilterShaderCompileWarnings(const FString& CompileWarnings, TArray<FString>& FilteredWarnings)
@@ -669,7 +670,7 @@ static bool RemoveContainerParts(const TConstArrayView<uint32> PartCodes, dxc::D
 }
 
 static HRESULT D3DCompileToDxil(const char* SourceText, const FDxcArguments& Arguments,
-	TRefCountPtr<IDxcBlob>& OutDxilBlob, TRefCountPtr<IDxcBlob>& OutReflectionBlob, TRefCountPtr<IDxcBlobEncoding>& OutErrorBlob)
+	TRefCountPtr<IDxcBlob>& OutDxilBlob, TRefCountPtr<IDxcBlob>& OutReflectionBlob, TRefCountPtr<IDxcBlobEncoding>& OutErrorBlob, TRefCountPtr<IDxcBlob>& OutPdbBlob, FString& OutPdbName)
 {
 	dxc::DxcDllSupport& DxcDllHelper = GetDxcDllHelper();
 
@@ -729,7 +730,8 @@ static HRESULT D3DCompileToDxil(const char* SourceText, const FDxcArguments& Arg
 		checkf(CompileResult->HasOutput(DXC_OUT_REFLECTION), TEXT("No reflection found!"));
 		VERIFYHRESULT(CompileResult->GetOutput(DXC_OUT_REFLECTION, IID_PPV_ARGS(OutReflectionBlob.GetInitReference()), ReflectionNameBlob.GetInitReference()));
 
-		const bool bHasOutputPDB = CompileResult->HasOutput(DXC_OUT_PDB);
+		RetrieveDebugNameAndBlob(CompileResult, OutPdbName, OutPdbBlob.GetInitReference());
+		const bool bHasOutputPDB = OutPdbBlob.IsValid() && !OutPdbName.IsEmpty();
 		const bool bRemovePDB = bHasOutputPDB && !Arguments.ShouldKeepEmbeddedPDB();
 
 		TArray<uint32, TInlineAllocator<4>> PartsToRemove;
@@ -751,17 +753,11 @@ static HRESULT D3DCompileToDxil(const char* SourceText, const FDxcArguments& Arg
 			FString DxilFile = Arguments.GetDumpDisassemblyFilename().LeftChop(7) + TEXT("_refl.dxil");
 			SaveDxcBlobToFile(OutDxilBlob, DxilFile);
 
+			// Dump the PDB.
 			if (bHasOutputPDB)
 			{
-				TRefCountPtr<IDxcBlob> PdbBlob;
-				TRefCountPtr<IDxcBlobUtf16> PdbNameBlob;
-				VERIFYHRESULT(CompileResult->GetOutput(DXC_OUT_PDB, IID_PPV_ARGS(PdbBlob.GetInitReference()), PdbNameBlob.GetInitReference()));
-
-				const FString PdbName = PdbNameBlob->GetStringPointer();
-
-				// Dump pdb (.d3dasm -> .pdb)
-				const FString PdbFile = Arguments.GetDumpDebugInfoPath() / PdbName;
-				SaveDxcBlobToFile(PdbBlob, PdbFile);
+				const FString PdbFile = Arguments.GetDumpDebugInfoPath() / OutPdbName;
+				SaveDxcBlobToFile(OutPdbBlob, PdbFile);
 			}
 		}
 
@@ -906,8 +902,26 @@ bool CompileAndProcessD3DShaderDXC(
 	TRefCountPtr<IDxcBlob> ShaderBlob;
 	TRefCountPtr<IDxcBlob> ReflectionBlob;
 	TRefCountPtr<IDxcBlobEncoding> DxcErrorBlob;
+	TRefCountPtr<IDxcBlob> PdbBlob;
+	FString PdbName;
+
+	const HRESULT D3DCompileToDxilResult = D3DCompileToDxil(AnsiSourceFile.Get(), Args, ShaderBlob, ReflectionBlob, DxcErrorBlob, PdbBlob, PdbName);
 
-	const HRESULT D3DCompileToDxilResult = D3DCompileToDxil(AnsiSourceFile.Get(), Args, ShaderBlob, ReflectionBlob, DxcErrorBlob);
+	// Populate the platform-specific debug data with the PDB name, if available.
+	bool bWriteDebugData = Input.Environment.CompilerFlags.Contains(CFLAG_GenerateSymbolsInfo);
+	if (bWriteDebugData && !PdbName.IsEmpty())
+	{
+		FD3DSM6ShaderDebugData DebugData;
+		DebugData.Name = PdbName;
+		DebugData.DebugInfo = Input.GenerateDebugInfo();
+
+		// We don't export the PDB contents here because it would result in duplicate data,
+		// as we use embedded PDBs. Once we are able to use external PDBs, the PDB contents
+		// can be exported too.
+		
+		FMemoryWriter Ar(Output.PlatformDebugData);
+		Ar << DebugData;
+	}
 
 	TArray<FString> FilteredErrors;
 	if (DxcErrorBlob && DxcErrorBlob->GetBufferSize())
diff --git a/Engine/Source/Developer/Windows/ShaderFormatD3D/Private/ShaderFormatD3D.cpp b/Engine/Source/Developer/Windows/ShaderFormatD3D/Private/ShaderFormatD3D.cpp
--- a/Engine/Source/Developer/Windows/ShaderFormatD3D/Private/ShaderFormatD3D.cpp
+++ b/Engine/Source/Developer/Windows/ShaderFormatD3D/Private/ShaderFormatD3D.cpp
@@ -3,6 +3,7 @@
 #include "ShaderFormatD3D.h"
 #include "ShaderCompilerCommon.h"
 #include "ShaderPreprocessTypes.h"
+#include "ShaderSymbolExport.h"
 #include "Modules/ModuleInterface.h"
 #include "Modules/ModuleManager.h"
 #include "Interfaces/IShaderFormat.h"
@@ -14,14 +15,19 @@ static FName NAME_PCD3D_SM5(TEXT("PCD3D_SM5"));
 static FName NAME_PCD3D_ES3_1(TEXT("PCD3D_ES31"));
 
 static const FGuid UE_SHADER_PCD3D_SHARED_VER = FGuid("dd4e6e76-4b48-4097-9ece-0f21118b7177");
-static const FGuid UE_SHADER_PCD3D_SM6_VER    = FGuid("96a53626-2b11-4663-8b85-01e30625b76e");
-static const FGuid UE_SHADER_PCD3D_SM5_VER    = FGuid("07cd0fa4-6d18-43e0-915d-e6c125effafd");
+static const FGuid UE_SHADER_PCD3D_SM6_VER    = FGuid("2283a0d0-1e70-4e59-af31-fd8c5260dba3");
+static const FGuid UE_SHADER_PCD3D_SM5_VER    = FGuid("3ca43c5c-c999-4508-8154-95de2c4be6e4");
 static const FGuid UE_SHADER_PCD3D_ES3_1_VER  = FGuid("75466d2b-e169-40d8-bac5-1e2f9d43e0bb");
 
 class FShaderFormatD3D : public UE::ShaderCompilerCommon::FBaseShaderFormat 
 {
 	uint32 DxcVersionHash = 0;
 
+#if WITH_ENGINE
+	// TODO: Support for D3D_SM5 and D3D_ES31
+	mutable FShaderSymbolExport ShaderSymbolExportSM6{ NAME_PCD3D_SM6 };
+#endif
+
 public:
 
 	FShaderFormatD3D(uint32 InDxcVersionHash)
@@ -146,10 +152,25 @@ public:
 		return ED3DShaderModel::SM5_0;
 	}
 
+#if WITH_ENGINE
+	virtual void NotifyShaderCompiled(const TConstArrayView<uint8>& PlatformDebugData, FName Format) const override
+	{
+		if (Format == NAME_PCD3D_SM6)
+		{
+			ShaderSymbolExportSM6.NotifyShaderCompiled<FD3DSM6ShaderDebugData>(PlatformDebugData);
+		}
+	}
+
+	virtual void NotifyShaderCompilersShutdown(FName Format) const override
 	static ED3DShaderModel DetermineShaderModel(const FShaderCompilerInput& Input)
 	{
+		if (Format == NAME_PCD3D_SM6)
+		{
+			ShaderSymbolExportSM6.NotifyShaderCompilersShutdown();
+		}
 		return DetermineShaderModel(Input, LanguageFromFormat(Input.ShaderFormat));
 	}
+#endif
 
 	virtual bool PreprocessShader(const FShaderCompilerInput& Input, const FShaderCompilerEnvironment& MergedEnvironment, FShaderPreprocessOutput& PreprocessOutput) const
 	{
diff --git a/Engine/Source/Developer/Windows/ShaderFormatD3D/Private/ShaderFormatD3D.h b/Engine/Source/Developer/Windows/ShaderFormatD3D/Private/ShaderFormatD3D.h
--- a/Engine/Source/Developer/Windows/ShaderFormatD3D/Private/ShaderFormatD3D.h
+++ b/Engine/Source/Developer/Windows/ShaderFormatD3D/Private/ShaderFormatD3D.h
@@ -60,3 +60,38 @@ bool CompileAndProcessD3DShaderDXC(
 	FShaderCompilerOutput& Output);
 
 bool ValidateResourceCounts(uint32 NumSRVs, uint32 NumSamplers, uint32 NumUAVs, uint32 NumCBs, TArray<FString>& OutFilteredErrors);
+
+struct FD3DSM6ShaderDebugData
+{
+	FString Name;
+	FString DebugInfo;
+	TArray<uint8> Contents;
+
+	inline friend FArchive& operator<<(FArchive& Ar, FD3DSM6ShaderDebugData& DebugData)
+	{
+		Ar << DebugData.Name;
+		Ar << DebugData.DebugInfo;
+		Ar << DebugData.Contents;
+		return Ar;
+	}
+
+	inline TConstArrayView<uint8> GetContents() const
+	{
+		return TConstArrayView<uint8>(Contents);
+	}
+
+	inline FString GetFilename() const
+	{
+		return Name;
+	}
+
+	inline FString GetDebugInfo() const
+	{
+		return DebugInfo;
+	}
+
+	TConstArrayView<FD3DSM6ShaderDebugData> GetAllSymbolData() const
+	{
+		return TConstArrayView<FD3DSM6ShaderDebugData>(this, 1);
+	}
+};
\ No newline at end of file
