Fix merge error

diff --git a/Engine/Source/Developer/Windows/ShaderFormatD3D/Private/ShaderFormatD3D.cpp b/Engine/Source/Developer/Windows/ShaderFormatD3D/Private/ShaderFormatD3D.cpp
--- a/Engine/Source/Developer/Windows/ShaderFormatD3D/Private/ShaderFormatD3D.cpp
+++ b/Engine/Source/Developer/Windows/ShaderFormatD3D/Private/ShaderFormatD3D.cpp
@@ -152,6 +152,11 @@ public:
 		return ED3DShaderModel::SM5_0;
 	}
 
+	static ED3DShaderModel DetermineShaderModel(const FShaderCompilerInput& Input)
+	{
+		return DetermineShaderModel(Input, LanguageFromFormat(Input.ShaderFormat));
+	}
+
 #if WITH_ENGINE
 	virtual void NotifyShaderCompiled(const TConstArrayView<uint8>& PlatformDebugData, FName Format) const override
 	{
@@ -160,15 +165,12 @@ public:
 			ShaderSymbolExportSM6.NotifyShaderCompiled<FD3DSM6ShaderDebugData>(PlatformDebugData);
 		}
 	}
-
 	virtual void NotifyShaderCompilersShutdown(FName Format) const override
-	static ED3DShaderModel DetermineShaderModel(const FShaderCompilerInput& Input)
 	{
 		if (Format == NAME_PCD3D_SM6)
 		{
 			ShaderSymbolExportSM6.NotifyShaderCompilersShutdown();
 		}
-		return DetermineShaderModel(Input, LanguageFromFormat(Input.ShaderFormat));
 	}
 #endif
 
