Bump PC SM6 and SM5 GUIDs.

diff --git a/Engine/Source/Developer/Windows/ShaderFormatD3D/Private/ShaderFormatD3D.cpp b/Engine/Source/Developer/Windows/ShaderFormatD3D/Private/ShaderFormatD3D.cpp
--- a/Engine/Source/Developer/Windows/ShaderFormatD3D/Private/ShaderFormatD3D.cpp
+++ b/Engine/Source/Developer/Windows/ShaderFormatD3D/Private/ShaderFormatD3D.cpp
@@ -15,8 +15,8 @@ static FName NAME_PCD3D_SM5(TEXT("PCD3D_SM5"));
 static FName NAME_PCD3D_ES3_1(TEXT("PCD3D_ES31"));
 
 static const FGuid UE_SHADER_PCD3D_SHARED_VER = FGuid("dd4e6e76-4b48-4097-9ece-0f21118b7177");
-static const FGuid UE_SHADER_PCD3D_SM6_VER    = FGuid("2283a0d0-1e70-4e59-af31-fd8c5260dba3");
-static const FGuid UE_SHADER_PCD3D_SM5_VER    = FGuid("3ca43c5c-c999-4508-8154-95de2c4be6e4");
+static const FGuid UE_SHADER_PCD3D_SM6_VER    = FGuid("78f5751f-107f-42b3-aa9d-7b871ebac5f6");
+static const FGuid UE_SHADER_PCD3D_SM5_VER    = FGuid("c2f60605-b462-4449-9f76-94bd2b1cd78e");
 static const FGuid UE_SHADER_PCD3D_ES3_1_VER  = FGuid("75466d2b-e169-40d8-bac5-1e2f9d43e0bb");
 
 class FShaderFormatD3D : public UE::ShaderCompilerCommon::FBaseShaderFormat 
