Remove redundant error logging from DXC.

- Only log a system error message if we didn't get any error messages back from the DXC error blob.

#rnx

diff --git a/Engine/Source/Developer/Windows/ShaderFormatD3D/Private/D3DShaderCompilerDXC.cpp b/Engine/Source/Developer/Windows/ShaderFormatD3D/Private/D3DShaderCompilerDXC.cpp
--- a/Engine/Source/Developer/Windows/ShaderFormatD3D/Private/D3DShaderCompilerDXC.cpp
+++ b/Engine/Source/Developer/Windows/ShaderFormatD3D/Private/D3DShaderCompilerDXC.cpp
@@ -1278,13 +1278,17 @@ bool CompileAndProcessD3DShaderDXC(
 	}
 	else
 	{
-		TCHAR ErrorMsg[1024];
-		FPlatformMisc::GetSystemErrorMessage(ErrorMsg, UE_ARRAY_COUNT(ErrorMsg), (int)D3DCompileToDxilResult);
-		const bool bKnownError = ErrorMsg[0] != TEXT('\0');
+		// If we failed and didn't get any error messages back from the compile call try and get a system error message.
+		if (FilteredErrors.Num() == 0)
+		{
+			TCHAR ErrorMsg[1024];
+			FPlatformMisc::GetSystemErrorMessage(ErrorMsg, UE_ARRAY_COUNT(ErrorMsg), (int)D3DCompileToDxilResult);
+			const bool bKnownError = ErrorMsg[0] != TEXT('\0');
 
-		FString ErrorString = FString::Printf(TEXT("D3DCompileToDxil failed. Error code: %s (0x%08X)."), bKnownError ? ErrorMsg : TEXT("Unknown error"), (int)D3DCompileToDxilResult);
+			FString ErrorString = FString::Printf(TEXT("D3DCompileToDxil failed. Error code: %s (0x%08X)."), bKnownError ? ErrorMsg : TEXT("Unknown error"), (int)D3DCompileToDxilResult);
 
-		FilteredErrors.Add(ErrorString);
+			FilteredErrors.Add(ErrorString);
+		}
 	}
 
 	// Move intermediate filtered errors into compiler context for unification.
