Removing invalid check on AutoBindingSpace. The ~0 check was leftover from a long-gone default argument.

diff --git a/Engine/Source/Developer/Windows/ShaderFormatD3D/Private/D3DShaderCompilerDXC.cpp b/Engine/Source/Developer/Windows/ShaderFormatD3D/Private/D3DShaderCompilerDXC.cpp
--- a/Engine/Source/Developer/Windows/ShaderFormatD3D/Private/D3DShaderCompilerDXC.cpp
+++ b/Engine/Source/Developer/Windows/ShaderFormatD3D/Private/D3DShaderCompilerDXC.cpp
@@ -242,7 +242,6 @@ public:
 		}
 
 		const uint32 AutoBindingSpace = GetAutoBindingSpace(Input.Target);
-		if (AutoBindingSpace != ~0u)
 		{
 			ExtraArguments.Add(TEXT("-auto-binding-space"));
 			ExtraArguments.Add(FString::Printf(TEXT("%d"), AutoBindingSpace));
