D3D Bindless
- Forcing EShaderResourceUsageFlags bindless flags when bindless is enabled to avoid ping ponging issues at runtime.

diff --git a/Engine/Source/Developer/Windows/ShaderFormatD3D/Private/D3DShaderCompilerDXC.cpp b/Engine/Source/Developer/Windows/ShaderFormatD3D/Private/D3DShaderCompilerDXC.cpp
--- a/Engine/Source/Developer/Windows/ShaderFormatD3D/Private/D3DShaderCompilerDXC.cpp
+++ b/Engine/Source/Developer/Windows/ShaderFormatD3D/Private/D3DShaderCompilerDXC.cpp
@@ -937,7 +937,6 @@ bool CompileAndProcessD3DShaderDXC(
 
 		bool bGlobalUniformBufferUsed = false;
 		bool bDiagnosticBufferUsed = false;
-		bool bRootConstants = false;
 		uint32 NumInstructions = 0;
 		uint32 NumSamplers = 0;
 		uint32 NumSRVs = 0;
@@ -958,11 +957,6 @@ bool CompileAndProcessD3DShaderDXC(
 		ReflBuffer.Ptr = ReflectionBlob->GetBufferPointer();
 		ReflBuffer.Size = ReflectionBlob->GetBufferSize();
 
-		if (Input.Environment.CompilerFlags.Contains(CFLAG_RootConstants))
-		{
-			bRootConstants = true;
-		}
-
 		bool bHasNoDerivativeOps = false;
 
 		if (Input.Target.GetFrequency() == SF_Compute && Input.Environment.CompilerFlags.Contains(CFLAG_CheckForDerivativeOps))
@@ -1147,17 +1141,17 @@ bool CompileAndProcessD3DShaderDXC(
 				PackedResourceCounts.UsageFlags |= EShaderResourceUsageFlags::GlobalUniformBuffer;
 			}
 
-			if (bRootConstants)
+			if (Input.Environment.CompilerFlags.Contains(CFLAG_RootConstants))
 			{
 				PackedResourceCounts.UsageFlags |= EShaderResourceUsageFlags::RootConstants;
 			}
 
-			if (ShaderRequiresFlags & D3D_SHADER_REQUIRES_RESOURCE_DESCRIPTOR_HEAP_INDEXING)
+			if (Input.Environment.CompilerFlags.Contains(CFLAG_BindlessResources))
 			{
 				PackedResourceCounts.UsageFlags |= EShaderResourceUsageFlags::BindlessResources;
 			}
 
-			if (ShaderRequiresFlags & D3D_SHADER_REQUIRES_SAMPLER_DESCRIPTOR_HEAP_INDEXING)
+			if (Input.Environment.CompilerFlags.Contains(CFLAG_BindlessSamplers))
 			{
 				PackedResourceCounts.UsageFlags |= EShaderResourceUsageFlags::BindlessSamplers;
 			}
diff --git a/Engine/Source/Developer/Windows/ShaderFormatD3D/Private/ShaderFormatD3D.cpp b/Engine/Source/Developer/Windows/ShaderFormatD3D/Private/ShaderFormatD3D.cpp
--- a/Engine/Source/Developer/Windows/ShaderFormatD3D/Private/ShaderFormatD3D.cpp
+++ b/Engine/Source/Developer/Windows/ShaderFormatD3D/Private/ShaderFormatD3D.cpp
@@ -15,7 +15,7 @@ static FName NAME_PCD3D_SM5(TEXT("PCD3D_SM5"));
 static FName NAME_PCD3D_ES3_1(TEXT("PCD3D_ES31"));
 
 static const FGuid UE_SHADER_PCD3D_SHARED_VER = FGuid("dd4e6e76-4b48-4097-9ece-0f21118b7177");
-static const FGuid UE_SHADER_PCD3D_SM6_VER    = FGuid("35A5B0FD-362E-45D5-8C27-9D94DAF94B58");
+static const FGuid UE_SHADER_PCD3D_SM6_VER    = FGuid("EEDED915-D5D7-4F45-B682-2657DCA0C297");
 static const FGuid UE_SHADER_PCD3D_SM5_VER    = FGuid("c2f60605-b462-4449-9f76-94bd2b1cd78e");
 static const FGuid UE_SHADER_PCD3D_ES3_1_VER  = FGuid("75466d2b-e169-40d8-bac5-1e2f9d43e0bb");
 
