Shader Compiler - Propagate filename to dxc command line instead of using  the same dummy name (zzz) for all shaders.

This is helpful when using a profiler such as nSight

diff --git a/Engine/Source/Developer/Windows/ShaderFormatD3D/Private/D3DShaderCompilerDXC.cpp b/Engine/Source/Developer/Windows/ShaderFormatD3D/Private/D3DShaderCompilerDXC.cpp
--- a/Engine/Source/Developer/Windows/ShaderFormatD3D/Private/D3DShaderCompilerDXC.cpp
+++ b/Engine/Source/Developer/Windows/ShaderFormatD3D/Private/D3DShaderCompilerDXC.cpp
@@ -348,12 +348,6 @@ public:
 
 		Out.Add(TEXT("-T"));
 		Out.Add(*ShaderProfile);
-
-		Out.Add(TEXT(" -Fc "));
-		Out.Add(TEXT("zzz.d3dasm"));	// Dummy
-
-		Out.Add(TEXT(" -Fo "));
-		Out.Add(TEXT("zzz.dxil"));	// Dummy
 	}
 
 	const FString& GetBatchBaseFilename() const
@@ -554,6 +548,14 @@ static HRESULT DXCCompileWrapper(
 	TArray<const WCHAR*> CompilerArgs;
 	Arguments.GetCompilerArgs(CompilerArgs);
 
+	// Give a unique name to the d3dasm and dxil outputs (Must have same scope as CompilerArgs so the temporary strings remain valid)
+	FString AsmFilename  = Arguments.GetBatchBaseFilename() + TEXT(".d3dasm");
+	FString DXILFilename = Arguments.GetBatchBaseFilename() + TEXT(".dxil");
+	CompilerArgs.Add(TEXT(" -Fc "));
+	CompilerArgs.Add(*AsmFilename);
+	CompilerArgs.Add(TEXT(" -Fo "));
+	CompilerArgs.Add(*DXILFilename);
+
 	HRESULT Result = InnerDXCCompileWrapper(Compiler, TextBlob,
 		CompilerArgs.GetData(), CompilerArgs.Num(), bExceptionError, OutCompileResult);
 
