Allow enabling/disabling LiveCoding via command line switch -LiveCoding=true|false
Useful in cases where there are multiple live coding targets, and you are only interested in having live coding enabled for one.

diff --git a/Engine/Source/Developer/Windows/LiveCoding/Private/LiveCodingModule.cpp b/Engine/Source/Developer/Windows/LiveCoding/Private/LiveCodingModule.cpp
--- a/Engine/Source/Developer/Windows/LiveCoding/Private/LiveCodingModule.cpp
+++ b/Engine/Source/Developer/Windows/LiveCoding/Private/LiveCodingModule.cpp
@@ -465,12 +465,25 @@ void FLiveCodingModule::StartupModule()
 
 	LppStartup();
 
+	bool bAllowAutoStart = true;
+	bool bForceStart = false;
+
+	if (bool bCommandLineEnable; FParse::Bool(FCommandLine::Get(), TEXT("-LiveCoding="), bCommandLineEnable))
+	{
+		bAllowAutoStart &= bCommandLineEnable;
+		bForceStart = bCommandLineEnable;
+	}
+	else if (FParse::Param(FCommandLine::Get(), TEXT("LiveCoding")))
+	{
+		bForceStart = true;
+	}
+
 	bSettingsEnabledLastTick = Settings->bEnabled;
-	if (Settings->bEnabled && Settings->Startup != ELiveCodingStartupMode::Manual && !FApp::IsUnattended())
+	if (Settings->bEnabled && Settings->Startup != ELiveCodingStartupMode::Manual && !FApp::IsUnattended() && bAllowAutoStart)
 	{
 		StartLiveCodingAsync(Settings->Startup);
-	} 
-	else if (FParse::Param(FCommandLine::Get(), TEXT("LiveCoding")))
+	}
+	else if (bForceStart)
 	{
 		StartLiveCodingAsync(ELiveCodingStartupMode::Manual);
 	}
