Add synchronous version of live coding compile command for use with cmdlink tool/remote console commands in general

diff --git a/Engine/Source/Developer/Windows/LiveCoding/Private/LiveCodingModule.cpp b/Engine/Source/Developer/Windows/LiveCoding/Private/LiveCodingModule.cpp
--- a/Engine/Source/Developer/Windows/LiveCoding/Private/LiveCodingModule.cpp
+++ b/Engine/Source/Developer/Windows/LiveCoding/Private/LiveCodingModule.cpp
@@ -422,6 +422,13 @@ void FLiveCodingModule::StartupModule()
 		FConsoleCommandDelegate::CreateLambda([this] { Compile(ELiveCodingCompileFlags::None, nullptr); }),
 		ECVF_Cheat
 	);
+	
+	CompileSyncCommand = ConsoleManager.RegisterConsoleCommand(
+		TEXT("LiveCoding.CompileSync"),
+		TEXT("Initiates a live coding compile and waits for completion"),
+		FConsoleCommandDelegate::CreateLambda([this] { Compile(ELiveCodingCompileFlags::WaitForCompletion, nullptr); }),
+		ECVF_Cheat
+	);
 
 #if IS_MONOLITHIC
 	FString DefaultEngineDir = GLiveCodingEngineDir;
@@ -501,6 +508,7 @@ void FLiveCodingModule::ShutdownModule()
 	ConsoleManager.UnregisterConsoleObject(SourceProjectVariable);
 	ConsoleManager.UnregisterConsoleObject(ConsolePathVariable);
 	ConsoleManager.UnregisterConsoleObject(CompileCommand);
+	ConsoleManager.UnregisterConsoleObject(CompileSyncCommand);
 	ConsoleManager.UnregisterConsoleObject(EnableCommand);
 
 	// Unregister from the dll notifications
diff --git a/Engine/Source/Developer/Windows/LiveCoding/Private/LiveCodingModule.h b/Engine/Source/Developer/Windows/LiveCoding/Private/LiveCodingModule.h
--- a/Engine/Source/Developer/Windows/LiveCoding/Private/LiveCodingModule.h
+++ b/Engine/Source/Developer/Windows/LiveCoding/Private/LiveCodingModule.h
@@ -128,6 +128,7 @@ private:
 
 	IConsoleCommand* EnableCommand;
 	IConsoleCommand* CompileCommand;
+	IConsoleCommand* CompileSyncCommand;
 	IConsoleVariable* ConsolePathVariable;
 	IConsoleVariable* SourceProjectVariable;
 	FDelegateHandle EndFrameDelegateHandle;
