Shader statistics
- Adding utility function for easier additions to FShaderCompilerOutput::ShaderStatistics
- Simplified existing uses of FShaderCompilerOutput::ShaderStatistics
- Reworked D3D compilation to allow DXC vs FXC to specify their different shader limits.
- Adding Shader Statistics for D3D SamplerState and Resource usages.
- Adding Shader Statistics for D3D SM6.6 bindless resource usages.

diff --git a/Engine/Source/Developer/Windows/ShaderFormatD3D/Private/D3DShaderCompiler.cpp b/Engine/Source/Developer/Windows/ShaderFormatD3D/Private/D3DShaderCompiler.cpp
--- a/Engine/Source/Developer/Windows/ShaderFormatD3D/Private/D3DShaderCompiler.cpp
+++ b/Engine/Source/Developer/Windows/ShaderFormatD3D/Private/D3DShaderCompiler.cpp
@@ -293,40 +293,6 @@ static FString D3D11CreateShaderCompileCommandLine(
 }
 
 
-// Validate that we are not going over to maximum amount of resource bindings support by the default root signature on DX12
-// Currently limited for hard-coded root signature setup (see: FD3D12Adapter::StaticGraphicsRootSignature)
-// In theory this limitation is only required for DX12, but we don't want a shader to compile on DX11 while not working on DX12.
-// (DX11 has an API limit on 128 SRVs, 16 Samplers, 8 UAVs and 14 CBs but if you go over these values then the shader won't compile)
-bool ValidateResourceCounts(uint32 NumSRVs, uint32 NumSamplers, uint32 NumUAVs, uint32 NumCBs, TArray<FString>& OutFilteredErrors)
-{
-	if (NumSRVs > MAX_SRVS || NumSamplers > MAX_SAMPLERS || NumUAVs > MAX_UAVS || NumCBs > MAX_CBS)
-	{
-		if (NumSRVs > MAX_SRVS)
-		{
-			OutFilteredErrors.Add(FString::Printf(TEXT("Shader is using too many SRVs: %d (only %d supported)"), NumSRVs, MAX_SRVS));
-		}
-
-		if (NumSamplers > MAX_SAMPLERS)
-		{
-			OutFilteredErrors.Add(FString::Printf(TEXT("Shader is using too many Samplers: %d (only %d supported)"), NumSamplers, MAX_SAMPLERS));
-		}
-
-		if (NumUAVs > MAX_UAVS)
-		{
-			OutFilteredErrors.Add(FString::Printf(TEXT("Shader is using too many UAVs: %d (only %d supported)"), NumUAVs, MAX_UAVS));
-		}
-
-		if (NumCBs > MAX_CBS)
-		{
-			OutFilteredErrors.Add(FString::Printf(TEXT("Shader is using too many Constant Buffers: %d (only %d supported)"), NumCBs, MAX_CBS));
-		}
-
-		return false;
-	}
-
-	return true;
-}
-
 /** Creates a batch file string to call the AMD shader analyzer. */
 static FString CreateAMDCodeXLCommandLine(
 	const FString& ShaderPath, 
@@ -923,22 +889,15 @@ static bool CompileAndProcessD3DShaderFXCExt(
 
 	// Gather reflection information
 	TArray<FString> ShaderInputs;
-	TArray<FShaderCodeVendorExtension> VendorExtensions;
 
 	if (SUCCEEDED(Result))
 	{
-		bool bGlobalUniformBufferUsed = false;
-		bool bDiagnosticBufferUsed = false;
-		uint32 NumInstructions = 0;
-		uint32 NumSamplers = 0;
-		uint32 NumSRVs = 0;
-		uint32 NumCBs = 0;
-		uint32 NumUAVs = 0;
-
-		TArray<FString> UniformBufferNames;
-
-		TBitArray<> UsedUniformBufferSlots;
-		UsedUniformBufferSlots.Init(false, 32);
+		FD3DShaderCompileData CompileData;
+		// D3D11RHI uses the D3D11_ defines, but we want to enforce the engine limits as well.
+		CompileData.MaxSamplers = FMath::Min(D3D11_COMMONSHADER_SAMPLER_SLOT_COUNT,             MAX_SAMPLERS);
+		CompileData.MaxSRVs     = FMath::Min(D3D11_COMMONSHADER_INPUT_RESOURCE_SLOT_COUNT,      MAX_SRVS);
+		CompileData.MaxCBs      = FMath::Min(D3D11_COMMONSHADER_CONSTANT_BUFFER_API_SLOT_COUNT, MAX_CBS);
+		CompileData.MaxUAVs     = FMath::Min(D3D11_PS_CS_UAV_REGISTER_COUNT,                    MAX_UAVS);
 
 		if (D3DReflectFunc)
 		{
@@ -971,15 +930,15 @@ static bool CompileAndProcessD3DShaderFXCExt(
 						{
 							FString SemanticName = ANSI_TO_TCHAR(ParamDesc.SemanticName);
 
-							ShaderInputs.AddUnique(SemanticName);
+							CompileData.ShaderInputs.AddUnique(SemanticName);
 
 							// Add the number (for the case of TEXCOORD)
 							FString SemanticIndexName = FString::Printf(TEXT("%s%d"), *SemanticName, ParamDesc.SemanticIndex);
-							ShaderInputs.AddUnique(SemanticIndexName);
+							CompileData.ShaderInputs.AddUnique(SemanticIndexName);
 
 							// Add _centroid
-							ShaderInputs.AddUnique(SemanticName + TEXT("_centroid"));
-							ShaderInputs.AddUnique(SemanticIndexName + TEXT("_centroid"));
+							CompileData.ShaderInputs.AddUnique(SemanticName + TEXT("_centroid"));
+							CompileData.ShaderInputs.AddUnique(SemanticIndexName + TEXT("_centroid"));
 						}
 						else
 						{
@@ -991,7 +950,7 @@ static bool CompileAndProcessD3DShaderFXCExt(
 						//if (ParamDesc.ReadWriteMask != 0)
 						{
 							// Keep system values
-							ShaderInputs.AddUnique(FString(ANSI_TO_TCHAR(ParamDesc.SemanticName)));
+							CompileData.ShaderInputs.AddUnique(FString(ANSI_TO_TCHAR(ParamDesc.SemanticName)));
 						}
 					}
 				}
@@ -1012,7 +971,7 @@ static bool CompileAndProcessD3DShaderFXCExt(
 						TArray<FString> RemoveErrors;
 						FString ModifiedShaderSource = PreprocessedShaderSource;
 						FString ModifiedEntryPointName = Input.EntryPointName;
-						if (RemoveUnusedInputs(ModifiedShaderSource, ShaderInputs, ModifiedEntryPointName, RemoveErrors))
+						if (RemoveUnusedInputs(ModifiedShaderSource, CompileData.ShaderInputs, ModifiedEntryPointName, RemoveErrors))
 						{
 							Output = OriginalOutput;
 							if (!CompileAndProcessD3DShaderFXCExt(CompileFlags, Input, ModifiedShaderSource, ModifiedEntryPointName, ShaderParameterParser, ShaderProfile, true, FilteredErrors, Output))
@@ -1022,7 +981,7 @@ static bool CompileAndProcessD3DShaderFXCExt(
 							}
 
 							// check if the ShaderInputs changed - if not, we're done here
-							if (Output.UsedAttributes.Num() == ShaderInputs.Num())
+							if (Output.UsedAttributes.Num() == CompileData.ShaderInputs.Num())
 							{
 								Output.ModifiedShaderSource = MoveTemp(ModifiedShaderSource);
 								Output.ModifiedEntryPointName = MoveTemp(ModifiedEntryPointName);
@@ -1031,11 +990,11 @@ static bool CompileAndProcessD3DShaderFXCExt(
 							}
 
 							// second pass cannot use more attributes than previously
-							if (Output.UsedAttributes.Num() > ShaderInputs.Num())
+							if (Output.UsedAttributes.Num() > CompileData.ShaderInputs.Num())
 							{
-								UE_LOG(LogD3D11ShaderCompiler, Warning, TEXT("Second pass had more used attributes (%d) than first pass (%d)"), Output.UsedAttributes.Num(), ShaderInputs.Num());
+								UE_LOG(LogD3D11ShaderCompiler, Warning, TEXT("Second pass had more used attributes (%d) than first pass (%d)"), Output.UsedAttributes.Num(), CompileData.ShaderInputs.Num());
 								FShaderCompilerError NewError;
-								NewError.StrippedErrorMessage = FString::Printf(TEXT("Second pass had more used attributes (%d) than first pass (%d)"), Output.UsedAttributes.Num(), ShaderInputs.Num());
+								NewError.StrippedErrorMessage = FString::Printf(TEXT("Second pass had more used attributes (%d) than first pass (%d)"), Output.UsedAttributes.Num(), CompileData.ShaderInputs.Num());
 								Output = OriginalOutput;
 								Output.Errors.Add(NewError);
 								break;
@@ -1047,20 +1006,20 @@ static bool CompileAndProcessD3DShaderFXCExt(
 								UE_LOG(LogD3D11ShaderCompiler, Warning, TEXT("Unable to determine unused inputs after %d attempts (last number of used attributes: %d, previous step:%d)!"), 
 									Attempt + 1,
 									Output.UsedAttributes.Num(),
-									ShaderInputs.Num()
+									CompileData.ShaderInputs.Num()
 									);
 								FShaderCompilerError NewError;
 								NewError.StrippedErrorMessage = FString::Printf(TEXT("Unable to determine unused inputs after %d attempts (last number of used attributes: %d, previous step:%d)!"),
 									Attempt + 1,
 									Output.UsedAttributes.Num(),
-									ShaderInputs.Num()
+									CompileData.ShaderInputs.Num()
 									);
 								Output = OriginalOutput;
 								Output.Errors.Add(NewError);
 								break;
 							}
 
-							ShaderInputs = Output.UsedAttributes;
+							CompileData.ShaderInputs = Output.UsedAttributes;
 							// go around to remove newly identified unused inputs
 						}
 						else
@@ -1085,13 +1044,14 @@ static bool CompileAndProcessD3DShaderFXCExt(
 				ID3D11ShaderReflection, D3D11_SHADER_DESC, D3D11_SHADER_INPUT_BIND_DESC,
 				ID3D11ShaderReflectionConstantBuffer, D3D11_SHADER_BUFFER_DESC,
 				ID3D11ShaderReflectionVariable, D3D11_SHADER_VARIABLE_DESC>(
-					Input, ShaderParameterParser,
-					BindingSpace, Reflector, ShaderDesc,
-					bGlobalUniformBufferUsed, bDiagnosticBufferUsed,
-					NumSamplers, NumSRVs, NumCBs, NumUAVs,
-					Output, UniformBufferNames, UsedUniformBufferSlots, VendorExtensions);
-
-			NumInstructions = ShaderDesc.InstructionCount;
+					Input,
+					ShaderParameterParser,
+					BindingSpace,
+					Reflector,
+					ShaderDesc,
+					CompileData,
+					Output
+				);
 		}
 		else
 		{
@@ -1100,16 +1060,16 @@ static bool CompileAndProcessD3DShaderFXCExt(
 			Output.bSucceeded = false;
 		}
 		
-		if (!ValidateResourceCounts(NumSRVs, NumSamplers, NumUAVs, NumCBs, FilteredErrors))
+		if (!ValidateResourceCounts(CompileData, FilteredErrors))
 		{
 			Result = E_FAIL;
 			Output.bSucceeded = false;
 		}
 
 		// Check for resource limits for feature level 11.0
-		if (NumUAVs > GD3DMaximumNumUAVs)
+		if (CompileData.NumUAVs > GD3DMaximumNumUAVs)
 		{
-			FilteredErrors.Add(FString::Printf(TEXT("Number of UAVs exceeded limit: %d slots used, but limit is %d due to maximum feature level 11.0"), NumUAVs, GD3DMaximumNumUAVs));
+			FilteredErrors.Add(FString::Printf(TEXT("Number of UAVs exceeded limit: %d slots used, but limit is %d due to maximum feature level 11.0"), CompileData.NumUAVs, GD3DMaximumNumUAVs));
 			Result = E_FAIL;
 			Output.bSucceeded = false;
 		}
@@ -1162,25 +1122,19 @@ static bool CompileAndProcessD3DShaderFXCExt(
 				Output.ShaderCode.AddOptionalData(ResourceMasks);
 			};
 
-			FShaderCodePackedResourceCounts PackedResourceCounts{};
-			if (bGlobalUniformBufferUsed)
-			{
-				PackedResourceCounts.UsageFlags |= EShaderResourceUsageFlags::GlobalUniformBuffer;
-			}
-
-			PackedResourceCounts.NumSamplers = static_cast<uint8>(NumSamplers);
-			PackedResourceCounts.NumSRVs = static_cast<uint8>(NumSRVs);
-			PackedResourceCounts.NumCBs = static_cast<uint8>(NumCBs);
-			PackedResourceCounts.NumUAVs = static_cast<uint8>(NumUAVs);
+			FShaderCodePackedResourceCounts PackedResourceCounts = InitPackedResourceCounts(CompileData);
 
-			GenerateFinalOutput(CompressedData,
-				Input, VendorExtensions,
-				UsedUniformBufferSlots, UniformBufferNames,
-				bSecondPassAferUnusedInputRemoval, ShaderInputs,
-				PackedResourceCounts, NumInstructions,
+			GenerateFinalOutput(
+				CompressedData,
+				Input,
+				ED3DShaderModel::SM5_0,
+				bSecondPassAferUnusedInputRemoval,
+				CompileData,
+				PackedResourceCounts,
 				Output,
 				[](FMemoryWriter&){},
-				AddOptionalDataCallback);
+				AddOptionalDataCallback
+			);
 		}
 	}
 
@@ -1196,7 +1150,7 @@ bool CompileAndProcessD3DShaderFXC(
 	bool bSecondPassAferUnusedInputRemoval,
 	FShaderCompilerOutput& Output)
 {
-	// @TODO - implement different material path to allow us to remove backwards compat flag on sm5 shaders
+	// @TODO - implement different material path to allow us to remove backwards compatibility flag on sm5 shaders
 	uint32 CompileFlags = D3DCOMPILE_ENABLE_BACKWARDS_COMPATIBILITY
 		// Unpack uniform matrices as row-major to match the CPU layout.
 		| D3DCOMPILE_PACK_MATRIX_ROW_MAJOR;
@@ -1210,16 +1164,13 @@ bool CompileAndProcessD3DShaderFXC(
 	{
 		CompileFlags |= D3DCOMPILE_SKIP_OPTIMIZATION;
 	}
+	else if (Input.Environment.CompilerFlags.Contains(CFLAG_StandardOptimization))
+	{
+		CompileFlags |= D3DCOMPILE_OPTIMIZATION_LEVEL1;
+	}
 	else
 	{
-		if (Input.Environment.CompilerFlags.Contains(CFLAG_StandardOptimization))
-		{
-			CompileFlags |= D3DCOMPILE_OPTIMIZATION_LEVEL1;
-		}
-		else
-		{
-			CompileFlags |= D3DCOMPILE_OPTIMIZATION_LEVEL3;
-		}
+		CompileFlags |= D3DCOMPILE_OPTIMIZATION_LEVEL3;
 	}
 
 	Input.Environment.CompilerFlags.Iterate([&CompileFlags](uint32 Flag)
diff --git a/Engine/Source/Developer/Windows/ShaderFormatD3D/Private/D3DShaderCompiler.inl b/Engine/Source/Developer/Windows/ShaderFormatD3D/Private/D3DShaderCompiler.inl
--- a/Engine/Source/Developer/Windows/ShaderFormatD3D/Private/D3DShaderCompiler.inl
+++ b/Engine/Source/Developer/Windows/ShaderFormatD3D/Private/D3DShaderCompiler.inl
@@ -2,6 +2,37 @@
 
 #pragma once
 
+#include "D3D12RHI.h"
+
+struct FD3DShaderCompileData
+{
+	FD3DShaderCompileData()
+		: UsedUniformBufferSlots(false, 32)
+	{
+	}
+
+	TArray<FShaderCodeVendorExtension> VendorExtensions;
+	TArray<FString> ShaderInputs;
+	TArray<FString> UniformBufferNames;
+	TBitArray<> UsedUniformBufferSlots;
+
+	bool bBindlessResources = false;
+	bool bBindlessSamplers = false;
+	bool bGlobalUniformBufferUsed = false;
+	bool bDiagnosticBufferUsed = false;
+
+	uint32 NumInstructions = 0;
+	uint32 NumSamplers = 0;
+	uint32 NumSRVs = 0;
+	uint32 NumCBs = 0;
+	uint32 NumUAVs = 0;
+
+	uint32 MaxSamplers = 0;
+	uint32 MaxSRVs = 0;
+	uint32 MaxCBs = 0;
+	uint32 MaxUAVs = 0;
+};
+
 template <typename ID3D1xShaderReflection, typename D3D1x_SHADER_DESC, typename D3D1x_SHADER_INPUT_BIND_DESC,
 	typename ID3D1xShaderReflectionConstantBuffer, typename D3D1x_SHADER_BUFFER_DESC,
 	typename ID3D1xShaderReflectionVariable, typename D3D1x_SHADER_VARIABLE_DESC>
@@ -10,10 +41,11 @@ template <typename ID3D1xShaderReflection, typename D3D1x_SHADER_DESC, typename
 		const FShaderParameterParser& ShaderParameterParser,
 		uint32 BindingSpace,
 		ID3D1xShaderReflection* Reflector, const D3D1x_SHADER_DESC& ShaderDesc,
-		bool& bGlobalUniformBufferUsed, bool& bDiagnosticBufferUsed, uint32& NumSamplers, uint32& NumSRVs, uint32& NumCBs, uint32& NumUAVs,
-		FShaderCompilerOutput& Output, TArray<FString>& UniformBufferNames, TBitArray<>& UsedUniformBufferSlots, TArray<FShaderCodeVendorExtension>& VendorExtensions)
+		FD3DShaderCompileData& CompileData,
+		FShaderCompilerOutput& Output
+	)
 {
-	const bool bBindlessEnabled = (Input.Environment.CompilerFlags.Contains(CFLAG_BindlessResources) || Input.Environment.CompilerFlags.Contains(CFLAG_BindlessSamplers));
+	const bool bBindlessEnabled = (CompileData.bBindlessResources || CompileData.bBindlessSamplers);
 
 	// Add parameters for shader resources (constant buffers, textures, samplers, etc. */
 	for (uint32 ResourceIndex = 0; ResourceIndex < ShaderDesc.BoundResources; ResourceIndex++)
@@ -67,7 +99,7 @@ template <typename ID3D1xShaderReflection, typename D3D1x_SHADER_DESC, typename
 						Variable->GetDesc(&VariableDesc);
 						if (VariableDesc.uFlags & D3D_SVF_USED)
 						{
-							bGlobalUniformBufferUsed = true;
+							CompileData.bGlobalUniformBufferUsed = true;
 
 							HandleReflectedGlobalConstantBufferMember(
 								FString(VariableDesc.Name),
@@ -77,7 +109,7 @@ template <typename ID3D1xShaderReflection, typename D3D1x_SHADER_DESC, typename
 								Output
 							);
 
-							UsedUniformBufferSlots[CBIndex] = true;
+							CompileData.UsedUniformBufferSlots[CBIndex] = true;
 						}
 					}
 				}
@@ -118,8 +150,8 @@ template <typename ID3D1xShaderReflection, typename D3D1x_SHADER_DESC, typename
 					{
 						HandleReflectedRootConstantBuffer(ConstantBufferSize, Output);
 
-						bGlobalUniformBufferUsed = true;
-						UsedUniformBufferSlots[CBIndex] = true;
+						CompileData.bGlobalUniformBufferUsed = true;
+						CompileData.UsedUniformBufferSlots[CBIndex] = true;
 					}
 				}
 				else
@@ -141,7 +173,7 @@ template <typename ID3D1xShaderReflection, typename D3D1x_SHADER_DESC, typename
 				AddShaderValidationUBSize(CBIndex, CBDesc.Size, Output);
 				HandleReflectedUniformBuffer(UniformBufferName, CBIndex, Output);
 				
-				UsedUniformBufferSlots[CBIndex] = true;
+				CompileData.UsedUniformBufferSlots[CBIndex] = true;
 
 				if (bBindlessEnabled)
 				{
@@ -166,13 +198,13 @@ template <typename ID3D1xShaderReflection, typename D3D1x_SHADER_DESC, typename
 				}
 			}
 
-			if (UniformBufferNames.Num() <= (int32)CBIndex)
+			if (CompileData.UniformBufferNames.Num() <= (int32)CBIndex)
 			{
-				UniformBufferNames.AddDefaulted(CBIndex - UniformBufferNames.Num() + 1);
+				CompileData.UniformBufferNames.AddDefaulted(CBIndex - CompileData.UniformBufferNames.Num() + 1);
 			}
-			UniformBufferNames[CBIndex] = UE::ShaderCompilerCommon::RemoveConstantBufferPrefix(FString(CBDesc.Name));
+			CompileData.UniformBufferNames[CBIndex] = UE::ShaderCompilerCommon::RemoveConstantBufferPrefix(FString(CBDesc.Name));
 
-			NumCBs = FMath::Max(NumCBs, BindDesc.BindPoint + BindDesc.BindCount);
+			CompileData.NumCBs = FMath::Max(CompileData.NumCBs, BindDesc.BindPoint + BindDesc.BindCount);
 		}
 		else if (BindDesc.Type == D3D_SIT_TEXTURE || BindDesc.Type == D3D_SIT_SAMPLER)
 		{
@@ -188,12 +220,12 @@ template <typename ID3D1xShaderReflection, typename D3D1x_SHADER_DESC, typename
 
 			if (bIsVendorParameter)
 			{
-				VendorExtensions.Emplace(EGpuVendorId::Amd, 0, BindDesc.BindPoint, BindCount, ParameterType);
+				CompileData.VendorExtensions.Emplace(EGpuVendorId::Amd, 0, BindDesc.BindPoint, BindCount, ParameterType);
 			}
 			else if (ParameterType == EShaderParameterType::Sampler)
 			{
 				HandleReflectedShaderSampler(FString(BindDesc.Name), BindDesc.BindPoint, Output);
-				NumSamplers = FMath::Max(NumSamplers, BindDesc.BindPoint + BindCount);
+				CompileData.NumSamplers = FMath::Max(CompileData.NumSamplers, BindDesc.BindPoint + BindCount);
 			}
 			else
 			{
@@ -201,7 +233,7 @@ template <typename ID3D1xShaderReflection, typename D3D1x_SHADER_DESC, typename
 				AddShaderValidationSRVType(BindDesc.BindPoint, ParsedParam ? ParsedParam->ParsedTypeDecl : EShaderCodeResourceBindingType::Invalid, Output);
 
 				HandleReflectedShaderResource(FString(BindDesc.Name), BindDesc.BindPoint, Output);
-				NumSRVs = FMath::Max(NumSRVs, BindDesc.BindPoint + BindCount);
+				CompileData.NumSRVs = FMath::Max(CompileData.NumSRVs, BindDesc.BindPoint + BindCount);
 			}
 		}
 		else if (BindDesc.Type == D3D_SIT_UAV_RWTYPED || BindDesc.Type == D3D_SIT_UAV_RWSTRUCTURED ||
@@ -235,11 +267,11 @@ template <typename ID3D1xShaderReflection, typename D3D1x_SHADER_DESC, typename
 					(bIsAMDExtensionDX11 || bIsAMDExtensionDX12) ? EGpuVendorId::Amd :
 					bIsIntelExtension ? EGpuVendorId::Intel :
 					EGpuVendorId::Unknown;
-				VendorExtensions.Emplace(VendorId, 0, BindDesc.BindPoint, BindCount, EShaderParameterType::UAV);
+				CompileData.VendorExtensions.Emplace(VendorId, 0, BindDesc.BindPoint, BindCount, EShaderParameterType::UAV);
 			}
 			else if (bIsDiagnosticBufferParameter)
 			{
-				bDiagnosticBufferUsed = true;
+				CompileData.bDiagnosticBufferUsed = true;
 			}
 			else
 			{
@@ -247,7 +279,7 @@ template <typename ID3D1xShaderReflection, typename D3D1x_SHADER_DESC, typename
 				AddShaderValidationUAVType(BindDesc.BindPoint, ParsedParam ? ParsedParam->ParsedTypeDecl : EShaderCodeResourceBindingType::Invalid, Output);
 
 				HandleReflectedShaderUAV(FString(BindDesc.Name), BindDesc.BindPoint, Output);
-				NumUAVs = FMath::Max(NumUAVs, BindDesc.BindPoint + BindCount);
+				CompileData.NumUAVs = FMath::Max(CompileData.NumUAVs, BindDesc.BindPoint + BindCount);
 			}
 		}
 		else if (BindDesc.Type == D3D_SIT_STRUCTURED || BindDesc.Type == D3D_SIT_BYTEADDRESS)
@@ -266,7 +298,7 @@ template <typename ID3D1xShaderReflection, typename D3D1x_SHADER_DESC, typename
 				UpdateStructuredBufferStride(Input, BindDescName, BindDesc.BindPoint, BindDesc.NumSamples, Output);
 			}
 
-			NumSRVs = FMath::Max(NumSRVs, BindDesc.BindPoint + 1);
+			CompileData.NumSRVs = FMath::Max(CompileData.NumSRVs, BindDesc.BindPoint + 1);
 		}
 		else if (BindDesc.Type == (D3D_SHADER_INPUT_TYPE)(D3D_SIT_UAV_RWSTRUCTURED_WITH_COUNTER + 1)) // D3D_SIT_RTACCELERATIONSTRUCTURE (12)
 		{
@@ -277,21 +309,94 @@ template <typename ID3D1xShaderReflection, typename D3D1x_SHADER_DESC, typename
 			AddShaderValidationSRVType(BindDesc.BindPoint, ParsedParam ? ParsedParam->ParsedTypeDecl : EShaderCodeResourceBindingType::Invalid, Output);
 
 			HandleReflectedShaderResource(FString(BindDesc.Name), BindDesc.BindPoint, Output);
-			NumSRVs = FMath::Max(NumSRVs, BindDesc.BindPoint + 1);
+			CompileData.NumSRVs = FMath::Max(CompileData.NumSRVs, BindDesc.BindPoint + 1);
 		}
 	}
+
+	CompileData.NumInstructions = ShaderDesc.InstructionCount;
+}
+
+// Validate that we are not going over to maximum amount of resource bindings support by the default root signature on DX12
+// Currently limited for hard-coded root signature setup (see: FD3D12Adapter::StaticGraphicsRootSignature)
+// In theory this limitation is only required for DX12, but we don't want a shader to compile on DX11 while not working on DX12.
+// (DX11 has an API limit on 128 SRVs, 16 Samplers, 8 UAVs and 14 CBs but if you go over these values then the shader won't compile)
+inline bool ValidateResourceCounts(const FD3DShaderCompileData& CompileData, TArray<FString>& OutFilteredErrors)
+{
+	const bool bTooManySRVs     = !CompileData.bBindlessResources && CompileData.NumSRVs     > CompileData.MaxSRVs;
+	const bool bTooManyUAVs     = !CompileData.bBindlessResources && CompileData.NumUAVs     > CompileData.MaxUAVs;
+	const bool bTooManySamplers = !CompileData.bBindlessSamplers && CompileData.NumSamplers  > CompileData.MaxSamplers;
+	const bool bTooManyCBs = CompileData.NumCBs > CompileData.MaxCBs;
+
+	if (bTooManySRVs || bTooManySamplers || bTooManyUAVs || bTooManyCBs)
+	{
+		if (bTooManySRVs)
+		{
+			OutFilteredErrors.Add(FString::Printf(TEXT("Shader is using too many SRVs: %d (only %d supported)"), CompileData.NumSRVs, CompileData.MaxSRVs));
+		}
+
+		if (bTooManySamplers)
+		{
+			OutFilteredErrors.Add(FString::Printf(TEXT("Shader is using too many Samplers: %d (only %d supported)"), CompileData.NumSamplers, CompileData.MaxSamplers));
+		}
+
+		if (bTooManyUAVs)
+		{
+			OutFilteredErrors.Add(FString::Printf(TEXT("Shader is using too many UAVs: %d (only %d supported)"), CompileData.NumUAVs, CompileData.MaxUAVs));
+		}
+
+		if (bTooManyCBs)
+		{
+			OutFilteredErrors.Add(FString::Printf(TEXT("Shader is using too many Constant Buffers: %d (only %d supported)"), CompileData.NumCBs, CompileData.MaxCBs));
+		}
+
+		return false;
+	}
+
+	return true;
+}
+
+inline FShaderCodePackedResourceCounts InitPackedResourceCounts(const FD3DShaderCompileData& CompileData)
+{
+	FShaderCodePackedResourceCounts PackedResourceCounts{};
+
+	if (CompileData.bGlobalUniformBufferUsed)
+	{
+		PackedResourceCounts.UsageFlags |= EShaderResourceUsageFlags::GlobalUniformBuffer;
+	}
+
+	if (CompileData.bBindlessResources)
+	{
+		PackedResourceCounts.UsageFlags |= EShaderResourceUsageFlags::BindlessResources;
+	}
+
+	if (CompileData.bBindlessSamplers)
+	{
+		PackedResourceCounts.UsageFlags |= EShaderResourceUsageFlags::BindlessSamplers;
+	}
+
+	PackedResourceCounts.NumSamplers = static_cast<uint8>(CompileData.NumSamplers);
+	PackedResourceCounts.NumSRVs = static_cast<uint8>(CompileData.NumSRVs);
+	PackedResourceCounts.NumCBs = static_cast<uint8>(CompileData.NumCBs);
+	PackedResourceCounts.NumUAVs = static_cast<uint8>(CompileData.NumUAVs);
+
+	return PackedResourceCounts;
 }
 
 template <typename TBlob>
-inline void GenerateFinalOutput(TRefCountPtr<TBlob>& CompressedData,
-	const FShaderCompilerInput& Input, TArray<FShaderCodeVendorExtension>& VendorExtensions, 
-	TBitArray<>& UsedUniformBufferSlots, TArray<FString>& UniformBufferNames,
-	bool bProcessingSecondTime, const TArray<FString>& ShaderInputs,
-	FShaderCodePackedResourceCounts& PackedResourceCounts, uint32 NumInstructions,
+inline void GenerateFinalOutput(
+	TRefCountPtr<TBlob>& CompressedData,
+	const FShaderCompilerInput& Input,
+	ED3DShaderModel ShaderModel,
+	bool bProcessingSecondTime,
+	FD3DShaderCompileData& CompileData,
+	const FShaderCodePackedResourceCounts& PackedResourceCounts,
 	FShaderCompilerOutput& Output,
 	TFunction<void(FMemoryWriter&)> PostSRTWriterCallback,
 	TFunction<void(FShaderCode&)> AddOptionalDataCallback)
 {
+	const uint32 NumBindlessResources = CompileData.bBindlessResources ? Output.ParameterMap.CountParametersOfType(EShaderParameterType::BindlessSRV) : 0;
+	const uint32 NumBindlessSamplers = CompileData.bBindlessSamplers ? Output.ParameterMap.CountParametersOfType(EShaderParameterType::BindlessSampler) : 0;
+
 	// Build the SRT for this shader.
 	FShaderResourceTable SRT;
 
@@ -300,7 +405,7 @@ inline void GenerateFinalOutput(TRefCountPtr<TBlob>& CompressedData,
 	{	
 		// Build the generic SRT for this shader.
 		FShaderCompilerResourceTable GenericSRT;
-		BuildResourceTableMapping(Input.Environment.ResourceTableMap, Input.Environment.UniformBufferMap, UsedUniformBufferSlots, Output.ParameterMap, GenericSRT);
+		BuildResourceTableMapping(Input.Environment.ResourceTableMap, Input.Environment.UniformBufferMap, CompileData.UsedUniformBufferSlots, Output.ParameterMap, GenericSRT);
 
 		// Ray generation shaders rely on a different binding model that aren't compatible with global uniform buffers.
 		if (Input.Target.Frequency != SF_RayGen)
@@ -308,14 +413,14 @@ inline void GenerateFinalOutput(TRefCountPtr<TBlob>& CompressedData,
 			CullGlobalUniformBuffers(Input.Environment.UniformBufferMap, Output.ParameterMap);
 		}
 
-		if (UniformBufferNames.Num() < GenericSRT.ResourceTableLayoutHashes.Num())
+		if (CompileData.UniformBufferNames.Num() < GenericSRT.ResourceTableLayoutHashes.Num())
 		{
-			UniformBufferNames.AddDefaulted(GenericSRT.ResourceTableLayoutHashes.Num() - UniformBufferNames.Num());
+			CompileData.UniformBufferNames.AddDefaulted(GenericSRT.ResourceTableLayoutHashes.Num() - CompileData.UniformBufferNames.Num());
 		}
 
 		for (int32 Index = 0; Index < GenericSRT.ResourceTableLayoutHashes.Num(); ++Index)
 		{
-			if (GenericSRT.ResourceTableLayoutHashes[Index] != 0 && UniformBufferNames[Index].Len() == 0)
+			if (GenericSRT.ResourceTableLayoutHashes[Index] != 0 && CompileData.UniformBufferNames[Index].Len() == 0)
 			{
 				for (const auto& KeyValue : Input.Environment.UniformBufferMap)
 				{
@@ -323,7 +428,7 @@ inline void GenerateFinalOutput(TRefCountPtr<TBlob>& CompressedData,
 
 					if (UniformBufferEntry.LayoutHash == GenericSRT.ResourceTableLayoutHashes[Index])
 					{
-						UniformBufferNames[Index] = KeyValue.Key;
+						CompileData.UniformBufferNames[Index] = KeyValue.Key;
 						break;
 					}
 				}
@@ -331,7 +436,7 @@ inline void GenerateFinalOutput(TRefCountPtr<TBlob>& CompressedData,
 		}
 
 		FMemoryWriter UniformBufferNameWriter(UniformBufferNameBytes);
-		UniformBufferNameWriter << UniformBufferNames;
+		UniformBufferNameWriter << CompileData.UniformBufferNames;
 
 		// Copy over the bits indicating which resource tables are active.
 		SRT.ResourceTableBits = GenericSRT.ResourceTableBits;
@@ -346,10 +451,9 @@ inline void GenerateFinalOutput(TRefCountPtr<TBlob>& CompressedData,
 	}
 
 	if (Input.Environment.CompilerFlags.Contains(CFLAG_ForceRemoveUnusedInterpolators) && Input.Target.Frequency == SF_Pixel && Input.bCompilingForShaderPipeline && bProcessingSecondTime)
-
 	{
 		Output.bSupportsQueryingUsedAttributes = true;
-		Output.UsedAttributes = ShaderInputs;
+		Output.UsedAttributes = CompileData.ShaderInputs;
 	}
 
 	// Generate the final Output
@@ -368,11 +472,11 @@ inline void GenerateFinalOutput(TRefCountPtr<TBlob>& CompressedData,
 	}
 
 	// Append information about optional hardware vendor extensions
-	if (VendorExtensions.Num() > 0)
+	if (CompileData.VendorExtensions.Num() > 0)
 	{
 		TArray<uint8> WriterBytes;
 		FMemoryWriter Writer(WriterBytes);
-		Writer << VendorExtensions;
+		Writer << CompileData.VendorExtensions;
 		if (WriterBytes.Num() > 0)
 		{
 			Output.ShaderCode.AddOptionalData(FShaderCodeVendorExtension::Key, WriterBytes.GetData(), WriterBytes.Num());
@@ -388,10 +492,33 @@ inline void GenerateFinalOutput(TRefCountPtr<TBlob>& CompressedData,
 	Output.SerializeShaderDiagnosticData();
 
 	// Set the number of instructions.
-	Output.NumInstructions = NumInstructions;
+	Output.NumInstructions = CompileData.NumInstructions;
 
 	Output.NumTextureSamplers = PackedResourceCounts.NumSamplers;
 
 	// Pass the target through to the output.
 	Output.Target = Input.Target;
+
+	// SRV Limits
+	{
+		if (CompileData.bBindlessResources)
+		{
+			Output.AddStatistic<uint32>(TEXT("Bindless Resources"), NumBindlessResources);
+		}
+		else
+		{
+			Output.AddStatistic<uint32>(TEXT("Resources Used"), CompileData.NumSRVs);
+			Output.AddStatistic<uint32>(TEXT("Resource Limit"), CompileData.MaxSRVs);
+		}
+
+		if (CompileData.bBindlessSamplers)
+		{
+			Output.AddStatistic<uint32>(TEXT("Bindless Samplers"), NumBindlessSamplers);
+		}
+		else
+		{
+			Output.AddStatistic<uint32>(TEXT("Samplers Used"), CompileData.NumSamplers);
+			Output.AddStatistic<uint32>(TEXT("Sampler Limit"), CompileData.MaxSamplers);
+		}
+	}
 }
diff --git a/Engine/Source/Developer/Windows/ShaderFormatD3D/Private/D3DShaderCompilerDXC.cpp b/Engine/Source/Developer/Windows/ShaderFormatD3D/Private/D3DShaderCompilerDXC.cpp
--- a/Engine/Source/Developer/Windows/ShaderFormatD3D/Private/D3DShaderCompilerDXC.cpp
+++ b/Engine/Source/Developer/Windows/ShaderFormatD3D/Private/D3DShaderCompilerDXC.cpp
@@ -22,7 +22,7 @@ DEFINE_LOG_CATEGORY_STATIC(LogD3D12ShaderCompiler, Log, All);
 #pragma warning(disable : 4005)	// macro redefinition
 
 #include "Windows/AllowWindowsPlatformTypes.h"
-	#include <D3D11.h>
+	#include <D3D12.h>
 	#include <D3Dcompiler.h>
 	#include <d3d11Shader.h>
 	#include "amd_ags.h"
@@ -934,21 +934,34 @@ bool CompileAndProcessD3DShaderDXC(
 	if (SUCCEEDED(D3DCompileToDxilResult))
 	{
 		// Gather reflection information
-		TArray<FString> ShaderInputs;
-		TArray<FShaderCodeVendorExtension> VendorExtensions;
-
-		bool bGlobalUniformBufferUsed = false;
-		bool bDiagnosticBufferUsed = false;
-		uint32 NumInstructions = 0;
-		uint32 NumSamplers = 0;
-		uint32 NumSRVs = 0;
-		uint32 NumCBs = 0;
-		uint32 NumUAVs = 0;
-		TArray<FString> UniformBufferNames;
-		TArray<FString> ShaderOutputs;
-
-		TBitArray<> UsedUniformBufferSlots;
-		UsedUniformBufferSlots.Init(false, 32);
+		FD3DShaderCompileData CompileData;
+		CompileData.bBindlessResources = Input.Environment.CompilerFlags.Contains(CFLAG_BindlessResources);
+		CompileData.bBindlessSamplers = Input.Environment.CompilerFlags.Contains(CFLAG_BindlessSamplers);
+
+		if (CompileData.bBindlessSamplers)
+		{
+			CompileData.MaxSamplers = D3D12_MAX_SHADER_VISIBLE_SAMPLER_HEAP_SIZE;
+		}
+		else if (ShaderModel == ED3DShaderModel::SM6_6)
+		{
+			CompileData.MaxSamplers = 32; // DDSPI: MaxSamplers=32
+		}
+		else
+		{
+			CompileData.MaxSamplers = D3D12_COMMONSHADER_SAMPLER_REGISTER_COUNT;
+		}
+
+		if (Input.Environment.CompilerFlags.Contains(CFLAG_BindlessResources))
+		{
+			CompileData.MaxSRVs = D3D12_MAX_SHADER_VISIBLE_DESCRIPTOR_HEAP_SIZE_TIER_2;
+		}
+		else
+		{
+			static_assert(MAX_SRVS <= D3D12_COMMONSHADER_INPUT_RESOURCE_SLOT_COUNT);
+			CompileData.MaxSRVs = MAX_SRVS; // Max for D3D12RHI bindful
+		}
+		CompileData.MaxCBs = MAX_CBS; // Max for D3D12RHI
+		CompileData.MaxUAVs = MAX_UAVS; // Max for D3D12RHI
 
 		uint64 ShaderRequiresFlags{};
 
@@ -1037,11 +1050,13 @@ bool CompileAndProcessD3DShaderDXC(
 						ExtractParameterMapFromD3DShader<ID3D12FunctionReflection, D3D12_FUNCTION_DESC, D3D12_SHADER_INPUT_BIND_DESC,
 							ID3D12ShaderReflectionConstantBuffer, D3D12_SHADER_BUFFER_DESC,
 							ID3D12ShaderReflectionVariable, D3D12_SHADER_VARIABLE_DESC>(
-								Input, ShaderParameterParser,
-								AutoBindingSpace, FunctionReflection, FunctionDesc, 
-								bGlobalUniformBufferUsed, bDiagnosticBufferUsed,
-								NumSamplers, NumSRVs, NumCBs, NumUAVs,
-								Output, UniformBufferNames, UsedUniformBufferSlots, VendorExtensions);
+								Input,
+								ShaderParameterParser,
+								AutoBindingSpace,
+								FunctionReflection,
+								FunctionDesc, 
+								CompileData,
+								Output);
 
 						NumFoundEntryPoints++;
 					}
@@ -1051,11 +1066,11 @@ bool CompileAndProcessD3DShaderDXC(
 			// @todo - working around DXC issue [REDACTED_PROFILE_URL]/DirectXShaderCompiler/issues/4715
 			if (LibraryDesc.FunctionCount > 0)
 			{
-				if (Input.Environment.CompilerFlags.Contains(CFLAG_BindlessResources))
+				if (CompileData.bBindlessResources)
 				{
 					ShaderRequiresFlags |= D3D_SHADER_REQUIRES_RESOURCE_DESCRIPTOR_HEAP_INDEXING;
 				}
-				if (Input.Environment.CompilerFlags.Contains(CFLAG_BindlessSamplers))
+				if (CompileData.bBindlessSamplers)
 				{
 					ShaderRequiresFlags |= D3D_SHADER_REQUIRES_SAMPLER_DESCRIPTOR_HEAP_INDEXING;
 				}
@@ -1067,7 +1082,7 @@ bool CompileAndProcessD3DShaderDXC(
 
 				bool bGlobalUniformBufferAllowed = false;
 
-				if (bGlobalUniformBufferUsed && !IsGlobalConstantBufferSupported(Input.Target))
+				if (CompileData.bGlobalUniformBufferUsed && !IsGlobalConstantBufferSupported(Input.Target))
 				{
 					const TCHAR* ShaderFrequencyString = GetShaderFrequencyString(Input.Target.GetFrequency(), false);
 					FString ErrorString = FString::Printf(TEXT("Global uniform buffer cannot be used in a %s shader."), ShaderFrequencyString);
@@ -1120,16 +1135,17 @@ bool CompileAndProcessD3DShaderDXC(
 			ExtractParameterMapFromD3DShader<ID3D12ShaderReflection, D3D12_SHADER_DESC, D3D12_SHADER_INPUT_BIND_DESC,
 				ID3D12ShaderReflectionConstantBuffer, D3D12_SHADER_BUFFER_DESC,
 				ID3D12ShaderReflectionVariable, D3D12_SHADER_VARIABLE_DESC>(
-					Input, ShaderParameterParser,
-					AutoBindingSpace, ShaderReflection, ShaderDesc,
-					bGlobalUniformBufferUsed, bDiagnosticBufferUsed,
-					NumSamplers, NumSRVs, NumCBs, NumUAVs,
-					Output, UniformBufferNames, UsedUniformBufferSlots, VendorExtensions);
-
-			NumInstructions = ShaderDesc.InstructionCount;
+					Input,
+					ShaderParameterParser,
+					AutoBindingSpace,
+					ShaderReflection,
+					ShaderDesc,
+					CompileData,
+					Output
+				);
 		}
 
-		if (!ValidateResourceCounts(NumSRVs, NumSamplers, NumUAVs, NumCBs, FilteredErrors))
+		if (!ValidateResourceCounts(CompileData, FilteredErrors))
 		{
 			Output.bSucceeded = false;
 		}
@@ -1138,36 +1154,18 @@ bool CompileAndProcessD3DShaderDXC(
 
 		if (Output.bSucceeded)
 		{
-			if (bGlobalUniformBufferUsed)
-			{
-				PackedResourceCounts.UsageFlags |= EShaderResourceUsageFlags::GlobalUniformBuffer;
-			}
+			PackedResourceCounts = InitPackedResourceCounts(CompileData);
 
 			if (Input.Environment.CompilerFlags.Contains(CFLAG_RootConstants))
 			{
 				PackedResourceCounts.UsageFlags |= EShaderResourceUsageFlags::RootConstants;
 			}
 
-			if (Input.Environment.CompilerFlags.Contains(CFLAG_BindlessResources))
-			{
-				PackedResourceCounts.UsageFlags |= EShaderResourceUsageFlags::BindlessResources;
-			}
-
-			if (Input.Environment.CompilerFlags.Contains(CFLAG_BindlessSamplers))
-			{
-				PackedResourceCounts.UsageFlags |= EShaderResourceUsageFlags::BindlessSamplers;
-			}
-
 			if (bHasNoDerivativeOps)
 			{
 				PackedResourceCounts.UsageFlags |= EShaderResourceUsageFlags::NoDerivativeOps;
 			}
 
-			PackedResourceCounts.NumSamplers = static_cast<uint8>(NumSamplers);
-			PackedResourceCounts.NumSRVs = static_cast<uint8>(NumSRVs);
-			PackedResourceCounts.NumCBs = static_cast<uint8>(NumCBs);
-			PackedResourceCounts.NumUAVs = static_cast<uint8>(NumUAVs);
-
 			Output.bSucceeded = UE::ShaderCompilerCommon::ValidatePackedResourceCounts(Output, PackedResourceCounts);
 		}
 
@@ -1219,7 +1217,7 @@ bool CompileAndProcessD3DShaderDXC(
 					EnumAddFlags(CodeFeatures.CodeFeatures, EShaderCodeFeatures::Atomic64);
 				}
 
-				if (bDiagnosticBufferUsed)
+				if (CompileData.bDiagnosticBufferUsed)
 				{
 					EnumAddFlags(CodeFeatures.CodeFeatures, EShaderCodeFeatures::DiagnosticBuffer);
 				}
@@ -1257,18 +1255,21 @@ bool CompileAndProcessD3DShaderDXC(
 			// Return a fraction of the number of instructions as DXIL is more verbose than DXBC.
 			// Ratio 119:307 was estimated by gathering average instruction count for D3D11 and D3D12 shaders in ShooterGame with result being ~ 357:921.
 			constexpr uint32 DxbcToDxilInstructionRatio[2] = { 119, 307 };
-			NumInstructions = NumInstructions * DxbcToDxilInstructionRatio[0] / DxbcToDxilInstructionRatio[1];
+			CompileData.NumInstructions = CompileData.NumInstructions * DxbcToDxilInstructionRatio[0] / DxbcToDxilInstructionRatio[1];
 
 			//#todo-rco: Should compress ShaderCode?
 
-			GenerateFinalOutput(ShaderBlob,
-				Input, VendorExtensions,
-				UsedUniformBufferSlots, UniformBufferNames,
-				bProcessingSecondTime, ShaderInputs,
-				PackedResourceCounts, NumInstructions,
+			GenerateFinalOutput(
+				ShaderBlob,
+				Input,
+				ShaderModel,
+				bProcessingSecondTime,
+				CompileData,
+				PackedResourceCounts,
 				Output,
 				PostSRTWriterCallback,
-				AddOptionalDataCallback);
+				AddOptionalDataCallback
+			);
 		}
 	}
 	else
diff --git a/Engine/Source/Developer/Windows/ShaderFormatD3D/Private/ShaderFormatD3D.h b/Engine/Source/Developer/Windows/ShaderFormatD3D/Private/ShaderFormatD3D.h
--- a/Engine/Source/Developer/Windows/ShaderFormatD3D/Private/ShaderFormatD3D.h
+++ b/Engine/Source/Developer/Windows/ShaderFormatD3D/Private/ShaderFormatD3D.h
@@ -54,7 +54,8 @@ bool CompileAndProcessD3DShaderDXC(
 	bool bProcessingSecondTime,
 	FShaderCompilerOutput& Output);
 
-bool ValidateResourceCounts(uint32 NumSRVs, uint32 NumSamplers, uint32 NumUAVs, uint32 NumCBs, TArray<FString>& OutFilteredErrors);
+struct FD3DShaderCompileData;
+bool ValidateResourceCounts(const FD3DShaderCompileData& CompiledData, TArray<FString>& OutFilteredErrors);
 
 struct FD3DSM6ShaderDebugData
 {
diff --git a/Engine/Source/Runtime/RenderCore/Private/ShaderCore.cpp b/Engine/Source/Runtime/RenderCore/Private/ShaderCore.cpp
--- a/Engine/Source/Runtime/RenderCore/Private/ShaderCore.cpp
+++ b/Engine/Source/Runtime/RenderCore/Private/ShaderCore.cpp
@@ -775,6 +775,19 @@ TArray<FString> FShaderParameterMap::GetAllParameterNamesOfType(EShaderParameter
 	return Result;
 }
 
+uint32 FShaderParameterMap::CountParametersOfType(EShaderParameterType InType) const
+{
+	uint32 Result = 0;
+	for (const TMap<FString, FParameterAllocation>::ElementType& Parameter : ParameterMap)
+	{
+		if (Parameter.Value.Type == InType)
+		{
+			Result++;
+		}
+	}
+	return Result;
+}
+
 void FShaderResourceTableMap::Append(const FShaderResourceTableMap& Other)
 {
 	// Get the set of uniform buffers used by the target resource table map
diff --git a/Engine/Source/Runtime/RenderCore/Public/ShaderCompilerCore.h b/Engine/Source/Runtime/RenderCore/Public/ShaderCompilerCore.h
--- a/Engine/Source/Runtime/RenderCore/Public/ShaderCompilerCore.h
+++ b/Engine/Source/Runtime/RenderCore/Public/ShaderCompilerCore.h
@@ -660,6 +660,12 @@ struct FShaderCompilerOutput
 	/** Add optional diagnostic data in ShaderCode to perform assert translation at runtime*/
 	RENDERCORE_API void SerializeShaderDiagnosticData();
 
+	template<typename TValue>
+	void AddStatistic(const TCHAR* Name, TValue Value)
+	{
+		ShaderStatistics.Emplace(FString(Name), FShaderStatVariant(TInPlaceType<TValue>(), Value));
+	}
+
 	// Bump ShaderCompileWorkerOutputVersion if FShaderCompilerOutput changes
 	friend FArchive& operator<<(FArchive& Ar, FShaderCompilerOutput& Output)
 	{
diff --git a/Engine/Source/Runtime/RenderCore/Public/ShaderCore.h b/Engine/Source/Runtime/RenderCore/Public/ShaderCore.h
--- a/Engine/Source/Runtime/RenderCore/Public/ShaderCore.h
+++ b/Engine/Source/Runtime/RenderCore/Public/ShaderCore.h
@@ -299,6 +299,9 @@ public:
 	/** Returns an array of all parameters with the given type. */
 	RENDERCORE_API TArray<FString> GetAllParameterNamesOfType(EShaderParameterType InType) const;
 
+	/** Returns a count of all parameters of the given type. */
+	RENDERCORE_API uint32 CountParametersOfType(EShaderParameterType InType) const;
+
 	/** Checks that all parameters are bound and asserts if any aren't in a debug build
 	* @param InVertexFactoryType can be 0
 	*/
