Disabling PLATFORM_SUPPORTS_CONSTANTBUFFER_OBJECT for now.
We can't get usage information from reflection with this enabled, mainly for bindless, so we'll need to revisit this later.

diff --git a/Engine/Source/Developer/Windows/ShaderFormatD3D/Private/ShaderFormatD3D.cpp b/Engine/Source/Developer/Windows/ShaderFormatD3D/Private/ShaderFormatD3D.cpp
--- a/Engine/Source/Developer/Windows/ShaderFormatD3D/Private/ShaderFormatD3D.cpp
+++ b/Engine/Source/Developer/Windows/ShaderFormatD3D/Private/ShaderFormatD3D.cpp
@@ -16,7 +16,7 @@ static FName NAME_PCD3D_SM5(TEXT("PCD3D_SM5"));
 static FName NAME_PCD3D_ES3_1(TEXT("PCD3D_ES31"));
 
 static const FGuid UE_SHADER_PCD3D_SHARED_VER = FGuid("232A2A59-A6D0-4CDB-A374-F3DB028E413E");
-static const FGuid UE_SHADER_PCD3D_SM6_VER    = FGuid("7BDBA1A9-FAEF-42A1-BCF7-1B921FDEAD9E");
+static const FGuid UE_SHADER_PCD3D_SM6_VER    = FGuid("51A4C815-0CF4-42ED-A23E-DED306076D61");
 static const FGuid UE_SHADER_PCD3D_SM5_VER    = FGuid("5B377D13-C70F-40C5-80C5-C9B228783469");
 static const FGuid UE_SHADER_PCD3D_ES3_1_VER  = FGuid("952939D9-1156-4347-97E9-9FFEA1A9FE14");
 
@@ -209,7 +209,8 @@ public:
 		Input.Environment.SetDefine(TEXT("COMPILER_DXC"),  bDXC);
 
 		// Do we need SM6.0+ features enabled? This is intentionally disconnected from the ED3DShaderModel to allow SM6.0 to be used without new language features.
-		Input.Environment.SetDefine(TEXT("PLATFORM_SUPPORTS_CONSTANTBUFFER_OBJECT"), bSM6Features);
+		// TODO: enable PLATFORM_SUPPORTS_CONSTANTBUFFER_OBJECT once we can get usage info from the constant buffer struct.
+		//Input.Environment.SetDefine(TEXT("PLATFORM_SUPPORTS_CONSTANTBUFFER_OBJECT"), bSM6Features);
 		Input.Environment.SetDefine(TEXT("PLATFORM_SUPPORTS_SM6_0_WAVE_OPERATIONS"), bSM6Features);
 		Input.Environment.SetDefine(TEXT("PLATFORM_SUPPORTS_DIAGNOSTIC_BUFFER"),     bSM6Features);
 
