Adding proper SM6.8 support to ShaderFormatD3D
- Only WorkGraph shaders use this currently.
- SM6.8 needs to set the __SHADER_TARGET_MAJOR and __SHADER_TARGET_MINOR defines
- Changing some logic that was SM6.6 specific to now allow the higher shader model.

diff --git a/Engine/Source/Developer/Windows/ShaderFormatD3D/Private/D3DShaderCompiler.cpp b/Engine/Source/Developer/Windows/ShaderFormatD3D/Private/D3DShaderCompiler.cpp
--- a/Engine/Source/Developer/Windows/ShaderFormatD3D/Private/D3DShaderCompiler.cpp
+++ b/Engine/Source/Developer/Windows/ShaderFormatD3D/Private/D3DShaderCompiler.cpp
@@ -107,7 +107,24 @@ static void D3D11FilterShaderCompileWarnings(const FString& CompileWarnings, TAr
 // @return 0 if not recognized
 static const TCHAR* GetShaderProfileName(const FShaderCompilerInput& Input, ED3DShaderModel ShaderModel)
 {
-	if (ShaderModel == ED3DShaderModel::SM6_6)
+	if (ShaderModel == ED3DShaderModel::SM6_8)
+	{
+		switch (Input.Target.GetFrequency())
+		{
+		case SF_Pixel:         return TEXT("ps_6_8");
+		case SF_Vertex:        return TEXT("vs_6_8");
+		case SF_Mesh:          return TEXT("ms_6_8");
+		case SF_Amplification: return TEXT("as_6_8");
+		case SF_Geometry:      return TEXT("gs_6_8");
+		case SF_Compute:       return TEXT("cs_6_8");
+		case SF_RayGen:
+		case SF_RayMiss:
+		case SF_RayHitGroup:
+		case SF_RayCallable:
+		case SF_WorkGraph:     return TEXT("lib_6_8");
+		}
+	}
+	else if (ShaderModel == ED3DShaderModel::SM6_6)
 	{
 		switch (Input.Target.GetFrequency())
 		{
@@ -121,7 +138,6 @@ static const TCHAR* GetShaderProfileName(const FShaderCompilerInput& Input, ED3D
 		case SF_RayMiss:
 		case SF_RayHitGroup:
 		case SF_RayCallable:   return TEXT("lib_6_6");
-		case SF_WorkGraph:     return TEXT("lib_6_8");
 		}
 	}
 	else if (ShaderModel == ED3DShaderModel::SM6_0)
diff --git a/Engine/Source/Developer/Windows/ShaderFormatD3D/Private/D3DShaderCompilerDXC.cpp b/Engine/Source/Developer/Windows/ShaderFormatD3D/Private/D3DShaderCompilerDXC.cpp
--- a/Engine/Source/Developer/Windows/ShaderFormatD3D/Private/D3DShaderCompilerDXC.cpp
+++ b/Engine/Source/Developer/Windows/ShaderFormatD3D/Private/D3DShaderCompilerDXC.cpp
@@ -950,7 +950,7 @@ bool CompileAndProcessD3DShaderDXC(
 		{
 			CompileData.MaxSamplers = D3D12_MAX_SHADER_VISIBLE_SAMPLER_HEAP_SIZE;
 		}
-		else if (ShaderModel == ED3DShaderModel::SM6_6)
+		else if (ShaderModel >= ED3DShaderModel::SM6_6)
 		{
 			CompileData.MaxSamplers = 32; // DDSPI: MaxSamplers=32
 		}
diff --git a/Engine/Source/Developer/Windows/ShaderFormatD3D/Private/ShaderFormatD3D.cpp b/Engine/Source/Developer/Windows/ShaderFormatD3D/Private/ShaderFormatD3D.cpp
--- a/Engine/Source/Developer/Windows/ShaderFormatD3D/Private/ShaderFormatD3D.cpp
+++ b/Engine/Source/Developer/Windows/ShaderFormatD3D/Private/ShaderFormatD3D.cpp
@@ -111,6 +111,11 @@ public:
 		return ELanguage::Invalid;
 	}
 
+	static bool IsSM68(const FShaderCompilerInput& Input, ELanguage Language)
+	{
+		return Input.Target.GetFrequency() == SF_WorkGraph;
+	}
+
 	static bool IsSM66(const FShaderCompilerInput& Input, ELanguage Language)
 	{
 		return Language == ELanguage::SM6 || IsRayTracingShaderFrequency(Input.Target.GetFrequency());
@@ -119,7 +124,8 @@ public:
 	// Do we need any SM6.0 features?
 	static bool RequiresSM6Features(const FShaderCompilerInput& Input, ELanguage Language)
 	{
-		return IsSM66(Input, Language)
+		return IsSM68(Input, Language)
+			|| IsSM66(Input, Language)
 			|| Input.Environment.CompilerFlags.Contains(CFLAG_WaveOperations)
 			// TODO: Forcing DXC should not change platform flags, this needs to be moved to IsSM60 once existing uses are accounted for.
 			|| Input.Environment.CompilerFlags.Contains(CFLAG_ForceDXC)
@@ -136,6 +142,11 @@ public:
 
 	static ED3DShaderModel DetermineShaderModel(const FShaderCompilerInput& Input, ELanguage Language)
 	{
+		if (IsSM68(Input, Language))
+		{
+			return ED3DShaderModel::SM6_8;
+		}
+
 		if (IsSM66(Input, Language))
 		{
 			return ED3DShaderModel::SM6_6;
@@ -200,7 +211,6 @@ public:
 		// Assume min. spec HW supports with DX12/SM5
 		Input.Environment.SetDefine(TEXT("PLATFORM_SUPPORTS_ROV"), true);
 
-		const bool bSM66        = IsSM66(Input, Language);
 		const bool bSM6Features = RequiresSM6Features(Input, Language);
 		const bool bDXC         = DoesShaderModelRequireDXC(ShaderModel);
 
@@ -215,12 +225,12 @@ public:
 		Input.Environment.SetDefine(TEXT("PLATFORM_SUPPORTS_DIAGNOSTIC_BUFFER"),     bSM6Features);
 
 		// "profiles" are almost analogous to ERHIFeatureLevel but RT shaders are forcing themselves to SM6
-		Input.Environment.SetDefine(TEXT("SM6_PROFILE"),   bSM66);
+		Input.Environment.SetDefine(TEXT("SM6_PROFILE"),   ShaderModel >= ED3DShaderModel::SM6_6);
 		Input.Environment.SetDefine(TEXT("SM5_PROFILE"),   (Language == ELanguage::SM5));
 		Input.Environment.SetDefine(TEXT("ES3_1_PROFILE"), (Language == ELanguage::ES3_1));
 
 		// Add SM6.6+ specific defines. None of these are intended to be enabled in lower SM's
-		if (bSM66)
+		if (ShaderModel >= ED3DShaderModel::SM6_6)
 		{
 			Input.Environment.SetDefine(TEXT("PLATFORM_SUPPORTS_REAL_TYPES"),         Input.Environment.CompilerFlags.Contains(CFLAG_AllowRealTypes));
 			Input.Environment.SetDefine(TEXT("PLATFORM_SUPPORTS_INLINE_RAY_TRACING"), Input.Environment.CompilerFlags.Contains(CFLAG_InlineRayTracing));
@@ -242,6 +252,9 @@ public:
 		case ED3DShaderModel::SM6_6:
 			AddShaderTargetDefines(Input, 6, 6);
 			break;
+		case ED3DShaderModel::SM6_8:
+			AddShaderTargetDefines(Input, 6, 8);
+			break;
 		}
 	}
 
diff --git a/Engine/Source/Developer/Windows/ShaderFormatD3D/Private/ShaderFormatD3D.h b/Engine/Source/Developer/Windows/ShaderFormatD3D/Private/ShaderFormatD3D.h
--- a/Engine/Source/Developer/Windows/ShaderFormatD3D/Private/ShaderFormatD3D.h
+++ b/Engine/Source/Developer/Windows/ShaderFormatD3D/Private/ShaderFormatD3D.h
@@ -13,6 +13,7 @@ enum class ED3DShaderModel
 	SM5_0,
 	SM6_0,
 	SM6_6,
+	SM6_8,
 };
 
 inline bool DoesShaderModelRequireDXC(ED3DShaderModel ShaderModel)
