Removing LogD3D11ShaderCompiler and LogD3D12ShaderCompiler in favor of LogD3DShaderCompiler.

diff --git a/Engine/Source/Developer/Windows/ShaderFormatD3D/Private/D3DShaderCompiler.cpp b/Engine/Source/Developer/Windows/ShaderFormatD3D/Private/D3DShaderCompiler.cpp
--- a/Engine/Source/Developer/Windows/ShaderFormatD3D/Private/D3DShaderCompiler.cpp
+++ b/Engine/Source/Developer/Windows/ShaderFormatD3D/Private/D3DShaderCompiler.cpp
@@ -14,8 +14,6 @@
 #include "ShaderPreprocessTypes.h"
 #include "SpirvCommon.h"
 
-DEFINE_LOG_CATEGORY_STATIC(LogD3D11ShaderCompiler, Log, All);
-
 #define DEBUG_SHADERS 0
 
 // D3D doesn't define a mask for this, so we do so here
@@ -371,7 +369,7 @@ private:
 		CompilerDLL = LoadLibrary(*CompilerPath);
 		if (!CompilerDLL)
 		{
-			UE_LOG(LogD3D11ShaderCompiler, Fatal, TEXT("Cannot find the compiler DLL '%s'"), *CompilerPath);
+			UE_LOG(LogD3DShaderCompiler, Fatal, TEXT("Cannot find the compiler DLL '%s'"), *CompilerPath);
 		}
 		Compile = (pD3DCompile)(void*)GetProcAddress(CompilerDLL, "D3DCompile");
 		Reflect = (pD3DReflect)(void*)GetProcAddress(CompilerDLL, "D3DReflect");
@@ -924,7 +922,7 @@ static bool CompileAndProcessD3DShaderFXCExt(
 			Result = D3DReflectFunc(Shader->GetBufferPointer(), Shader->GetBufferSize(), IID_ID3D11ShaderReflectionForCurrentCompiler, (void**)Reflector.GetInitReference());
 			if (FAILED(Result))
 			{
-				UE_LOG(LogD3D11ShaderCompiler, Fatal, TEXT("D3DReflect failed: Result=%08x"), Result);
+				UE_LOG(LogD3DShaderCompiler, Fatal, TEXT("D3DReflect failed: Result=%08x"), Result);
 			}
 
 			// Read the constant table description.
@@ -1009,7 +1007,7 @@ static bool CompileAndProcessD3DShaderFXCExt(
 							// second pass cannot use more attributes than previously
 							if (Output.UsedAttributes.Num() > CompileData.ShaderInputs.Num())
 							{
-								UE_LOG(LogD3D11ShaderCompiler, Warning, TEXT("Second pass had more used attributes (%d) than first pass (%d)"), Output.UsedAttributes.Num(), CompileData.ShaderInputs.Num());
+								UE_LOG(LogD3DShaderCompiler, Warning, TEXT("Second pass had more used attributes (%d) than first pass (%d)"), Output.UsedAttributes.Num(), CompileData.ShaderInputs.Num());
 								FShaderCompilerError NewError;
 								NewError.StrippedErrorMessage = FString::Printf(TEXT("Second pass had more used attributes (%d) than first pass (%d)"), Output.UsedAttributes.Num(), CompileData.ShaderInputs.Num());
 								Output = OriginalOutput;
@@ -1020,7 +1018,7 @@ static bool CompileAndProcessD3DShaderFXCExt(
 							// if we're about to run out of attempts, report
 							if (Attempt >= kMaxReasonableAttempts - 1)
 							{
-								UE_LOG(LogD3D11ShaderCompiler, Warning, TEXT("Unable to determine unused inputs after %d attempts (last number of used attributes: %d, previous step:%d)!"), 
+								UE_LOG(LogD3DShaderCompiler, Warning, TEXT("Unable to determine unused inputs after %d attempts (last number of used attributes: %d, previous step:%d)!"),
 									Attempt + 1,
 									Output.UsedAttributes.Num(),
 									CompileData.ShaderInputs.Num()
@@ -1041,11 +1039,11 @@ static bool CompileAndProcessD3DShaderFXCExt(
 						}
 						else
 						{
-							UE_LOG(LogD3D11ShaderCompiler, Warning, TEXT("Failed to remove unused inputs from shader: %s"), *Input.GenerateShaderName());
+							UE_LOG(LogD3DShaderCompiler, Warning, TEXT("Failed to remove unused inputs from shader: %s"), *Input.GenerateShaderName());
 							for (const FString& ErrorMessage : RemoveErrors)
 							{
 								// Add error to shader output but also make sure the error shows up on build farm by emitting a log entry
-								UE_LOG(LogD3D11ShaderCompiler, Warning, TEXT("%s"), *ErrorMessage);
+								UE_LOG(LogD3DShaderCompiler, Warning, TEXT("%s"), *ErrorMessage);
 								FShaderCompilerError NewError;
 								NewError.StrippedErrorMessage = ErrorMessage;
 								Output.Errors.Add(NewError);
@@ -1113,7 +1111,7 @@ static bool CompileAndProcessD3DShaderFXCExt(
 
 				if (FAILED(Result))
 				{
-					UE_LOG(LogD3D11ShaderCompiler, Fatal, TEXT("D3DStripShader failed: Result=%08x"), Result);
+					UE_LOG(LogD3DShaderCompiler, Fatal, TEXT("D3DStripShader failed: Result=%08x"), Result);
 				}
 			}
 			else
@@ -1322,11 +1320,11 @@ void CompileD3DShader(const FShaderCompilerInput& Input, const FShaderPreprocess
 		TArray<FString> Errors;
 		if (!RemoveUnusedOutputs(PreprocessedSource, UsedOutputs, Exceptions, ScopedDeclarations, EntryPointName, Errors))
 		{
-			UE_LOG(LogD3D11ShaderCompiler, Warning, TEXT("Failed to remove unused outputs from shader: %s"), *Input.GenerateShaderName());
+			UE_LOG(LogD3DShaderCompiler, Warning, TEXT("Failed to remove unused outputs from shader: %s"), *Input.GenerateShaderName());
 			for (const FString& ErrorReport : Errors)
 			{
 				// Add error to shader output but also make sure the error shows up on build farm by emitting a log entry
-				UE_LOG(LogD3D11ShaderCompiler, Warning, TEXT("%s"), *ErrorReport);
+				UE_LOG(LogD3DShaderCompiler, Warning, TEXT("%s"), *ErrorReport);
 				FShaderCompilerError NewError;
 				NewError.StrippedErrorMessage = ErrorReport;
 				Output.Errors.Add(NewError);
diff --git a/Engine/Source/Developer/Windows/ShaderFormatD3D/Private/D3DShaderCompilerDXC.cpp b/Engine/Source/Developer/Windows/ShaderFormatD3D/Private/D3DShaderCompilerDXC.cpp
--- a/Engine/Source/Developer/Windows/ShaderFormatD3D/Private/D3DShaderCompilerDXC.cpp
+++ b/Engine/Source/Developer/Windows/ShaderFormatD3D/Private/D3DShaderCompilerDXC.cpp
@@ -12,8 +12,6 @@
 #include "ShaderPreprocessTypes.h"
 #include "RayTracingDefinitions.h"
 
-DEFINE_LOG_CATEGORY_STATIC(LogD3D12ShaderCompiler, Log, All);
-
 // D3D doesn't define a mask for this, so we do so here
 #define SHADER_OPTIMIZATION_LEVEL_MASK (D3DCOMPILE_OPTIMIZATION_LEVEL0 | D3DCOMPILE_OPTIMIZATION_LEVEL1 | D3DCOMPILE_OPTIMIZATION_LEVEL2 | D3DCOMPILE_OPTIMIZATION_LEVEL3)
 
@@ -134,18 +132,18 @@ static void LogFailedHRESULT(const TCHAR* FailedExpressionStr, HRESULT Result)
 	{
 		const FString ErrorReport = FString::Printf(TEXT("%s failed: Result=0x%08x (E_OUTOFMEMORY)"), FailedExpressionStr, Result);
 		FSCWErrorCode::Report(FSCWErrorCode::OutOfMemory, ErrorReport);
-		UE_LOG(LogD3D12ShaderCompiler, Fatal, TEXT("%s"), *ErrorReport);
+		UE_LOG(LogD3DShaderCompiler, Fatal, TEXT("%s"), *ErrorReport);
 	}
 	else if (const TCHAR* ErrorCodeStr = DxcErrorCodeToString(Result))
 	{
-		UE_LOG(LogD3D12ShaderCompiler, Fatal, TEXT("%s failed: Result=0x%08x (%s)"), FailedExpressionStr, Result, ErrorCodeStr);
+		UE_LOG(LogD3DShaderCompiler, Fatal, TEXT("%s failed: Result=0x%08x (%s)"), FailedExpressionStr, Result, ErrorCodeStr);
 	}
 	else
 	{
 		// Turn HRESULT into human readable string for error report
 		TCHAR ResultStr[4096] = {};
 		FPlatformMisc::GetSystemErrorMessage(ResultStr, UE_ARRAY_COUNT(ResultStr), Result);
-		UE_LOG(LogD3D12ShaderCompiler, Fatal, TEXT("%s failed: Result=0x%08x (%s)"), FailedExpressionStr, Result, ResultStr);
+		UE_LOG(LogD3DShaderCompiler, Fatal, TEXT("%s failed: Result=0x%08x (%s)"), FailedExpressionStr, Result, ResultStr);
 	}
 }
 
@@ -1145,7 +1143,7 @@ bool CompileAndProcessD3DShaderDXC(
 			}
 			else
 			{
-				UE_LOG(LogD3D12ShaderCompiler, Fatal, TEXT("Failed to find required points in the shader library."));
+				UE_LOG(LogD3DShaderCompiler, Fatal, TEXT("Failed to find required points in the shader library."));
 				Output.bSucceeded = false;
 			}
 		}
diff --git a/Engine/Source/Developer/Windows/ShaderFormatD3D/Private/ShaderFormatD3D.cpp b/Engine/Source/Developer/Windows/ShaderFormatD3D/Private/ShaderFormatD3D.cpp
--- a/Engine/Source/Developer/Windows/ShaderFormatD3D/Private/ShaderFormatD3D.cpp
+++ b/Engine/Source/Developer/Windows/ShaderFormatD3D/Private/ShaderFormatD3D.cpp
@@ -11,6 +11,8 @@
 #include "Interfaces/IShaderFormatModule.h"
 #include "DXCWrapper.h"
 
+DEFINE_LOG_CATEGORY(LogD3DShaderCompiler);
+
 static FName NAME_PCD3D_SM6(TEXT("PCD3D_SM6"));
 static FName NAME_PCD3D_SM5(TEXT("PCD3D_SM5"));
 static FName NAME_PCD3D_ES3_1(TEXT("PCD3D_ES31"));
diff --git a/Engine/Source/Developer/Windows/ShaderFormatD3D/Private/ShaderFormatD3D.h b/Engine/Source/Developer/Windows/ShaderFormatD3D/Private/ShaderFormatD3D.h
--- a/Engine/Source/Developer/Windows/ShaderFormatD3D/Private/ShaderFormatD3D.h
+++ b/Engine/Source/Developer/Windows/ShaderFormatD3D/Private/ShaderFormatD3D.h
@@ -5,6 +5,8 @@
 #include "Windows/WindowsHWrapper.h"
 #include "ShaderCompilerCommon.h"
 
+DECLARE_LOG_CATEGORY_EXTERN(LogD3DShaderCompiler, Log, All)
+
 struct FShaderTarget;
 
 enum class ED3DShaderModel
