PR #11437: LiveCoding: Fix crash patching headless editor instances

diff --git a/Engine/Source/Developer/Windows/LiveCoding/Private/LiveCodingModule.cpp b/Engine/Source/Developer/Windows/LiveCoding/Private/LiveCodingModule.cpp
--- a/Engine/Source/Developer/Windows/LiveCoding/Private/LiveCodingModule.cpp
+++ b/Engine/Source/Developer/Windows/LiveCoding/Private/LiveCodingModule.cpp
@@ -1026,8 +1026,10 @@ void FLiveCodingModule::ShowNotification(bool Success, const FText& Title, const
 	{
 		Info.SubText = *SubText;
 	}
-	TSharedPtr<SNotificationItem> CompileNotification = FSlateNotificationManager::Get().AddNotification(Info);
-	CompileNotification->SetCompletionState(Success ? SNotificationItem::CS_Success : SNotificationItem::CS_Fail);
+	if (TSharedPtr<SNotificationItem> CompileNotification = FSlateNotificationManager::Get().AddNotification(Info))
+	{
+		CompileNotification->SetCompletionState(Success ? SNotificationItem::CS_Success : SNotificationItem::CS_Fail);
+	}
 }
 #endif
 
