Fixed non unity compile errors

diff --git a/Engine/Source/Developer/Windows/LiveCoding/Private/External/LC_FixedSizeString.cpp b/Engine/Source/Developer/Windows/LiveCoding/Private/External/LC_FixedSizeString.cpp
--- a/Engine/Source/Developer/Windows/LiveCoding/Private/External/LC_FixedSizeString.cpp
+++ b/Engine/Source/Developer/Windows/LiveCoding/Private/External/LC_FixedSizeString.cpp
@@ -8,7 +8,9 @@
 
 // BEGIN EPIC MOD
 // This is really treated as an include file.
-#pragma once 
+//#pragma once 
+#ifndef LC_FIXEDSIZESTRING_CPP
+#define LC_FIXEDSIZESTRING_CPP
 // END EPIC MOD
 
 template <typename T, unsigned int N>
@@ -119,3 +121,5 @@ FixedSizeString<T, N> FixedSizeString<T, N>::ToUpper(void) const LC_RESTRICT
 
 	return result;
 }
+
+#endif // LC_FIXEDSIZESTRING_CPP
\ No newline at end of file
