Added documentation to the PreloadNamedModules live coding setting

#rnx

diff --git a/Engine/Source/Developer/Windows/LiveCoding/Private/LiveCodingSettings.h b/Engine/Source/Developer/Windows/LiveCoding/Private/LiveCodingSettings.h
--- a/Engine/Source/Developer/Windows/LiveCoding/Private/LiveCodingSettings.h
+++ b/Engine/Source/Developer/Windows/LiveCoding/Private/LiveCodingSettings.h
@@ -44,8 +44,8 @@ public:
     UPROPERTY(config, EditAnywhere, Category=Modules, Meta=(ConfigRestartRequired=true, EditCondition="bEnabled"))
     bool bPreloadProjectPluginModules;
 
-    UPROPERTY(config, EditAnywhere, Category=Modules, Meta=(ConfigRestartRequired=true, EditCondition="bEnabled"))
-    TArray<FName> PreloadNamedModules;
+    UPROPERTY(config, EditAnywhere, Category=Modules, Meta=(ConfigRestartRequired=true, EditCondition="bEnabled", DisplayName="Names of modules, without target prefix, to always preload.  Wildcards are supported"))
+    TArray<FName> PreloadNamedModules; 
 
 	ULiveCodingSettings(const FObjectInitializer& Initializer);
 };
