Dump internal FXC compiler warnings when FXC fails to compile but DXC manages to, so that the user can know what failed and eventually correct it or at least make an informed decision about whether to bypass the error using DXC pre-compilation

diff --git a/Engine/Source/Developer/Windows/ShaderFormatD3D/Private/D3DShaderCompiler.cpp b/Engine/Source/Developer/Windows/ShaderFormatD3D/Private/D3DShaderCompiler.cpp
--- a/Engine/Source/Developer/Windows/ShaderFormatD3D/Private/D3DShaderCompiler.cpp
+++ b/Engine/Source/Developer/Windows/ShaderFormatD3D/Private/D3DShaderCompiler.cpp
@@ -13,6 +13,7 @@
 #include "ShaderParameterParser.h"
 #include "ShaderPreprocessTypes.h"
 #include "SpirvCommon.h"
+#include "Algo/Transform.h"
 
 #define DEBUG_SHADERS 0
 
@@ -685,6 +686,7 @@ static bool CompileAndProcessD3DShaderFXCExt(
 
 	if (D3DCompileFunc)
 	{
+		TArray<FString> InitialFXCRunFilteredErrors;
 		const bool bHlslVersion2021 = Input.Environment.CompilerFlags.Contains(CFLAG_HLSL2021);
 		const bool bPrecompileWithDXC = bHlslVersion2021 || Input.Environment.CompilerFlags.Contains(CFLAG_PrecompileWithDXC);
 		if (!bPrecompileWithDXC)
@@ -707,6 +709,15 @@ static bool CompileAndProcessD3DShaderFXCExt(
 				// ShaderCompileWorker handle the exception and log an error.
 				/* bCatchException */ true
 			);
+
+			if (Result == E_FAIL)
+			{
+				// We might have failed compiling with FXC but then we might actually manage to compile through DXC
+				if (void* ErrorBuffer = Errors ? Errors->GetBufferPointer() : nullptr)
+				{
+					D3D11FilterShaderCompileWarnings(ANSI_TO_TCHAR(ErrorBuffer), InitialFXCRunFilteredErrors);
+				}
+			}
 		}
 
 		// Some materials give FXC a hard time to optimize and the compiler fails with an internal error.
@@ -840,6 +851,9 @@ static bool CompileAndProcessD3DShaderFXCExt(
 
 				// Let the user know this shader had to be cross-compiled due to a crash in FXC. Only shows up if CVar 'r.ShaderDevelopmentMode' is enabled.
 				Output.Errors.Add(FShaderCompilerError(TEXT("Cross-compiled shader to intermediate HLSL after first attempt crashed FXC")));
+
+				// Output the errors from the initial run so that the user can know what failed and eventually correct it or at least make an informed decision about whether to bypass the error using DXC pre-compilation:
+				Algo::Transform(InitialFXCRunFilteredErrors, Output.Errors, [](const FString& InInitialError) { return FShaderCompilerError(*InInitialError); });
 			}
 		}
 	}
@@ -1356,4 +1370,4 @@ void CompileD3DShader(const FShaderCompilerInput& Input, const FShaderPreprocess
 			FPlatformMisc::LowLevelOutputDebugStringf(TEXT("%s\n"), *Error.GetErrorStringWithLineMarker());
 		}
 	}
-}
+}
\ No newline at end of file
