RHI Validation
- Removed reliance on ShaderParameterParser for most platforms as it doesn't always parse parameters

diff --git a/Engine/Source/Developer/Windows/ShaderFormatD3D/Private/D3DShaderCompiler.inl b/Engine/Source/Developer/Windows/ShaderFormatD3D/Private/D3DShaderCompiler.inl
--- a/Engine/Source/Developer/Windows/ShaderFormatD3D/Private/D3DShaderCompiler.inl
+++ b/Engine/Source/Developer/Windows/ShaderFormatD3D/Private/D3DShaderCompiler.inl
@@ -34,6 +34,48 @@ struct FD3DShaderCompileData
 	uint32 MaxUAVs = 0;
 };
 
+template <typename D3D1x_SHADER_INPUT_BIND_DESC>
+EShaderCodeResourceBindingType D3DBindDescToShaderCodeResourceBinding(const D3D1x_SHADER_INPUT_BIND_DESC& Binding)
+{
+	switch (Binding.Type)
+	{
+	case D3D_SIT_SAMPLER:		
+		return EShaderCodeResourceBindingType::SamplerState;
+	case D3D_SIT_TBUFFER:
+	case D3D_SIT_CBUFFER:
+		return EShaderCodeResourceBindingType::Buffer;
+	case D3D_SIT_TEXTURE:
+		switch (Binding.Dimension)
+		{
+		case D3D_SRV_DIMENSION_BUFFER:			return EShaderCodeResourceBindingType::Buffer;
+		case D3D_SRV_DIMENSION_TEXTURE2D:		return EShaderCodeResourceBindingType::Texture2D;
+		case D3D_SRV_DIMENSION_TEXTURE2DARRAY:	return EShaderCodeResourceBindingType::Texture2DArray;
+		case D3D_SRV_DIMENSION_TEXTURE2DMS:		return EShaderCodeResourceBindingType::Texture2DMS;
+		case D3D_SRV_DIMENSION_TEXTURE3D:		return EShaderCodeResourceBindingType::Texture3D;
+		case D3D_SRV_DIMENSION_TEXTURECUBE:		return EShaderCodeResourceBindingType::TextureCube;
+		default:
+			return EShaderCodeResourceBindingType::Invalid;
+		}
+	case D3D_SIT_UAV_RWTYPED:
+		switch (Binding.Dimension)
+		{
+		case D3D_SRV_DIMENSION_BUFFER:			return EShaderCodeResourceBindingType::RWBuffer;
+		case D3D_SRV_DIMENSION_TEXTURE2D:		return EShaderCodeResourceBindingType::RWTexture2D;
+		case D3D_SRV_DIMENSION_TEXTURE2DARRAY:	return EShaderCodeResourceBindingType::RWTexture2DArray;
+		case D3D_SRV_DIMENSION_TEXTURE3D:		return EShaderCodeResourceBindingType::RWTexture3D;
+		case D3D_SRV_DIMENSION_TEXTURECUBE:		return EShaderCodeResourceBindingType::RWTextureCube;
+		default:
+			return EShaderCodeResourceBindingType::Invalid;
+		}
+	case D3D_SIT_STRUCTURED:		return EShaderCodeResourceBindingType::StructuredBuffer;
+	case D3D_SIT_UAV_RWSTRUCTURED:	return EShaderCodeResourceBindingType::RWStructuredBuffer;
+	case D3D_SIT_BYTEADDRESS:		return EShaderCodeResourceBindingType::ByteAddressBuffer;
+	case D3D_SIT_UAV_RWBYTEADDRESS:	return EShaderCodeResourceBindingType::RWByteAddressBuffer;
+	default:
+		return EShaderCodeResourceBindingType::Invalid;
+	}
+}
+
 template <typename ID3D1xShaderReflection, typename D3D1x_SHADER_DESC, typename D3D1x_SHADER_INPUT_BIND_DESC,
 	typename ID3D1xShaderReflectionConstantBuffer, typename D3D1x_SHADER_BUFFER_DESC,
 	typename ID3D1xShaderReflectionVariable, typename D3D1x_SHADER_VARIABLE_DESC>
@@ -230,8 +272,8 @@ template <typename ID3D1xShaderReflection, typename D3D1x_SHADER_DESC, typename
 			}
 			else
 			{
-				const FShaderParameterParser::FParsedShaderParameter* ParsedParam = ShaderParameterParser.FindParameterInfosUnsafe(BindDesc.Name);
-				AddShaderValidationSRVType(BindDesc.BindPoint, ParsedParam ? ParsedParam->ParsedTypeDecl : EShaderCodeResourceBindingType::Invalid, Output);
+				EShaderCodeResourceBindingType ResourceBindingType = D3DBindDescToShaderCodeResourceBinding(BindDesc);
+				AddShaderValidationSRVType(BindDesc.BindPoint, ResourceBindingType, Output);
 
 				HandleReflectedShaderResource(FString(BindDesc.Name), BindDesc.BindPoint, Output);
 				CompileData.NumSRVs = FMath::Max(CompileData.NumSRVs, BindDesc.BindPoint + BindCount);
@@ -276,8 +318,8 @@ template <typename ID3D1xShaderReflection, typename D3D1x_SHADER_DESC, typename
 			}
 			else
 			{
-				const FShaderParameterParser::FParsedShaderParameter* ParsedParam = ShaderParameterParser.FindParameterInfosUnsafe(BindDesc.Name);
-				AddShaderValidationUAVType(BindDesc.BindPoint, ParsedParam ? ParsedParam->ParsedTypeDecl : EShaderCodeResourceBindingType::Invalid, Output);
+				EShaderCodeResourceBindingType ResourceBindingType = D3DBindDescToShaderCodeResourceBinding(BindDesc);
+				AddShaderValidationUAVType(BindDesc.BindPoint, ResourceBindingType, Output);
 
 				HandleReflectedShaderUAV(FString(BindDesc.Name), BindDesc.BindPoint, Output);
 				CompileData.NumUAVs = FMath::Max(CompileData.NumUAVs, BindDesc.BindPoint + BindCount);
@@ -287,8 +329,9 @@ template <typename ID3D1xShaderReflection, typename D3D1x_SHADER_DESC, typename
 		{
 			check(BindDesc.BindCount == 1);
 			FString BindDescName(BindDesc.Name);
-			const FShaderParameterParser::FParsedShaderParameter* ParsedParam = ShaderParameterParser.FindParameterInfosUnsafe(BindDesc.Name);
-			AddShaderValidationSRVType(BindDesc.BindPoint, ParsedParam ? ParsedParam->ParsedTypeDecl : EShaderCodeResourceBindingType::Invalid, Output);
+
+			EShaderCodeResourceBindingType ResourceBindingType = D3DBindDescToShaderCodeResourceBinding(BindDesc);
+			AddShaderValidationSRVType(BindDesc.BindPoint, ResourceBindingType, Output);
 
 			HandleReflectedShaderResource(BindDescName, BindDesc.BindPoint, Output);
 
@@ -306,8 +349,8 @@ template <typename ID3D1xShaderReflection, typename D3D1x_SHADER_DESC, typename
 			// Acceleration structure resources are treated as SRVs.
 			check(BindDesc.BindCount == 1);
 
-			const FShaderParameterParser::FParsedShaderParameter* ParsedParam = ShaderParameterParser.FindParameterInfosUnsafe(BindDesc.Name);
-			AddShaderValidationSRVType(BindDesc.BindPoint, ParsedParam ? ParsedParam->ParsedTypeDecl : EShaderCodeResourceBindingType::Invalid, Output);
+			EShaderCodeResourceBindingType ResourceBindingType = D3DBindDescToShaderCodeResourceBinding(BindDesc);
+			AddShaderValidationSRVType(BindDesc.BindPoint, ResourceBindingType, Output);
 
 			HandleReflectedShaderResource(FString(BindDesc.Name), BindDesc.BindPoint, Output);
 			CompileData.NumSRVs = FMath::Max(CompileData.NumSRVs, BindDesc.BindPoint + 1);
diff --git a/Engine/Source/Runtime/RenderCore/Private/ShaderParameterParser.cpp b/Engine/Source/Runtime/RenderCore/Private/ShaderParameterParser.cpp
--- a/Engine/Source/Runtime/RenderCore/Private/ShaderParameterParser.cpp
+++ b/Engine/Source/Runtime/RenderCore/Private/ShaderParameterParser.cpp
@@ -1158,7 +1158,7 @@ bool FShaderParameterParser::ParseAndModify(const FShaderCompilerInput& Compiler
 
 	const bool bUseStableConstantBuffer = EnumHasAnyFlags(PlatformConfiguration.Flags, EShaderParameterParserConfigurationFlags::UseStableConstantBuffer);
 	const bool bSupportsBindless = EnumHasAnyFlags(PlatformConfiguration.Flags, EShaderParameterParserConfigurationFlags::SupportsBindless);
-
+	const bool bAlwaysParseParams = EnumHasAnyFlags(PlatformConfiguration.Flags, EShaderParameterParserConfigurationFlags::AlwaysParseParams);
 	const bool bHasRootParameters = (CompilerInput.RootParametersStructure != nullptr);
 	const bool bRootParametersModification = bUseStableConstantBuffer && (CompilerInput.IsRayTracingShader() || CompilerInput.ShouldUseStableConstantBuffer());
 	const bool bBindlessModifications = bSupportsBindless && (bBindlessResources || bBindlessSamplers);
@@ -1166,7 +1166,7 @@ bool FShaderParameterParser::ParseAndModify(const FShaderCompilerInput& Compiler
 	const bool bShouldModify = bRootParametersModification || bBindlessModifications;
 
 	// Always parse if we have root parameters since we need that data during reflection validation
-	const bool bShouldParse = bHasRootParameters || bShouldModify;
+	const bool bShouldParse = bHasRootParameters || bShouldModify || bAlwaysParseParams;
 
 	// The shader doesn't have any parameter binding through shader structure, therefore don't do anything.
 	if (!bShouldParse)
diff --git a/Engine/Source/Runtime/RenderCore/Public/ShaderParameterParser.h b/Engine/Source/Runtime/RenderCore/Public/ShaderParameterParser.h
--- a/Engine/Source/Runtime/RenderCore/Public/ShaderParameterParser.h
+++ b/Engine/Source/Runtime/RenderCore/Public/ShaderParameterParser.h
@@ -29,6 +29,8 @@ enum class EShaderParameterParserConfigurationFlags
 	SupportsBindless        = 1 << 1,
 	// "Vulkan" style
 	BindlessUsesArrays      = 1 << 2,
+	// Some RHI Validation requires parameters to be parsed
+	AlwaysParseParams		= 1 << 3,
 };
 ENUM_CLASS_FLAGS(EShaderParameterParserConfigurationFlags)
 
