Fix reflection data unnecessarily kept in compiled PCD3D_SM6 shaders.

- Found while investigating compressed size (by FabianG), seemst to have been that way since 22233889. Shader size down significanty, to about  30-40% of the previous (both disk size and the runtime footprint).

diff --git a/Engine/Source/Developer/Windows/ShaderFormatD3D/Private/D3DShaderCompilerDXC.cpp b/Engine/Source/Developer/Windows/ShaderFormatD3D/Private/D3DShaderCompilerDXC.cpp
--- a/Engine/Source/Developer/Windows/ShaderFormatD3D/Private/D3DShaderCompilerDXC.cpp
+++ b/Engine/Source/Developer/Windows/ShaderFormatD3D/Private/D3DShaderCompilerDXC.cpp
@@ -739,15 +739,18 @@ static HRESULT D3DCompileToDxil(const char* SourceText, const FDxcArguments& Arg
 		VERIFYHRESULT(CompileResult->GetOutput(DXC_OUT_REFLECTION, IID_PPV_ARGS(OutReflectionBlob.GetInitReference()), ReflectionNameBlob.GetInitReference()));
 
 		RetrieveDebugNameAndBlob(CompileResult, OutPdbName, OutPdbBlob.GetInitReference());
-		const bool bHasOutputPDB = OutPdbBlob.IsValid() && !OutPdbName.IsEmpty();
-		const bool bRemovePDB = bHasOutputPDB && !Arguments.ShouldKeepEmbeddedPDB();
-
-		TArray<uint32, TInlineAllocator<4>> PartsToRemove;
-		if (bRemovePDB)
+ 
+ 		TArray<uint32, TInlineAllocator<4>> PartsToRemove;
+		if (!Arguments.ShouldKeepEmbeddedPDB())
 		{
-			// Try and remove both the PDB & Reflection Data
-			PartsToRemove.Add(DXC_PART_PDB);
-			PartsToRemove.Add(DXC_PART_REFLECTION_DATA);
+			if (CompileResult->HasOutput(DXC_OUT_PDB))
+			{
+				PartsToRemove.Add(DXC_PART_PDB);
+			}
+			if (CompileResult->HasOutput(DXC_OUT_REFLECTION))
+			{
+				PartsToRemove.Add(DXC_PART_REFLECTION_DATA);
+			}
 		}
 
 		if (Arguments.ShouldDump())
@@ -762,7 +765,7 @@ static HRESULT D3DCompileToDxil(const char* SourceText, const FDxcArguments& Arg
 			SaveDxcBlobToFile(OutDxilBlob, DxilFile);
 
 			// Dump the PDB.
-			if (bHasOutputPDB)
+			if (OutPdbBlob.IsValid() && !OutPdbName.IsEmpty())
 			{
 				const FString PdbFile = Arguments.GetDumpDebugInfoPath() / OutPdbName;
 				SaveDxcBlobToFile(OutPdbBlob, PdbFile);
diff --git a/Engine/Source/Developer/Windows/ShaderFormatD3D/Private/ShaderFormatD3D.cpp b/Engine/Source/Developer/Windows/ShaderFormatD3D/Private/ShaderFormatD3D.cpp
--- a/Engine/Source/Developer/Windows/ShaderFormatD3D/Private/ShaderFormatD3D.cpp
+++ b/Engine/Source/Developer/Windows/ShaderFormatD3D/Private/ShaderFormatD3D.cpp
@@ -17,7 +17,7 @@ static FName NAME_PCD3D_SM6(TEXT("PCD3D_SM6"));
 static FName NAME_PCD3D_SM5(TEXT("PCD3D_SM5"));
 static FName NAME_PCD3D_ES3_1(TEXT("PCD3D_ES31"));
 
-static const FGuid UE_SHADER_PCD3D_SHARED_VER = FGuid("232A2A59-A6D0-4CDB-A374-F3DB028E413E");
+static const FGuid UE_SHADER_PCD3D_SHARED_VER = FGuid("6A25E42A-3266-4692-81C4-35B5690B7578");
 static const FGuid UE_SHADER_PCD3D_SM6_VER    = FGuid("FC317769-CFF7-4007-A045-C1D4CD315BD4");
 static const FGuid UE_SHADER_PCD3D_SM5_VER    = FGuid("5B377D13-C70F-40C5-80C5-C9B228783469");
 static const FGuid UE_SHADER_PCD3D_ES3_1_VER  = FGuid("952939D9-1156-4347-97E9-9FFEA1A9FE14");
