[Shaders] Fixed re-ordering of shader interpolators in D3D ES31 backend.

When cross-compiling HLSL 2021 to HLSL 2018 for D3D11, DXC/SPIRV-Cross re-arrange the vertex/pixel interpolators in a way that can cause mismatches in the PSO.
The previous method to re-order shader interpolators was to just search the source for their variable declaration semantics, but that search can confuse semantics like TEXCOORD1 with TEXCOORD10.

This change includes two fixes:
1. The new approach parses the semantics and match the entire name, not just searching for the presence of the respective semantic identifier.
2. Always declare system value semantics such as SV_Position after user-defined semantics in case they are not used in the pixel shader to preserve the DXBC register assignments.

This fixes issues with Substrate as we have to cross-compile lots of shaders for D3D11 and some ES31 preview platforms.

#rnx

diff --git a/Engine/Source/Developer/Windows/ShaderFormatD3D/Private/D3DShaderCompiler.cpp b/Engine/Source/Developer/Windows/ShaderFormatD3D/Private/D3DShaderCompiler.cpp
--- a/Engine/Source/Developer/Windows/ShaderFormatD3D/Private/D3DShaderCompiler.cpp
+++ b/Engine/Source/Developer/Windows/ShaderFormatD3D/Private/D3DShaderCompiler.cpp
@@ -457,18 +457,17 @@ static void PatchSpirvForPrecompilation(FSpirv& Spirv)
 	}
 }
 
-// @param StageVariablesStorageClass Must be SpvStorageClassOutput for vertex shaders and SpvStorageClassInput for pixel shaders.
+// Re-orders all input/ouput stage variables from the specified HLSL source if it was cross-compiled with SPIRV-Cross.
+// SPIRV-Cross can arrange the stage variables in a way that causes a mismatch between vertex and pixel shader pipelines.
 static bool PatchHlslWithReorderedIOVariables(
 	FString& HlslSourceString,
 	const FString& OriginalShaderSource,
 	const FString& OriginalEntryPoint,
-	SpvStorageClass StageVariablesStorageClass,
+	EShaderParameterStorageClass StageVariablesStorageClass,
 	TArray<FShaderCompilerError>& OutErrors)
 {
-	check(StageVariablesStorageClass == SpvStorageClassInput || StageVariablesStorageClass == SpvStorageClassOutput);
-
 	// Find declaration struct for stage variables
-	const FStringView StageVariableDeclarationName = (StageVariablesStorageClass == SpvStorageClassInput ? TEXT("SPIRV_Cross_Input") : TEXT("SPIRV_Cross_Output"));
+	const FStringView StageVariableDeclarationName = (StageVariablesStorageClass == EShaderParameterStorageClass::Input ? TEXTVIEW("SPIRV_Cross_Input") : TEXTVIEW("SPIRV_Cross_Output"));
 	const int32 StageVariableDeclarationBegin = HlslSourceString.Find(StageVariableDeclarationName, ESearchCase::CaseSensitive);
 	if (StageVariableDeclarationBegin == INDEX_NONE)
 	{
@@ -487,16 +486,9 @@ static bool PatchHlslWithReorderedIOVariables(
 		return false;
 	}
 
-	// Parse declaration struct for stage variables into array of individual lines
-	const FString StageVariableDeclarationSource = HlslSourceString.Mid(StageVariableDelcarationBlockBegin + 1, StageVariableDelcarationBlockEnd - (StageVariableDelcarationBlockBegin + 1));
-
-	TArray<FString> StageVariableDeclarationLines;
-	StageVariableDeclarationSource.ParseIntoArrayLines(StageVariableDeclarationLines);
-
 	// Parse variable names from SPIR-V input
 	TArray<FString> Variables, ParsingErrors;
-	const EShaderParameterStorageClass ParameterStorageClass = (StageVariablesStorageClass == SpvStorageClassOutput ? EShaderParameterStorageClass::Output : EShaderParameterStorageClass::Input);
-	if (!FindEntryPointParameters(OriginalShaderSource, OriginalEntryPoint, ParameterStorageClass, Variables, ParsingErrors))
+	if (!FindEntryPointParameters(OriginalShaderSource, OriginalEntryPoint, StageVariablesStorageClass, Variables, ParsingErrors))
 	{
 		for (FString& Error : ParsingErrors)
 		{
@@ -505,35 +497,106 @@ static bool PatchHlslWithReorderedIOVariables(
 		return false;
 	}
 
-	if (Variables.Num() != StageVariableDeclarationLines.Num())
+	// Parse declaration struct for stage variables into array of individual lines
+	const FString StageVariableDeclSource = HlslSourceString.Mid(StageVariableDelcarationBlockBegin + 1, StageVariableDelcarationBlockEnd - (StageVariableDelcarationBlockBegin + 1));
+
+	TArray<FString> StageVariableDeclSourceLines;
+	StageVariableDeclSource.ParseIntoArrayLines(StageVariableDeclSourceLines);
+
+	if (Variables.Num() != StageVariableDeclSourceLines.Num())
 	{
 		// Failed to match SPIR-V variables to SPIRV-Cross generated source
 		return false;
 	}
 
-	// Re-arrange source lines of stage variable declarations
-	FString SortedStageVariableDeclarationSource = TEXT("\n");
+	// Returns true if the specified source line contains a variable declaration with the specified semantic.
+	// HLSL semantics are case insensitive, must always appear after a colon ':' and when declared in a structure end with a semicolon ';'.
+	// Here are some examples of such variable declarations, generated by SPIRV-Cross, we are parsing:
+	//          Variable             Semantic
+	//  ------------------------------------------
+	//  float4 out_var_TEXCOORD10 : TEXCOORD10;
+	//  float4 out_var_TEXCOORD1[1] : TEXCOORD1;
+	//  precise float4 gl_Position : SV_Position;
+	auto FindSemanticDeclarationInSourceLine = [](const FString& SourceLine, const FString& SemanticToSearch) -> bool
+	{
+		// Exit early if source line is not long enough for semantic
+		if (SemanticToSearch.Len() > SourceLine.Len())
+		{
+			return false;
+		}
+
+		// Scan souce line for semantic range starting with first ':' character
+		int32 StartPosition = 0;
+		if (!SourceLine.FindChar(TEXT(':'), StartPosition))
+		{
+			return false;
+		}
 
-	for (const FString& Variable : Variables)
+		// Progress until we have no more whitespaces, then we found the final start position of the semantic
+		do
+		{
+			++StartPosition;
+		}
+		while (StartPosition < SourceLine.Len() && FChar::IsWhitespace(SourceLine[StartPosition]));
+
+		// Exit early if remainder of source line is not long enough for semantic
+		if (SemanticToSearch.Len() + StartPosition > SourceLine.Len())
+		{
+			return false;
+		}
+
+		// Now find the first whitespace character or the statement terminator ';'
+		int32 EndPosition = StartPosition;
+		while (EndPosition < SourceLine.Len() && !(FChar::IsWhitespace(SourceLine[EndPosition]) || SourceLine[EndPosition] == TEXT(';')))
+		{
+			++EndPosition;
+		}
+
+		// Now compare the semantic to search for with the source line range
+		const int32 SemanticLen = EndPosition - StartPosition;
+		return SemanticToSearch.Len() == SemanticLen && FCString::Strnicmp(*SemanticToSearch, &SourceLine[StartPosition], SemanticLen) == 0;
+	};
+
+	// Returns true if the specified semantic name starts with "SV_" (case insensitive).
+	auto IsSemanticSystemValue = [](const FString& SemanticName) -> bool
+	{
+		return SemanticName.StartsWith(TEXT("SV_"), 3, ESearchCase::IgnoreCase);
+	};
+
+	auto BuildSortedStageVariableDeclSource = [&IsSemanticSystemValue, &FindSemanticDeclarationInSourceLine](
+		FString& OutStageVariableDeclSource, TArray<FString>& StageVariableDeclLines, const TArray<FString>& InVariables, bool bSystemValueSemanticsOnly)
 	{
-		for (FString& SourceLine : StageVariableDeclarationLines)
+		for (const FString& Variable : InVariables)
 		{
-			// Search for semantic name (always case insensitive) in current stage variable source line
-			if (SourceLine.Find(Variable, ESearchCase::IgnoreCase) != INDEX_NONE)
+			// Are we only emitting source for system value semantics right now?
+			if (IsSemanticSystemValue(Variable) == bSystemValueSemanticsOnly)
 			{
-				// Append source line for current variable at the end of sorted declaration string.
-				// Then empty this source line to avoid unnecessary string comparisons for next variables.
-				SortedStageVariableDeclarationSource += SourceLine;
-				SortedStageVariableDeclarationSource += TEXT('\n');
-				SourceLine.Empty();
-				break;
+				for (FString& SourceLine : StageVariableDeclLines)
+				{
+					// Search for semantic name (always case insensitive) in current stage variable source line
+					if (FindSemanticDeclarationInSourceLine(SourceLine, Variable))
+					{
+						// Append source line for current variable at the end of sorted declaration string.
+						// Then empty this source line to avoid unnecessary string comparisons for next variables.
+						OutStageVariableDeclSource += SourceLine;
+						OutStageVariableDeclSource += TEXT('\n');
+						SourceLine.Empty();
+						break;
+					}
+				}
 			}
 		}
-	}
+	};
+
+	// Re-arrange source lines of stage variable declarations and always emit system values last
+	FString SortedStageVariableDeclSource = TEXT("\n");
+
+	BuildSortedStageVariableDeclSource(SortedStageVariableDeclSource, StageVariableDeclSourceLines, Variables, /*bSystemValueSemanticsOnly:*/ false);
+	BuildSortedStageVariableDeclSource(SortedStageVariableDeclSource, StageVariableDeclSourceLines, Variables, /*bSystemValueSemanticsOnly:*/ true);
 
 	// Replace old declaration with sorted one
 	HlslSourceString.RemoveAt(StageVariableDelcarationBlockBegin + 1, StageVariableDelcarationBlockEnd - (StageVariableDelcarationBlockBegin + 1));
-	HlslSourceString.InsertAt(StageVariableDelcarationBlockBegin + 1, SortedStageVariableDeclarationSource);
+	HlslSourceString.InsertAt(StageVariableDelcarationBlockBegin + 1, SortedStageVariableDeclSource);
 
 	return true;
 }
@@ -589,7 +652,7 @@ static void PatchHlslForPrecompilation(
 	if (Frequency == SF_Vertex)
 	{
 		// Ensure order of output variables remains the same as declared in original shader source
-		PatchHlslWithReorderedIOVariables(HlslSourceString, OriginalShaderSource, OriginalEntryPoint, SpvStorageClassOutput, OutErrors);
+		PatchHlslWithReorderedIOVariables(HlslSourceString, OriginalShaderSource, OriginalEntryPoint, EShaderParameterStorageClass::Output, OutErrors);
 	}
 	else if (Frequency == SF_Pixel)
 	{
@@ -609,7 +672,7 @@ static void PatchHlslForPrecompilation(
 		}
 
 		// Ensure order of input variables remains the same as declared in original shader source
-		PatchHlslWithReorderedIOVariables(HlslSourceString, OriginalShaderSource, OriginalEntryPoint, SpvStorageClassInput, OutErrors);
+		PatchHlslWithReorderedIOVariables(HlslSourceString, OriginalShaderSource, OriginalEntryPoint, EShaderParameterStorageClass::Input, OutErrors);
 	}
 
 	// Return new HLSL source
