[Shaders] Add special case to parsing HLSL semantics in cross-compilation pipeline.

There are special cases for semantics when the index is omitted and defaults to zero, e.g. "SV_ClipDistance0" or "SV_ClipDistance[0]".
In those cases, our HlslParser does not provide the semantic index so we have to handle it when re-ordering the interpolators in the DXBC backend.

#rnx

diff --git a/Engine/Source/Developer/Windows/ShaderFormatD3D/Private/D3DShaderCompiler.cpp b/Engine/Source/Developer/Windows/ShaderFormatD3D/Private/D3DShaderCompiler.cpp
--- a/Engine/Source/Developer/Windows/ShaderFormatD3D/Private/D3DShaderCompiler.cpp
+++ b/Engine/Source/Developer/Windows/ShaderFormatD3D/Private/D3DShaderCompiler.cpp
@@ -554,7 +554,24 @@ static bool PatchHlslWithReorderedIOVariables(
 
 		// Now compare the semantic to search for with the source line range
 		const int32 SemanticLen = EndPosition - StartPosition;
-		return SemanticToSearch.Len() == SemanticLen && FCString::Strnicmp(*SemanticToSearch, &SourceLine[StartPosition], SemanticLen) == 0;
+		if (SemanticToSearch.Len() == SemanticLen && FCString::Strnicmp(*SemanticToSearch, &SourceLine[StartPosition], SemanticLen) == 0)
+		{
+			return true;
+		}
+
+		// Check for special case if semantic contains default index in source line, e.g. "SV_ClipDinstance0"
+		if (SemanticToSearch.Len() == SemanticLen - 1 && SourceLine[StartPosition + SemanticLen - 1] == TEXT('0') && FCString::Strnicmp(*SemanticToSearch, &SourceLine[StartPosition], SemanticLen - 1) == 0)
+		{
+			return true;
+		}
+
+		// Check for special case if semantic contains default index as subscript in source line, e.g. "SV_ClipDinstance[0]"
+		if (SemanticToSearch.Len() == SemanticLen - 3 && FCString::Strncmp(TEXT("[0]"), &SourceLine[StartPosition + SemanticLen - 3], 3) == 0 && FCString::Strnicmp(*SemanticToSearch, &SourceLine[StartPosition], SemanticLen - 3) == 0)
+		{
+			return true;
+		}
+
+		return false;
 	};
 
 	// Returns true if the specified semantic name starts with "SV_" (case insensitive).
