Fix LiveCoding hitting a check upon starting.

diff --git a/Engine/Source/Developer/Windows/LiveCoding/Private/LiveCodingModule.cpp b/Engine/Source/Developer/Windows/LiveCoding/Private/LiveCodingModule.cpp
--- a/Engine/Source/Developer/Windows/LiveCoding/Private/LiveCodingModule.cpp
+++ b/Engine/Source/Developer/Windows/LiveCoding/Private/LiveCodingModule.cpp
@@ -1040,6 +1040,13 @@ ILiveCodingModule::FOnPatchCompleteDelegate& FLiveCodingModule::GetOnPatchComple
 
 void FLiveCodingModule::StartLiveCodingAsync(ELiveCodingStartupMode StartupMode)
 {
+	// Make sure we can setup the console path correctly
+	if (!SetupConsolePath())
+	{
+		State = EState::NotRunning;
+		return;
+	}
+
 	if (IsRunningCommandlet())
 	{
 		StartLiveCoding(StartupMode);
@@ -1055,6 +1062,40 @@ void FLiveCodingModule::StartLiveCodingAsync(ELiveCodingStartupMode StartupMode)
 	}
 }
 
+bool FLiveCodingModule::SetupConsolePath()
+{
+	// Setup the console path
+	GLiveCodingConsolePath = ConsolePathVariable->GetString();
+	if (!FPaths::FileExists(GLiveCodingConsolePath))
+	{
+		// Check from the executable as the user might have specified different base dir
+		FString CodingConsolePathFromExecutable = FullEngineDirFromExecutable / DefaultConsolePath;
+		FPaths::CollapseRelativeDirectories(CodingConsolePathFromExecutable);
+		if (!FPaths::FileExists(CodingConsolePathFromExecutable))
+		{
+			FFormatNamedArguments Args;
+			Args.Add(TEXT("Executable"), FText::FromString(GLiveCodingConsolePath));
+			const static FText FormatString = LOCTEXT("LiveCodingMissingExecutable", "Unable to start live coding session. Missing executable '{Executable}'. Use the LiveCoding.ConsolePath console variable to modify.");
+			EnableErrorText = FText::Format(FormatString, Args);
+			UE_LOG(LogLiveCoding, Error, TEXT("Unable to start live coding session. Missing executable '%s'. Use the LiveCoding.ConsolePath console variable to modify."), *GLiveCodingConsolePath);
+			return false;
+		}
+
+		GLiveCodingConsolePath = CodingConsolePathFromExecutable;
+
+		// If we found the console from the executable path, chances are users wants the project from there as well
+		const FString ExecutablePath = FPaths::GetPath(FPlatformProcess::ExecutablePath());
+		FString SourceProjectFromExecutable = ExecutablePath / FPaths::GetProjectFilePath();
+		FPaths::NormalizeDirectoryName(SourceProjectFromExecutable);
+		FPaths::CollapseRelativeDirectories(SourceProjectFromExecutable);
+		if(SourceProjectFromExecutable.Len() > 0 && FPaths::FileExists(SourceProjectFromExecutable))
+		{
+			SourceProjectVariable->Set(*SourceProjectFromExecutable);
+		}
+	}
+	return true;
+}
+
 bool FLiveCodingModule::StartLiveCoding(ELiveCodingStartupMode StartupMode)
 {
 	EnableErrorText = FText::GetEmpty();
@@ -1079,39 +1120,6 @@ bool FLiveCodingModule::StartLiveCoding(ELiveCodingStartupMode StartupMode)
 			return false;
 		}
 
-		// Setup the console path
-		GLiveCodingConsolePath = ConsolePathVariable->GetString();
-		if (!FPaths::FileExists(GLiveCodingConsolePath))
-		{
-			// Check from the executable as the user might have specified different base dir
-			FString CodingConsolePathFromExecutable = FullEngineDirFromExecutable / DefaultConsolePath;
-			FPaths::CollapseRelativeDirectories(CodingConsolePathFromExecutable);
-			if (!FPaths::FileExists(CodingConsolePathFromExecutable))
-			{
-				FFormatNamedArguments Args;
-				Args.Add(TEXT("Executable"), FText::FromString(GLiveCodingConsolePath));
-				const static FText FormatString = LOCTEXT("LiveCodingMissingExecutable", "Unable to start live coding session. Missing executable '{Executable}'. Use the LiveCoding.ConsolePath console variable to modify.");
-				EnableErrorText = FText::Format(FormatString, Args);
-				UE_LOG(LogLiveCoding, Error, TEXT("Unable to start live coding session. Missing executable '%s'. Use the LiveCoding.ConsolePath console variable to modify."), *GLiveCodingConsolePath);
-				State = EState::NotRunning;
-				return false;
-			}
-			else
-			{
-				GLiveCodingConsolePath = CodingConsolePathFromExecutable;
-
-				// If we found the console from the executable path, chances are users wants the project from there as well
-				const FString ExecutablePath = FPaths::GetPath(FPlatformProcess::ExecutablePath());
-				FString SourceProjectFromExecutable = ExecutablePath / FPaths::GetProjectFilePath();
-				FPaths::NormalizeDirectoryName(SourceProjectFromExecutable);
-				FPaths::CollapseRelativeDirectories(SourceProjectFromExecutable);
-				if(SourceProjectFromExecutable.Len() > 0 && FPaths::FileExists(SourceProjectFromExecutable))
-				{
-					SourceProjectVariable->Set(*SourceProjectFromExecutable);
-				}
-			}
-		}
-
 		// Get the source project filename
 		const FString SourceProject = SourceProjectVariable->GetString();
 		if (SourceProject.Len() > 0 && !FPaths::FileExists(SourceProject))
diff --git a/Engine/Source/Developer/Windows/LiveCoding/Private/LiveCodingModule.h b/Engine/Source/Developer/Windows/LiveCoding/Private/LiveCodingModule.h
--- a/Engine/Source/Developer/Windows/LiveCoding/Private/LiveCodingModule.h
+++ b/Engine/Source/Developer/Windows/LiveCoding/Private/LiveCodingModule.h
@@ -61,6 +61,7 @@ private:
 	void HideConsole();
 	void EnableConsoleCommand(FOutputDevice& out);
 	void StartLiveCodingAsync(ELiveCodingStartupMode StartupMode);
+	bool SetupConsolePath();
 	bool StartLiveCoding(ELiveCodingStartupMode StartupMode);
 	void OnModulesChanged(FName ModuleName, EModuleChangeReason Reason);
 	void UpdateModules(bool bAllowStarting);
