[shaders]
- add strings as an option for shader statistics, and a flags member/option to hide specific shader statistics from UI (so we can use statistics for other internal usage)
- add debug output for shader platform hashes (use to track down cases of shader job output which has identical bytecode but differing metadata), using the above to record a stat containing the shader hash computed by the compiler. currently only implemented for a subset of platforms
- fix a bug in shader diagnostic remapping for cache hits - the value stored in the cache/DDC for a job should contain non-remapped diagnostics (so cache hits can properly remap them to their version of the unstripped shader code). this also slightly improves deduplication within the job cache/DDC. fixed by storing the job output in the cache/DDC before running the job oncomplete callback (which performs the remapping)
- cleanup dxc hash retrieval code to use actual type instead of hacking around with an offset into the opaque blob

diff --git a/Engine/Shaders/Public/ShaderVersion.ush b/Engine/Shaders/Public/ShaderVersion.ush
--- a/Engine/Shaders/Public/ShaderVersion.ush
+++ b/Engine/Shaders/Public/ShaderVersion.ush
@@ -3,4 +3,4 @@
 // Each time the console command is executed it generates a new GUID. As this file is included
 // in Platform.ush (which should be included in any shader) it allows to invalidate the shader DDC.
 // 
-#pragma message("UESHADERMETADATA_VERSION ED1BB052-50D6-4BE0-A700-AE22838CE9AA")
+#pragma message("UESHADERMETADATA_VERSION 6A88E9C0-AFE1-48D7-9FBA-FDF04FB6C286")
diff --git a/Engine/Source/Developer/ShaderCompilerCommon/Private/ShaderCompilerCommon.cpp b/Engine/Source/Developer/ShaderCompilerCommon/Private/ShaderCompilerCommon.cpp
--- a/Engine/Source/Developer/ShaderCompilerCommon/Private/ShaderCompilerCommon.cpp
+++ b/Engine/Source/Developer/ShaderCompilerCommon/Private/ShaderCompilerCommon.cpp
@@ -1893,6 +1893,19 @@ namespace UE::ShaderCompilerCommon
 		// if modifications were made, this is output as an additional artifact, appending "_DirectCompile" to the path to indicate that it can be used as such.
 		DumpDebugShaderData(Input, PreprocessOutput.GetUnstrippedSourceView(), Options, bHasModifiedSource ? TEXT("_DirectCompile") : nullptr);
 
+		if (EnumHasAnyFlags(Input.DebugInfoFlags, EShaderDebugInfoFlags::ShaderCodePlatformHashes))
+		{
+			// if the platform has registered a CodeHash stat, output a file containing this as well
+			const FGenericShaderStat* Hash = Output.ShaderStatistics.FindByPredicate([](const FGenericShaderStat& Stat)
+				{
+					return Stat.StatName == kPlatformHashStatName;
+				});
+			if (Hash)
+			{
+				FFileHelper::SaveStringToFile(Hash->Value.Get<FString>(), *GetDebugFileName(Input, Options, TEXT("PlatformHash.txt")), FFileHelper::EEncodingOptions::ForceAnsi);
+			}
+		}
+
 		FFileHelper::SaveStringToFile(Output.OutputHash.ToString(), *GetDebugFileName(Input, Options, TEXT("OutputHash.txt")), FFileHelper::EEncodingOptions::ForceAnsi);
 
 		if (EnumHasAnyFlags(Input.DebugInfoFlags, EShaderDebugInfoFlags::Diagnostics))
diff --git a/Engine/Source/Developer/ShaderCompilerCommon/Public/DXCUtils.inl b/Engine/Source/Developer/ShaderCompilerCommon/Public/DXCUtils.inl
--- a/Engine/Source/Developer/ShaderCompilerCommon/Public/DXCUtils.inl
+++ b/Engine/Source/Developer/ShaderCompilerCommon/Public/DXCUtils.inl
@@ -7,6 +7,8 @@
 	#define DXCUTILS_HAS_DXCCREATEINSTANCE 1
 #endif
 
+#include <String/BytesToHex.h>
+
 #if DXCUTILS_HAS_DXCCREATEINSTANCE
 static void DumpDebugBlobDetail(IDxcBlob* Blob, const TCHAR* BlobName)
 #else
@@ -77,8 +79,19 @@ static void DumpDebugBlobDetail(IDxcBlob* Blob, const TCHAR* BlobName, TRefCount
 	FPlatformMisc::LowLevelOutputDebugString(TEXT("\n\n"));
 }
 
-HRESULT RetrieveDebugNameAndBlob(TRefCountPtr<IDxcResult>& CompileResult, FString& OutDebugBlobName, IDxcBlob** OutDebugBlob)
+HRESULT RetrieveDebugNameAndBlob(TRefCountPtr<IDxcResult>& CompileResult, FString& OutDebugBlobName, IDxcBlob** OutDebugBlob, DxcShaderHash& ShaderHash)
 {
+	check(CompileResult->HasOutput(DXC_OUT_SHADER_HASH));
+	{
+		TRefCountPtr<IDxcBlob> HashBlob;
+		TRefCountPtr<IDxcBlobUtf16> OutputNameUnused;
+
+		HRESULT HashResult = CompileResult->GetOutput(DXC_OUT_SHADER_HASH, __uuidof(IDxcBlob), (void**)HashBlob.GetInitReference(), OutputNameUnused.GetInitReference());
+		check(SUCCEEDED(HashResult));
+		check(HashBlob->GetBufferSize() == sizeof(DxcShaderHash));
+		ShaderHash = *reinterpret_cast<DxcShaderHash*>(HashBlob->GetBufferPointer());
+	}
+
 	HRESULT Result = S_OK;
 	if (OutDebugBlob && CompileResult->HasOutput(DXC_OUT_PDB))
 	{
@@ -90,19 +103,14 @@ HRESULT RetrieveDebugNameAndBlob(TRefCountPtr<IDxcResult>& CompileResult, FStrin
 			OutDebugBlobName = DebugBlobName->GetStringPointer();
 		}
 	}
-	else if (CompileResult->HasOutput(DXC_OUT_SHADER_HASH))
+	else
 	{
 		// If the PDB blob is not available, or if we don't need to return it, we can reconstruct the name 
-		// the PDB would have used from the Shader Hash, using only 16 bytes if there is more.
-		TRefCountPtr<IDxcBlob> HashBlob;
-		TRefCountPtr<IDxcBlobUtf16> OutputNameUnused;
-		Result = CompileResult->GetOutput(DXC_OUT_SHADER_HASH, __uuidof(IDxcBlob), (void**)HashBlob.GetInitReference(), OutputNameUnused.GetInitReference());
-		if (SUCCEEDED(Result))
-		{
-			size_t BlobSize = HashBlob->GetBufferSize();
-			uint32 Offset = BlobSize <= 16 ? 0 : BlobSize - 16;
-			OutDebugBlobName = FString(BytesToHex(reinterpret_cast<uint8*>(HashBlob->GetBufferPointer()) + Offset, HashBlob->GetBufferSize() - Offset)) + TEXT(".pdb");
-		}
+		// the PDB would have used from the Shader Hash
+		TStringBuilder<128> PdbNameBuilder;
+		UE::String::BytesToHex(TConstArrayView<uint8>(ShaderHash.HashDigest, sizeof(ShaderHash.HashDigest)), PdbNameBuilder);
+		PdbNameBuilder << TEXT(".pdb");
+		OutDebugBlobName = PdbNameBuilder.ToString();
 	}
 	return Result;
 }
diff --git a/Engine/Source/Developer/ShaderCompilerCommon/Public/ShaderCompilerCommon.h b/Engine/Source/Developer/ShaderCompilerCommon/Public/ShaderCompilerCommon.h
--- a/Engine/Source/Developer/ShaderCompilerCommon/Public/ShaderCompilerCommon.h
+++ b/Engine/Source/Developer/ShaderCompilerCommon/Public/ShaderCompilerCommon.h
@@ -16,6 +16,7 @@ class FShaderSource;
 namespace UE::ShaderCompilerCommon
 {
 	static constexpr const TCHAR* kUniformBufferConstantBufferPrefix = TEXT("UniformBufferConstants_");
+	static constexpr const TCHAR* kPlatformHashStatName = TEXT("PlatformHash");
 }
 
 /**
diff --git a/Engine/Source/Developer/Windows/ShaderFormatD3D/Private/D3DShaderCompilerDXC.cpp b/Engine/Source/Developer/Windows/ShaderFormatD3D/Private/D3DShaderCompilerDXC.cpp
--- a/Engine/Source/Developer/Windows/ShaderFormatD3D/Private/D3DShaderCompilerDXC.cpp
+++ b/Engine/Source/Developer/Windows/ShaderFormatD3D/Private/D3DShaderCompilerDXC.cpp
@@ -7,6 +7,7 @@
 #include "D3D12RHI.h"
 #include "Misc/Paths.h"
 #include "Misc/FileHelper.h"
+#include "Misc/Fnv.h"
 #include "HAL/FileManager.h"
 #include "Serialization/MemoryWriter.h"
 #include "ShaderPreprocessTypes.h"
@@ -678,7 +679,7 @@ static bool RemoveContainerParts(const TConstArrayView<uint32> PartCodes, dxc::D
 }
 
 static HRESULT D3DCompileToDxil(const char* SourceText, const FDxcArguments& Arguments,
-	TRefCountPtr<IDxcBlob>& OutDxilBlob, TRefCountPtr<IDxcBlob>& OutReflectionBlob, TRefCountPtr<IDxcBlobEncoding>& OutErrorBlob, TRefCountPtr<IDxcBlob>& OutPdbBlob, FString& OutPdbName)
+	TRefCountPtr<IDxcBlob>& OutDxilBlob, TRefCountPtr<IDxcBlob>& OutReflectionBlob, TRefCountPtr<IDxcBlobEncoding>& OutErrorBlob, TRefCountPtr<IDxcBlob>& OutPdbBlob, FString& OutPdbName, DxcShaderHash& OutHash)
 {
 	dxc::DxcDllSupport& DxcDllHelper = GetDxcDllHelper();
 
@@ -737,8 +738,7 @@ static HRESULT D3DCompileToDxil(const char* SourceText, const FDxcArguments& Arg
 		TRefCountPtr<IDxcBlobUtf16> ReflectionNameBlob; // Dummy name blob to silence static analysis warning
 		checkf(CompileResult->HasOutput(DXC_OUT_REFLECTION), TEXT("No reflection found!"));
 		VERIFYHRESULT(CompileResult->GetOutput(DXC_OUT_REFLECTION, IID_PPV_ARGS(OutReflectionBlob.GetInitReference()), ReflectionNameBlob.GetInitReference()));
-
-		RetrieveDebugNameAndBlob(CompileResult, OutPdbName, OutPdbBlob.GetInitReference());
+		RetrieveDebugNameAndBlob(CompileResult, OutPdbName, OutPdbBlob.GetInitReference(), OutHash);
  
  		TArray<uint32, TInlineAllocator<4>> PartsToRemove;
 		if (!Arguments.ShouldKeepEmbeddedPDB())
@@ -915,9 +915,11 @@ bool CompileAndProcessD3DShaderDXC(
 	TRefCountPtr<IDxcBlobEncoding> DxcErrorBlob;
 	TRefCountPtr<IDxcBlob> PdbBlob;
 	FString PdbName;
+	DxcShaderHash ShaderHash;
+	const HRESULT D3DCompileToDxilResult = D3DCompileToDxil(AnsiSourceFile.Get(), Args, ShaderBlob, ReflectionBlob, DxcErrorBlob, PdbBlob, PdbName, ShaderHash);
 
-	const HRESULT D3DCompileToDxilResult = D3DCompileToDxil(AnsiSourceFile.Get(), Args, ShaderBlob, ReflectionBlob, DxcErrorBlob, PdbBlob, PdbName);
-
+	Output.AddStatistic(UE::ShaderCompilerCommon::kPlatformHashStatName, BytesToHex(ShaderHash.HashDigest, sizeof(ShaderHash.HashDigest)), FGenericShaderStat::EFlags::Hidden);
+	
 	// Populate the platform-specific debug data with the PDB name, if available.
 	bool bWriteDebugData = Input.Environment.CompilerFlags.Contains(CFLAG_GenerateSymbolsInfo);
 	if (bWriteDebugData && !PdbName.IsEmpty())
diff --git a/Engine/Source/Runtime/Engine/Private/ShaderCompiler/ShaderCompiler.cpp b/Engine/Source/Runtime/Engine/Private/ShaderCompiler/ShaderCompiler.cpp
--- a/Engine/Source/Runtime/Engine/Private/ShaderCompiler/ShaderCompiler.cpp
+++ b/Engine/Source/Runtime/Engine/Private/ShaderCompiler/ShaderCompiler.cpp
@@ -101,6 +101,13 @@ static TAutoConsoleVariable<bool> CVarDebugDumpShaderCode(
 	TEXT("If true, each shader job will dump a ShaderCode.bin containing the contents of the output shader code object (the contents of this can differ for each shader format; note that this is the data that is hashed to produce the OutputHash.txt file)"),
 	ECVF_ReadOnly);
 
+static TAutoConsoleVariable<bool> CVarDebugDumpShaderCodePlatformHashes(
+	TEXT("r.ShaderCompiler.DebugDumpShaderCodePlatformHashes"),
+	false,
+	TEXT("If true, each shader job will dump a PlatformHash.txt file containing the shader code hash as reported by the platform compiler (if the associated shader format registers this hash with the shader stats).\n")
+	TEXT("Note the distinction between this and OutputHash.txt - these files can be used to find shaders which have identical code and only result in different output hashes due to diffs in other metadata."),
+	ECVF_ReadOnly);
+
 static TAutoConsoleVariable<bool> CVarDebugDumpDetailedShaderSource(
 	TEXT("r.ShaderCompiler.DebugDumpDetailedShaderSource"),
 	false,
@@ -1177,6 +1184,11 @@ EShaderDebugInfoFlags FShaderCompilingManager::GetDumpShaderDebugInfoFlags() con
 		Flags |= EShaderDebugInfoFlags::ShaderCodeBinary;
 	}
 
+	if (CVarDebugDumpShaderCodePlatformHashes.GetValueOnAnyThread())
+	{
+		Flags |= EShaderDebugInfoFlags::ShaderCodePlatformHashes;
+	}
+
 	if (CVarDebugDumpDetailedShaderSource.GetValueOnAnyThread())
 	{
 		Flags |= EShaderDebugInfoFlags::DetailedSource;
diff --git a/Engine/Source/Runtime/RenderCore/Private/ShaderCore.cpp b/Engine/Source/Runtime/RenderCore/Private/ShaderCore.cpp
--- a/Engine/Source/Runtime/RenderCore/Private/ShaderCore.cpp
+++ b/Engine/Source/Runtime/RenderCore/Private/ShaderCore.cpp
@@ -69,6 +69,10 @@ bool operator==(const FShaderStatVariant LHS, const FShaderStatVariant RHS)
 	{
 		return LHS.Get<uint32>() == RHS.Get<uint32>();
 	}
+	else if (LHS.IsType<FString>() && RHS.IsType<FString>())
+	{
+		return LHS.Get<FString>() == RHS.Get<FString>();
+	}
 
 	return false;
 }
@@ -92,12 +96,13 @@ FArchive& operator<<(FArchive& Ar, FGenericShaderStat& Stat)
 	}
 
 	Ar << Stat.Value;
+	Ar << Stat.Flags;
 	return Ar;
 }
 
 bool FGenericShaderStat::operator==(const FGenericShaderStat& RHS) const
 {
-	return (StatName == RHS.StatName) && (Value == RHS.Value);
+	return (StatName == RHS.StatName) && (Value == RHS.Value) && (Flags == RHS.Flags);
 }
 
 static TAutoConsoleVariable<bool> CVarDumpDebugInfoForCacheHits(
diff --git a/Engine/Source/Runtime/RenderCore/Private/ShaderMap.cpp b/Engine/Source/Runtime/RenderCore/Private/ShaderMap.cpp
--- a/Engine/Source/Runtime/RenderCore/Private/ShaderMap.cpp
+++ b/Engine/Source/Runtime/RenderCore/Private/ShaderMap.cpp
@@ -394,7 +394,13 @@ TArray<FGenericShaderStat> FShaderMapBase::GetShaderStatistics(FShader& Shader)
 		const int32 ShaderIndex = Code->FindShaderIndex(Shader.GetOutputHash());
 		if (Code->ShaderEditorOnlyDataEntries.IsValidIndex(ShaderIndex))
 		{
-			ShaderStatistics = Code->ShaderEditorOnlyDataEntries[ShaderIndex].ShaderStatistics;
+			for (const FGenericShaderStat& Stat : Code->ShaderEditorOnlyDataEntries[ShaderIndex].ShaderStatistics)
+			{
+				if (!EnumHasAnyFlags(Stat.Flags, FGenericShaderStat::EFlags::Hidden))
+				{
+					ShaderStatistics.Add(Stat);
+				}
+			}
 		}
 	}
 
diff --git a/Engine/Source/Runtime/RenderCore/Public/ShaderCompilerCore.h b/Engine/Source/Runtime/RenderCore/Public/ShaderCompilerCore.h
--- a/Engine/Source/Runtime/RenderCore/Public/ShaderCompilerCore.h
+++ b/Engine/Source/Runtime/RenderCore/Public/ShaderCompilerCore.h
@@ -229,6 +229,7 @@ enum class EShaderDebugInfoFlags : uint8
 	ShaderCodeBinary = 1 << 3,
 	DetailedSource = 1 << 4,
 	CompileFromDebugUSF = 1 << 5,
+	ShaderCodePlatformHashes = 1 << 6,
 };
 ENUM_CLASS_FLAGS(EShaderDebugInfoFlags)
 
@@ -612,11 +613,12 @@ struct FShaderCompilerOutput
 	RENDERCORE_API void SerializeShaderDiagnosticData();
 
 	template<typename TValue>
-	void AddStatistic(const TCHAR* Name, TValue Value)
+	void AddStatistic(const TCHAR* Name, TValue Value, FGenericShaderStat::EFlags Flags = FGenericShaderStat::EFlags::None)
 	{
 		FGenericShaderStat& Stat = ShaderStatistics.AddZeroed_GetRef();
 		Stat.StatName = FName(Name);
 		Stat.Value = FShaderStatVariant(TInPlaceType<TValue>(), Value);
+		Stat.Flags = Flags;
 	}
 
 	// Bump ShaderCompileWorkerOutputVersion if FShaderCompilerOutput changes
diff --git a/Engine/Source/Runtime/RenderCore/Public/ShaderCore.h b/Engine/Source/Runtime/RenderCore/Public/ShaderCore.h
--- a/Engine/Source/Runtime/RenderCore/Public/ShaderCore.h
+++ b/Engine/Source/Runtime/RenderCore/Public/ShaderCore.h
@@ -56,15 +56,21 @@ class FShaderPreprocessorUtilities;
 class FSHA1;
 class ITargetPlatform;
 
-using FShaderStatVariant = TVariant<bool, float, int32, uint32>;
+using FShaderStatVariant = TVariant<bool, float, int32, uint32, FString>;
 
 bool operator==(const FShaderStatVariant LHS, const FShaderStatVariant RHS);
 
 struct FGenericShaderStat
 {
 public:
+	enum class EFlags : uint8
+	{
+		None = 0,
+		Hidden = 1 << 0, // If set this stat will not be shown to the user in the shader stats UI (i.e. stat is for internal use only)
+	};
 	FName StatName;
 	FShaderStatVariant Value;
+	EFlags Flags = EFlags::None;
 
 	FGenericShaderStat() = default;
 
