Fixes a bug where CFLAG_HLSL2021 was ignoring CFLAG_WarningsAsErrors on D3D11

diff --git a/Engine/Shaders/Public/Platform/D3D/D3DCommon.ush b/Engine/Shaders/Public/Platform/D3D/D3DCommon.ush
--- a/Engine/Shaders/Public/Platform/D3D/D3DCommon.ush
+++ b/Engine/Shaders/Public/Platform/D3D/D3DCommon.ush
@@ -5,7 +5,7 @@ D3DCommon.ush: Common defines for the D3D shader compiler
 =============================================================================*/
 
 // Update this GUID to invalidate D3D shaders and force a shader recompilation for D3D RHIs.
-#pragma message("UESHADERMETADATA_VERSION 46268304-2E81-419C-8152-0154868CFB1E")
+#pragma message("UESHADERMETADATA_VERSION 2A25C1D1-C5AB-4121-B43E-EF3CE44B05B1")
 
 #include "/Engine/Shared/HLSLReservedSpaces.h"
 
diff --git a/Engine/Source/Developer/Windows/ShaderFormatD3D/Private/D3DShaderCompiler.cpp b/Engine/Source/Developer/Windows/ShaderFormatD3D/Private/D3DShaderCompiler.cpp
--- a/Engine/Source/Developer/Windows/ShaderFormatD3D/Private/D3DShaderCompiler.cpp
+++ b/Engine/Source/Developer/Windows/ShaderFormatD3D/Private/D3DShaderCompiler.cpp
@@ -818,7 +818,7 @@ static bool CompileAndProcessD3DShaderFXCExt(
 
 			// Compile HLSL source to SPIR-V binary
 			CrossCompiler::FShaderConductorOptions Options;
-
+			Options.bWarningsAsErrors = Input.Environment.CompilerFlags.Contains(CFLAG_WarningsAsErrors);
 			Options.bPreserveStorageInput = true; // Input/output stage variables must match
 			if (bHlslVersion2021)
 			{
