Fixed LiveCoding not detecting changes when working from vanilla editor

diff --git a/Engine/Source/Developer/Windows/LiveCoding/Private/LiveCodingModule.cpp b/Engine/Source/Developer/Windows/LiveCoding/Private/LiveCodingModule.cpp
--- a/Engine/Source/Developer/Windows/LiveCoding/Private/LiveCodingModule.cpp
+++ b/Engine/Source/Developer/Windows/LiveCoding/Private/LiveCodingModule.cpp
@@ -1299,7 +1299,7 @@ void FLiveCodingModule::UpdateModules(bool bAllowStarting)
 			}
 		}
 
-		if (EnableModules.Num() > 0)
+		if (!EnableModules.IsEmpty() || !LazyLoadModules.IsEmpty())
 		{
 			TArray<const TCHAR*> EnableModuleFileNames;
 			for (const FString& EnableModule : EnableModules)
