[Shaders] Change re-ordering of shader interpolators in D3D backend back to order of original source input.

This was previously changed to move all system value semantics (e.g. SV_Position) at the end of the declaration in entry points because sometimes these values are not used in the pixel shader (see CL 34200367).
This change relys solely on the input shader source to re-construct the original order of interpolators after it has been cross-compiled via DXC/SPIRV-Cross.

#rnx

diff --git a/Engine/Shaders/Public/Platform/D3D/D3DCommon.ush b/Engine/Shaders/Public/Platform/D3D/D3DCommon.ush
--- a/Engine/Shaders/Public/Platform/D3D/D3DCommon.ush
+++ b/Engine/Shaders/Public/Platform/D3D/D3DCommon.ush
@@ -5,7 +5,7 @@ D3DCommon.ush: Common defines for the D3D shader compiler
 =============================================================================*/
 
 // Update this GUID to invalidate D3D shaders and force a shader recompilation for D3D RHIs.
-#pragma message("UESHADERMETADATA_VERSION 2A25C1D1-C5AB-4121-B43E-EF3CE44B05B1")
+#pragma message("UESHADERMETADATA_VERSION DE59F027-43C4-4E40-970B-44D1337E6A3F")
 
 #include "/Engine/Shared/HLSLReservedSpaces.h"
 
diff --git a/Engine/Source/Developer/Windows/ShaderFormatD3D/Private/D3DShaderCompiler.cpp b/Engine/Source/Developer/Windows/ShaderFormatD3D/Private/D3DShaderCompiler.cpp
--- a/Engine/Source/Developer/Windows/ShaderFormatD3D/Private/D3DShaderCompiler.cpp
+++ b/Engine/Source/Developer/Windows/ShaderFormatD3D/Private/D3DShaderCompiler.cpp
@@ -581,25 +581,21 @@ static bool PatchHlslWithReorderedIOVariables(
 	};
 
 	auto BuildSortedStageVariableDeclSource = [&IsSemanticSystemValue, &FindSemanticDeclarationInSourceLine](
-		FString& OutStageVariableDeclSource, TArray<FString>& StageVariableDeclLines, const TArray<FString>& InVariables, bool bSystemValueSemanticsOnly)
+		FString& OutStageVariableDeclSource, TArray<FString>& StageVariableDeclLines, const TArray<FString>& InVariables)
 	{
 		for (const FString& Variable : InVariables)
 		{
-			// Are we only emitting source for system value semantics right now?
-			if (IsSemanticSystemValue(Variable) == bSystemValueSemanticsOnly)
+			for (FString& SourceLine : StageVariableDeclLines)
 			{
-				for (FString& SourceLine : StageVariableDeclLines)
+				// Search for semantic name (always case insensitive) in current stage variable source line
+				if (FindSemanticDeclarationInSourceLine(SourceLine, Variable))
 				{
-					// Search for semantic name (always case insensitive) in current stage variable source line
-					if (FindSemanticDeclarationInSourceLine(SourceLine, Variable))
-					{
-						// Append source line for current variable at the end of sorted declaration string.
-						// Then empty this source line to avoid unnecessary string comparisons for next variables.
-						OutStageVariableDeclSource += SourceLine;
-						OutStageVariableDeclSource += TEXT('\n');
-						SourceLine.Empty();
-						break;
-					}
+					// Append source line for current variable at the end of sorted declaration string.
+					// Then empty this source line to avoid unnecessary string comparisons for next variables.
+					OutStageVariableDeclSource += SourceLine;
+					OutStageVariableDeclSource += TEXT('\n');
+					SourceLine.Empty();
+					break;
 				}
 			}
 		}
@@ -608,8 +604,7 @@ static bool PatchHlslWithReorderedIOVariables(
 	// Re-arrange source lines of stage variable declarations and always emit system values last
 	FString SortedStageVariableDeclSource = TEXT("\n");
 
-	BuildSortedStageVariableDeclSource(SortedStageVariableDeclSource, StageVariableDeclSourceLines, Variables, /*bSystemValueSemanticsOnly:*/ false);
-	BuildSortedStageVariableDeclSource(SortedStageVariableDeclSource, StageVariableDeclSourceLines, Variables, /*bSystemValueSemanticsOnly:*/ true);
+	BuildSortedStageVariableDeclSource(SortedStageVariableDeclSource, StageVariableDeclSourceLines, Variables);
 
 	// Replace old declaration with sorted one
 	HlslSourceString.RemoveAt(StageVariableDelcarationBlockBegin + 1, StageVariableDelcarationBlockEnd - (StageVariableDelcarationBlockBegin + 1));
