Fix LiveCoding console not starting when activated through command line
#rnx

diff --git a/Engine/Source/Developer/Windows/LiveCoding/Private/LiveCodingModule.cpp b/Engine/Source/Developer/Windows/LiveCoding/Private/LiveCodingModule.cpp
--- a/Engine/Source/Developer/Windows/LiveCoding/Private/LiveCodingModule.cpp
+++ b/Engine/Source/Developer/Windows/LiveCoding/Private/LiveCodingModule.cpp
@@ -348,7 +348,7 @@ void FLiveCodingModule::StartupModule()
 {
 	LLM_SCOPE_BYTAG(LiveCoding);
 
-	// Register with NT to get dll nitrifications
+	// Register with NT to get dll notifications
 	FNtDllFunction RegisterFunc("LdrRegisterDllNotification");
 	RegisterFunc(0, OnDllNotification, this, &CallbackCookie);
 
@@ -570,6 +570,10 @@ void FLiveCodingModule::EnableForSession(bool bEnable)
 		switch (State)
 		{
 		case EState::NotRunning:
+			if (!SetupConsolePath())
+			{
+				return;
+			}
 			StartLiveCoding(ELiveCodingStartupMode::Manual); // State set in this method
 			ShowConsole();
 			break;
