Check for no derivative ops in work graph shader compilation.
Previously we only checked  and set the flag on compute shaders.

diff --git a/Engine/Shaders/Private/ShaderBundleWorkGraphCommon.ush b/Engine/Shaders/Private/ShaderBundleWorkGraphCommon.ush
--- a/Engine/Shaders/Private/ShaderBundleWorkGraphCommon.ush
+++ b/Engine/Shaders/Private/ShaderBundleWorkGraphCommon.ush
@@ -2,6 +2,8 @@
 
 #include "Common.ush"
 
+#pragma message("UESHADERMETADATA_VERSION 7EDFC9CF-F736-47E6-84FE-D3B245272961")
+
 /** Record passed to shader bundle work graph nodes .*/
 struct FShaderBundleNodeRecord
 {
diff --git a/Engine/Source/Developer/Windows/ShaderFormatD3D/Private/D3DShaderCompilerDXC.cpp b/Engine/Source/Developer/Windows/ShaderFormatD3D/Private/D3DShaderCompilerDXC.cpp
--- a/Engine/Source/Developer/Windows/ShaderFormatD3D/Private/D3DShaderCompilerDXC.cpp
+++ b/Engine/Source/Developer/Windows/ShaderFormatD3D/Private/D3DShaderCompilerDXC.cpp
@@ -985,7 +985,7 @@ bool CompileAndProcessD3DShaderDXC(
 
 		bool bHasNoDerivativeOps = false;
 
-		if (Input.Target.GetFrequency() == SF_Compute && Input.Environment.CompilerFlags.Contains(CFLAG_CheckForDerivativeOps))
+		if ((Input.Target.GetFrequency() == SF_Compute || Input.Target.GetFrequency() == SF_WorkGraphComputeNode) && Input.Environment.CompilerFlags.Contains(CFLAG_CheckForDerivativeOps))
 		{
 			TRefCountPtr<IDxcContainerReflection> ContainerRefl;
 			VERIFYHRESULT(DxcDllHelper.CreateInstance2(GetDxcMalloc(), CLSID_DxcContainerReflection, ContainerRefl.GetInitReference()));
diff --git a/Engine/Source/Runtime/D3D12RHI/Private/D3D12Shaders.cpp b/Engine/Source/Runtime/D3D12RHI/Private/D3D12Shaders.cpp
--- a/Engine/Source/Runtime/D3D12RHI/Private/D3D12Shaders.cpp
+++ b/Engine/Source/Runtime/D3D12RHI/Private/D3D12Shaders.cpp
@@ -210,6 +210,7 @@ FWorkGraphShaderRHIRef FD3D12DynamicRHI::RHICreateWorkGraphShader(TArrayView<con
 	if (Shader)
 	{
 		Shader->RootSignature = GetAdapter().GetRootSignature(Shader);
+		Shader->SetNoDerivativeOps(EnumHasAnyFlags(Shader->ResourceCounts.UsageFlags, EShaderResourceUsageFlags::NoDerivativeOps));
 	}
 
 	return Shader;
