[Shaders] Flush both initial FXC and new DXC errors when cross-compilation from HLSL2021 to HLSL2018 in D3D11 backend failed.

#rnx

diff --git a/Engine/Source/Developer/Windows/ShaderFormatD3D/Private/D3DShaderCompiler.cpp b/Engine/Source/Developer/Windows/ShaderFormatD3D/Private/D3DShaderCompiler.cpp
--- a/Engine/Source/Developer/Windows/ShaderFormatD3D/Private/D3DShaderCompiler.cpp
+++ b/Engine/Source/Developer/Windows/ShaderFormatD3D/Private/D3DShaderCompiler.cpp
@@ -807,6 +807,15 @@ static bool CompileAndProcessD3DShaderFXCExt(
 
 			CrossCompiler::FShaderConductorContext CompilerContext;
 
+			auto FlushInitialFxcAndNewDxcErrors = [&CompilerContext, &InitialFXCRunFilteredErrors, &Output]() -> void
+			{
+				// Flush new errors from DXC
+				CompilerContext.FlushErrors(Output.Errors);
+
+				// Also flush initial errors from FXC
+				Algo::Transform(InitialFXCRunFilteredErrors, Output.Errors, [](const FString& InInitialError) { return FShaderCompilerError(*InInitialError); });
+			};
+
 			// Load shader source into compiler context
 			const EShaderFrequency Frequency = (EShaderFrequency)Input.Target.Frequency;
 			CompilerContext.LoadSource(PreprocessedShaderSource, Input.VirtualSourceFilePath, EntryPointName, Frequency);
@@ -823,7 +832,7 @@ static bool CompileAndProcessD3DShaderFXCExt(
 			FSpirv Spirv;
 			if (!CompilerContext.CompileHlslToSpirv(Options, Spirv.Data))
 			{
-				CompilerContext.FlushErrors(Output.Errors);
+				FlushInitialFxcAndNewDxcErrors();
 				return false;
 			}
 
@@ -846,7 +855,7 @@ static bool CompileAndProcessD3DShaderFXCExt(
 			TArray<ANSICHAR> CrossCompiledSource;
 			if (!CompilerContext.CompileSpirvToSourceAnsi(Options, TargetDesc, Spirv.GetByteData(), Spirv.GetByteSize(), CrossCompiledSource))
 			{
-				CompilerContext.FlushErrors(Output.Errors);
+				FlushInitialFxcAndNewDxcErrors();
 				return false;
 			}
 
