Moved information about live coding info for monolithic builds into BuildSettings. With this change there are fewer changed modules between machines with different engine paths

diff --git a/Engine/Source/Developer/Windows/LiveCoding/LiveCoding.Build.cs b/Engine/Source/Developer/Windows/LiveCoding/LiveCoding.Build.cs
--- a/Engine/Source/Developer/Windows/LiveCoding/LiveCoding.Build.cs
+++ b/Engine/Source/Developer/Windows/LiveCoding/LiveCoding.Build.cs
@@ -10,6 +10,7 @@ public class LiveCoding : ModuleRules
 		PrivateDependencyModuleNames.Add("Core");
 		PrivateDependencyModuleNames.Add("CoreUObject");
 		PrivateDependencyModuleNames.Add("Settings");
+		PrivateDependencyModuleNames.Add("BuildSettings");
 
 		if (Target.bCompileAgainstEngine)
 		{
diff --git a/Engine/Source/Developer/Windows/LiveCoding/Private/LiveCodingModule.cpp b/Engine/Source/Developer/Windows/LiveCoding/Private/LiveCodingModule.cpp
--- a/Engine/Source/Developer/Windows/LiveCoding/Private/LiveCodingModule.cpp
+++ b/Engine/Source/Developer/Windows/LiveCoding/Private/LiveCodingModule.cpp
@@ -1,6 +1,7 @@
 // Copyright Epic Games, Inc. All Rights Reserved.
 
 #include "LiveCodingModule.h"
+#include "BuildSettings.h"
 #include "Modules/ModuleManager.h"
 #include "Features/IModularFeatures.h"
 #include "HAL/IConsoleManager.h"
@@ -56,11 +57,6 @@ FString GLiveCodingConsolePath;
 FString GLiveCodingConsoleArguments;
 FLiveCodingModule* GLiveCodingModule = nullptr;
 
-#if IS_MONOLITHIC
-extern const TCHAR* GLiveCodingEngineDir;
-extern const TCHAR* GLiveCodingProject;
-#endif
-
 #if !WITH_EDITOR
 class FNullReload : public IReload
 {
@@ -444,7 +440,7 @@ void FLiveCodingModule::StartupModule()
 	);
 
 #if IS_MONOLITHIC
-	FString DefaultEngineDir = GLiveCodingEngineDir;
+	FString DefaultEngineDir = BuildSettings::GetLiveCodingEngineDir();
 #else
 	FString DefaultEngineDir = FPaths::EngineDir();
 #endif
@@ -456,7 +452,8 @@ void FLiveCodingModule::StartupModule()
 	);
 
 #if IS_MONOLITHIC
-	FString SourceProject = (GLiveCodingProject != nullptr)? GLiveCodingProject : TEXT("");
+	const TCHAR* LiveCodingProject = BuildSettings::GetLiveCodingProject();
+	FString SourceProject = (LiveCodingProject != nullptr)? LiveCodingProject : TEXT("");
 #else
 	FString SourceProject = FPaths::IsProjectFilePathSet() ? FPaths::GetProjectFilePath() : TEXT("");
 #endif
diff --git a/Engine/Source/Programs/UnrealBuildTool/Configuration/UEBuildTarget.cs b/Engine/Source/Programs/UnrealBuildTool/Configuration/UEBuildTarget.cs
--- a/Engine/Source/Programs/UnrealBuildTool/Configuration/UEBuildTarget.cs
+++ b/Engine/Source/Programs/UnrealBuildTool/Configuration/UEBuildTarget.cs
@@ -5162,16 +5162,6 @@ namespace UnrealBuildTool
 				OutputDirectory = GetOutputDirectoryForExecutable(Unreal.EngineDirectory, Rules.File);
 			}
 
-			// We only want these defines in the launch module to not create non-deterministic Definitions.h on other modules
-			if (Rules.bWithLiveCoding && Rules.LinkType == TargetLinkType.Monolithic)
-			{
-				LaunchModule.Rules.PrivateDefinitions.Add(String.Format("UE_LIVE_CODING_ENGINE_DIR=\"{0}\"", Unreal.EngineDirectory.FullName.Replace("\\", "\\\\")));
-				if (Rules.ProjectFile != null)
-				{
-					LaunchModule.Rules.PrivateDefinitions.Add(String.Format("UE_LIVE_CODING_PROJECT=\"{0}\"", Rules.ProjectFile.FullName.Replace("\\", "\\\\")));
-				}
-			}
-
 			bool bCompileAsDLL = Rules.bShouldCompileAsDLL && bCompileMonolithic;
 			List<FileReference> OutputPaths = MakeBinaryPaths(OutputDirectory, bCompileMonolithic ? TargetName : AppName, Platform, Configuration, bCompileAsDLL ? UEBuildBinaryType.DynamicLinkLibrary : UEBuildBinaryType.Executable, Rules.Architectures, Rules.UndecoratedConfiguration, bCompileMonolithic && ProjectFile != null, Rules.ExeBinariesSubFolder, ProjectFile, Rules);
 
diff --git a/Engine/Source/Runtime/BuildSettings/BuildSettings.Build.cs b/Engine/Source/Runtime/BuildSettings/BuildSettings.Build.cs
--- a/Engine/Source/Runtime/BuildSettings/BuildSettings.Build.cs
+++ b/Engine/Source/Runtime/BuildSettings/BuildSettings.Build.cs
@@ -48,5 +48,14 @@ public class BuildSettings : ModuleRules
 		PrivateDefinitions.Add("SUPPRESS_PER_MODULE_INLINE_FILE"); // This module does not use core's standard operator new/delete overloads
 		
 		bDisableAutoRTFMInstrumentation = true;
+
+		if (Target.bWithLiveCoding && Target.LinkType == TargetLinkType.Monolithic)
+		{
+			PrivateDefinitions.Add(String.Format("UE_LIVE_CODING_ENGINE_DIR=\"{0}\"", UnrealBuildBase.Unreal.EngineDirectory.FullName.Replace("\\", "\\\\")));
+			if (Target.ProjectFile != null)
+			{
+				PrivateDefinitions.Add(String.Format("UE_LIVE_CODING_PROJECT=\"{0}\"", Target.ProjectFile.FullName.Replace("\\", "\\\\")));
+			}
+		}
 	}
 }
diff --git a/Engine/Source/Runtime/BuildSettings/Private/BuildSettings.cpp b/Engine/Source/Runtime/BuildSettings/Private/BuildSettings.cpp
--- a/Engine/Source/Runtime/BuildSettings/Private/BuildSettings.cpp
+++ b/Engine/Source/Runtime/BuildSettings/Private/BuildSettings.cpp
@@ -88,4 +88,22 @@ namespace BuildSettings
 	{
 		return TEXT(BUILD_MACHINENAME);
 	}
+
+	const TCHAR* GetLiveCodingEngineDir()
+	{
+		#ifdef UE_LIVE_CODING_ENGINE_DIR
+		return TEXT(UE_LIVE_CODING_ENGINE_DIR);
+		#else
+		return nullptr;
+		#endif
+	}
+
+	const TCHAR* GetLiveCodingProject()
+	{
+		#ifdef UE_LIVE_CODING_PROJECT
+		return TEXT(UE_LIVE_CODING_PROJECT);
+		#else
+		return nullptr;
+		#endif
+	}
 }
diff --git a/Engine/Source/Runtime/BuildSettings/Public/BuildSettings.h b/Engine/Source/Runtime/BuildSettings/Public/BuildSettings.h
--- a/Engine/Source/Runtime/BuildSettings/Public/BuildSettings.h
+++ b/Engine/Source/Runtime/BuildSettings/Public/BuildSettings.h
@@ -122,4 +122,14 @@ namespace BuildSettings
 	 * Returns the name of the machine which created these binaries
 	 */
 	BUILDSETTINGS_API const TCHAR* GetBuildMachine();
+
+	/**
+	 * Returns the live coding engine dir. Only set in monolithic builds (live coding engine dir is calculated differently in modular builds)
+	 */
+	BUILDSETTINGS_API const TCHAR* GetLiveCodingEngineDir();
+
+	/**
+	 * Returns the live coding project. Only set in monolithic builds
+	 */
+	BUILDSETTINGS_API const TCHAR* GetLiveCodingProject();
 }
diff --git a/Engine/Source/Runtime/Launch/Private/Launch.cpp b/Engine/Source/Runtime/Launch/Private/Launch.cpp
--- a/Engine/Source/Runtime/Launch/Private/Launch.cpp
+++ b/Engine/Source/Runtime/Launch/Private/Launch.cpp
@@ -21,15 +21,6 @@
 	#include "Windows/WindowsHWrapper.h"
 #endif
 
-#ifdef UE_LIVE_CODING_ENGINE_DIR
-const TCHAR* GLiveCodingEngineDir = TEXT(UE_LIVE_CODING_ENGINE_DIR);
-#ifdef UE_LIVE_CODING_PROJECT
-	const TCHAR* GLiveCodingProject = TEXT(UE_LIVE_CODING_PROJECT);
-#else
-	const TCHAR* GLiveCodingProject = nullptr;
-#endif
-#endif
-
 IMPLEMENT_MODULE(FDefaultModuleImpl, Launch);
 
 #if PLATFORM_WINDOWS || PLATFORM_MAC || PLATFORM_UNIX || PLATFORM_USE_GENERIC_LAUNCH_IMPLEMENTATION
