[LiveCoding]
* Added #if LC_VERSION=1 around all lpp1 code

diff --git a/Engine/Source/Developer/Windows/LiveCoding/LiveCoding.Build.cs b/Engine/Source/Developer/Windows/LiveCoding/LiveCoding.Build.cs
--- a/Engine/Source/Developer/Windows/LiveCoding/LiveCoding.Build.cs
+++ b/Engine/Source/Developer/Windows/LiveCoding/LiveCoding.Build.cs
@@ -32,7 +32,9 @@ public class LiveCoding : ModuleRules
             PrivateDefinitions.Add("USE_DEBUG_LIVE_CODING_CONSOLE=0");
         }
 
-        if (Target.Configuration == UnrealTargetConfiguration.Debug)
+		PrivateDefinitions.Add("LC_VERSION=1");
+
+		if (Target.Configuration == UnrealTargetConfiguration.Debug)
         {
         	PrivateDefinitions.Add("LC_DEBUG=1");
         }
diff --git a/Engine/Source/Developer/Windows/LiveCoding/Private/External/LC_API.cpp b/Engine/Source/Developer/Windows/LiveCoding/Private/External/LC_API.cpp
--- a/Engine/Source/Developer/Windows/LiveCoding/Private/External/LC_API.cpp
+++ b/Engine/Source/Developer/Windows/LiveCoding/Private/External/LC_API.cpp
@@ -1,5 +1,7 @@
 // Copyright 2011-2020 Molecular Matters GmbH, all rights reserved.
 
+#if LC_VERSION == 1
+
 // EPIC BEGIN MOD
 //#include PCH_INCLUDE
 // EPIC END MOD
@@ -402,4 +404,6 @@ LPP_DLL_API(void) LppApplySettingString(const char* settingName, const wchar_t*
 	}
 
 	g_startupThread->ApplySettingString(settingName, value);
-}
\ No newline at end of file
+}
+
+#endif // LC_VERSION
\ No newline at end of file
diff --git a/Engine/Source/Developer/Windows/LiveCoding/Private/External/LC_API.h b/Engine/Source/Developer/Windows/LiveCoding/Private/External/LC_API.h
--- a/Engine/Source/Developer/Windows/LiveCoding/Private/External/LC_API.h
+++ b/Engine/Source/Developer/Windows/LiveCoding/Private/External/LC_API.h
@@ -2,6 +2,8 @@
 
 #pragma once
 
+#if LC_VERSION == 1
+
 // BEGIN EPIC MOD
 #include <stddef.h>
 #include <stdint.h>
@@ -52,3 +54,6 @@ LPP_DLL_API(void*) LppEnableModulesEx(const wchar_t* moduleNames[], unsigned int
 LPP_DLL_API(void) LppApplySettingBool(const char* settingName, int value);
 LPP_DLL_API(void) LppApplySettingInt(const char* settingName, int value);
 LPP_DLL_API(void) LppApplySettingString(const char* settingName, const wchar_t* value);
+
+
+#endif // LC_VERSION
\ No newline at end of file
diff --git a/Engine/Source/Developer/Windows/LiveCoding/Private/External/LC_Assert.h b/Engine/Source/Developer/Windows/LiveCoding/Private/External/LC_Assert.h
--- a/Engine/Source/Developer/Windows/LiveCoding/Private/External/LC_Assert.h
+++ b/Engine/Source/Developer/Windows/LiveCoding/Private/External/LC_Assert.h
@@ -2,6 +2,8 @@
 
 #pragma once
 
+#if LC_VERSION == 1
+
 // BEGIN EPIC MOD
 #include "LC_Platform.h"
 #include "LC_Debugger_Windows.h"
@@ -16,3 +18,6 @@
 #	define LC_BREAKPOINT()						(void)true
 #	define LC_ASSERT(_condition, _msg)			LC_UNUSED(_condition)
 #endif
+
+
+#endif // LC_VERSION
\ No newline at end of file
diff --git a/Engine/Source/Developer/Windows/LiveCoding/Private/External/LC_Atomic_Windows.cpp b/Engine/Source/Developer/Windows/LiveCoding/Private/External/LC_Atomic_Windows.cpp
--- a/Engine/Source/Developer/Windows/LiveCoding/Private/External/LC_Atomic_Windows.cpp
+++ b/Engine/Source/Developer/Windows/LiveCoding/Private/External/LC_Atomic_Windows.cpp
@@ -1,5 +1,7 @@
 // Copyright 2011-2020 Molecular Matters GmbH, all rights reserved.
 
+#if LC_VERSION == 1
+
 // BEGIN EPIC MOD
 //#include PCH_INCLUDE
 // END EPIC MOD
@@ -29,3 +31,6 @@ void Atomic::DecrementConsistent(volatile int32_t* value)
 	LC_MEMORY_FENCE;
 	LC_COMPILER_FENCE;
 }
+
+
+#endif // LC_VERSION
\ No newline at end of file
diff --git a/Engine/Source/Developer/Windows/LiveCoding/Private/External/LC_Atomic_Windows.h b/Engine/Source/Developer/Windows/LiveCoding/Private/External/LC_Atomic_Windows.h
--- a/Engine/Source/Developer/Windows/LiveCoding/Private/External/LC_Atomic_Windows.h
+++ b/Engine/Source/Developer/Windows/LiveCoding/Private/External/LC_Atomic_Windows.h
@@ -2,6 +2,8 @@
 
 #pragma once
 
+#if LC_VERSION == 1
+
 // BEGIN EPIC MOD
 #include <cinttypes>
 // END EPIC MOD
@@ -14,3 +16,6 @@ namespace Atomic
 	// Decrements a value atomically, guaranteeing sequential consistency.
 	void DecrementConsistent(volatile int32_t* value);
 }
+
+
+#endif // LC_VERSION
\ No newline at end of file
diff --git a/Engine/Source/Developer/Windows/LiveCoding/Private/External/LC_BitUtil.h b/Engine/Source/Developer/Windows/LiveCoding/Private/External/LC_BitUtil.h
--- a/Engine/Source/Developer/Windows/LiveCoding/Private/External/LC_BitUtil.h
+++ b/Engine/Source/Developer/Windows/LiveCoding/Private/External/LC_BitUtil.h
@@ -2,6 +2,8 @@
 
 #pragma once
 
+#if LC_VERSION == 1
+
 // BEGIN EPIC MOD
 #include "Windows/WindowsHWrapper.h"
 // END EPIC MOD
@@ -35,3 +37,6 @@ namespace bitUtil
 		return index;
 	}
 }
+
+
+#endif // LC_VERSION
\ No newline at end of file
diff --git a/Engine/Source/Developer/Windows/LiveCoding/Private/External/LC_ClientCommandActions.cpp b/Engine/Source/Developer/Windows/LiveCoding/Private/External/LC_ClientCommandActions.cpp
--- a/Engine/Source/Developer/Windows/LiveCoding/Private/External/LC_ClientCommandActions.cpp
+++ b/Engine/Source/Developer/Windows/LiveCoding/Private/External/LC_ClientCommandActions.cpp
@@ -1,5 +1,7 @@
 // Copyright 2011-2020 Molecular Matters GmbH, all rights reserved.
 
+#if LC_VERSION == 1
+
 // BEGIN EPIC MOD
 //#include PCH_INCLUDE
 // END EPIC MOD
@@ -238,4 +240,6 @@ bool actions::TriggerReload::Execute(const CommandType*, const DuplexPipe* pipe,
 
 	return true;
 }
-// END EPIC MOD
\ No newline at end of file
+// END EPIC MOD
+
+#endif // LC_VERSION
\ No newline at end of file
diff --git a/Engine/Source/Developer/Windows/LiveCoding/Private/External/LC_ClientCommandActions.h b/Engine/Source/Developer/Windows/LiveCoding/Private/External/LC_ClientCommandActions.h
--- a/Engine/Source/Developer/Windows/LiveCoding/Private/External/LC_ClientCommandActions.h
+++ b/Engine/Source/Developer/Windows/LiveCoding/Private/External/LC_ClientCommandActions.h
@@ -2,6 +2,8 @@
 
 #pragma once
 
+#if LC_VERSION == 1
+
 // BEGIN EPIC MOD
 #include "CoreTypes.h"
 // END EPIC MOD
@@ -39,3 +41,6 @@ namespace actions
 
 	#undef DECLARE_ACTION
 }
+
+
+#endif // LC_VERSION
\ No newline at end of file
diff --git a/Engine/Source/Developer/Windows/LiveCoding/Private/External/LC_ClientCommandThread.cpp b/Engine/Source/Developer/Windows/LiveCoding/Private/External/LC_ClientCommandThread.cpp
--- a/Engine/Source/Developer/Windows/LiveCoding/Private/External/LC_ClientCommandThread.cpp
+++ b/Engine/Source/Developer/Windows/LiveCoding/Private/External/LC_ClientCommandThread.cpp
@@ -1,5 +1,7 @@
 // Copyright 2011-2020 Molecular Matters GmbH, all rights reserved.
 
+#if LC_VERSION == 1
+
 // BEGIN EPIC MOD
 //#include PCH_INCLUDE
 // END EPIC MOD
@@ -101,3 +103,6 @@ Thread::ReturnValue ClientCommandThread::ThreadFunction(const std::wstring& proc
 	return Thread::ReturnValue(0u);
 	// END EPIC MOD
 }
+
+
+#endif // LC_VERSION
\ No newline at end of file
diff --git a/Engine/Source/Developer/Windows/LiveCoding/Private/External/LC_ClientCommandThread.h b/Engine/Source/Developer/Windows/LiveCoding/Private/External/LC_ClientCommandThread.h
--- a/Engine/Source/Developer/Windows/LiveCoding/Private/External/LC_ClientCommandThread.h
+++ b/Engine/Source/Developer/Windows/LiveCoding/Private/External/LC_ClientCommandThread.h
@@ -2,6 +2,8 @@
 
 #pragma once
 
+#if LC_VERSION == 1
+
 // BEGIN EPIC MOD
 #include "CoreTypes.h"
 // END EPIC MOD
@@ -35,3 +37,6 @@ private:
 	Thread::Handle m_thread;
 	DuplexPipeClient* m_pipe;
 };
+
+
+#endif // LC_VERSION
\ No newline at end of file
diff --git a/Engine/Source/Developer/Windows/LiveCoding/Private/External/LC_ClientStartupThread.cpp b/Engine/Source/Developer/Windows/LiveCoding/Private/External/LC_ClientStartupThread.cpp
--- a/Engine/Source/Developer/Windows/LiveCoding/Private/External/LC_ClientStartupThread.cpp
+++ b/Engine/Source/Developer/Windows/LiveCoding/Private/External/LC_ClientStartupThread.cpp
@@ -1,5 +1,7 @@
 // Copyright 2011-2020 Molecular Matters GmbH, all rights reserved.
 
+#if LC_VERSION == 1
+
 // BEGIN EPIC MOD
 //#include PCH_INCLUDE
 // END EPIC MOD
@@ -689,3 +691,6 @@ Thread::ReturnValue ClientStartupThread::ThreadFunction(const std::wstring& proc
 
 	return Thread::ReturnValue(0u);
 }
+
+
+#endif // LC_VERSION
\ No newline at end of file
diff --git a/Engine/Source/Developer/Windows/LiveCoding/Private/External/LC_ClientStartupThread.h b/Engine/Source/Developer/Windows/LiveCoding/Private/External/LC_ClientStartupThread.h
--- a/Engine/Source/Developer/Windows/LiveCoding/Private/External/LC_ClientStartupThread.h
+++ b/Engine/Source/Developer/Windows/LiveCoding/Private/External/LC_ClientStartupThread.h
@@ -2,6 +2,8 @@
 
 #pragma once
 
+#if LC_VERSION == 1
+
 // BEGIN EPIC MOD
 #include "CoreTypes.h"
 // END EPIC MOD
@@ -121,3 +123,6 @@ private:
 	// process-wide event that is signaled by the Live++ server when compilation is about to begin
 	Event* m_compilationEvent;
 };
+
+
+#endif // LC_VERSION
\ No newline at end of file
diff --git a/Engine/Source/Developer/Windows/LiveCoding/Private/External/LC_ClientUserCommandThread.cpp b/Engine/Source/Developer/Windows/LiveCoding/Private/External/LC_ClientUserCommandThread.cpp
--- a/Engine/Source/Developer/Windows/LiveCoding/Private/External/LC_ClientUserCommandThread.cpp
+++ b/Engine/Source/Developer/Windows/LiveCoding/Private/External/LC_ClientUserCommandThread.cpp
@@ -1,5 +1,7 @@
 // Copyright 2011-2020 Molecular Matters GmbH, all rights reserved.
 
+#if LC_VERSION == 1
+
 // BEGIN EPIC MOD
 //#include PCH_INCLUDE
 // END EPIC MOD
@@ -762,3 +764,6 @@ Thread::ReturnValue ClientUserCommandThread::ThreadFunction(Event* waitForStartE
 		delete command;
 	}
 }
+
+
+#endif // LC_VERSION
\ No newline at end of file
diff --git a/Engine/Source/Developer/Windows/LiveCoding/Private/External/LC_ClientUserCommandThread.h b/Engine/Source/Developer/Windows/LiveCoding/Private/External/LC_ClientUserCommandThread.h
--- a/Engine/Source/Developer/Windows/LiveCoding/Private/External/LC_ClientUserCommandThread.h
+++ b/Engine/Source/Developer/Windows/LiveCoding/Private/External/LC_ClientUserCommandThread.h
@@ -2,6 +2,8 @@
 
 #pragma once
 
+#if LC_VERSION == 1
+
 // BEGIN EPIC MOD
 #include "CoreTypes.h"
 // END EPIC MOD
@@ -132,3 +134,6 @@ private:
 	CriticalSection m_userCommandQueueCS;
 	Semaphore m_userCommandQueueSema;
 };
+
+
+#endif // LC_VERSION
\ No newline at end of file
diff --git a/Engine/Source/Developer/Windows/LiveCoding/Private/External/LC_CommandMap.cpp b/Engine/Source/Developer/Windows/LiveCoding/Private/External/LC_CommandMap.cpp
--- a/Engine/Source/Developer/Windows/LiveCoding/Private/External/LC_CommandMap.cpp
+++ b/Engine/Source/Developer/Windows/LiveCoding/Private/External/LC_CommandMap.cpp
@@ -1,5 +1,7 @@
 // Copyright 2011-2020 Molecular Matters GmbH, all rights reserved.
 
+#if LC_VERSION == 1
+
 // BEGIN EPIC MOD
 //#include PCH_INCLUDE
 // END EPIC MOD
@@ -134,3 +136,6 @@ bool CommandMap::HandleCommands(const DuplexPipe* pipe, void* context)
 		}
 	}
 }
+
+
+#endif // LC_VERSION
\ No newline at end of file
diff --git a/Engine/Source/Developer/Windows/LiveCoding/Private/External/LC_CommandMap.h b/Engine/Source/Developer/Windows/LiveCoding/Private/External/LC_CommandMap.h
--- a/Engine/Source/Developer/Windows/LiveCoding/Private/External/LC_CommandMap.h
+++ b/Engine/Source/Developer/Windows/LiveCoding/Private/External/LC_CommandMap.h
@@ -2,6 +2,8 @@
 
 #pragma once
 
+#if LC_VERSION == 1
+
 // BEGIN EPIC MOD
 #include "CoreTypes.h"
 // END EPIC MOD
@@ -47,3 +49,6 @@ public:
 private:
 	Action m_actions[commands::COUNT];
 };
+
+
+#endif // LC_VERSION
\ No newline at end of file
diff --git a/Engine/Source/Developer/Windows/LiveCoding/Private/External/LC_Commands.h b/Engine/Source/Developer/Windows/LiveCoding/Private/External/LC_Commands.h
--- a/Engine/Source/Developer/Windows/LiveCoding/Private/External/LC_Commands.h
+++ b/Engine/Source/Developer/Windows/LiveCoding/Private/External/LC_Commands.h
@@ -2,6 +2,8 @@
 
 #pragma once
 
+#if LC_VERSION == 1
+
 // BEGIN EPIC MOD
 #include "CoreTypes.h"
 #include "LC_Hook.h"
@@ -370,3 +372,6 @@ namespace commands
 	static const uint32_t COUNT = EnableModulesEx::ID + 1u;
 	// END EPIC MOD
 }
+
+
+#endif // LC_VERSION
\ No newline at end of file
diff --git a/Engine/Source/Developer/Windows/LiveCoding/Private/External/LC_CriticalSection.cpp b/Engine/Source/Developer/Windows/LiveCoding/Private/External/LC_CriticalSection.cpp
--- a/Engine/Source/Developer/Windows/LiveCoding/Private/External/LC_CriticalSection.cpp
+++ b/Engine/Source/Developer/Windows/LiveCoding/Private/External/LC_CriticalSection.cpp
@@ -1,5 +1,7 @@
 // Copyright 2011-2020 Molecular Matters GmbH, all rights reserved.
 
+#if LC_VERSION == 1
+
 // BEGIN EPIC MOD
 //#include PCH_INCLUDE
 // END EPIC MOD
@@ -36,3 +38,6 @@ void CriticalSection::Leave(void)
 	Windows::LeaveCriticalSection(&m_cs);
 	// END EPIC MOD
 }
+
+
+#endif // LC_VERSION
\ No newline at end of file
diff --git a/Engine/Source/Developer/Windows/LiveCoding/Private/External/LC_CriticalSection.h b/Engine/Source/Developer/Windows/LiveCoding/Private/External/LC_CriticalSection.h
--- a/Engine/Source/Developer/Windows/LiveCoding/Private/External/LC_CriticalSection.h
+++ b/Engine/Source/Developer/Windows/LiveCoding/Private/External/LC_CriticalSection.h
@@ -2,6 +2,8 @@
 
 #pragma once
 
+#if LC_VERSION == 1
+
 // BEGIN EPIC MOD
 #include "CoreTypes.h"
 #include "Windows/MinimalWindowsApi.h"
@@ -46,3 +48,6 @@ private:
 	Windows::CRITICAL_SECTION m_cs;
 	// END EPIC MOD
 };
+
+
+#endif // LC_VERSION
\ No newline at end of file
diff --git a/Engine/Source/Developer/Windows/LiveCoding/Private/External/LC_Debugger_Windows.cpp b/Engine/Source/Developer/Windows/LiveCoding/Private/External/LC_Debugger_Windows.cpp
--- a/Engine/Source/Developer/Windows/LiveCoding/Private/External/LC_Debugger_Windows.cpp
+++ b/Engine/Source/Developer/Windows/LiveCoding/Private/External/LC_Debugger_Windows.cpp
@@ -1,5 +1,7 @@
 // Copyright 2011-2020 Molecular Matters GmbH, all rights reserved.
 
+#if LC_VERSION == 1
+
 // BEGIN EPIC MOD
 //#include PCH_INCLUDE
 // END EPIC MOD
@@ -13,3 +15,6 @@ bool Debugger::IsConnected(void)
 {
 	return (::IsDebuggerPresent() != 0);
 }
+
+
+#endif // LC_VERSION
\ No newline at end of file
diff --git a/Engine/Source/Developer/Windows/LiveCoding/Private/External/LC_Debugger_Windows.h b/Engine/Source/Developer/Windows/LiveCoding/Private/External/LC_Debugger_Windows.h
--- a/Engine/Source/Developer/Windows/LiveCoding/Private/External/LC_Debugger_Windows.h
+++ b/Engine/Source/Developer/Windows/LiveCoding/Private/External/LC_Debugger_Windows.h
@@ -2,6 +2,8 @@
 
 #pragma once
 
+#if LC_VERSION == 1
+
 // BEGIN EPIC MOD
 #include "CoreTypes.h"
 // END EPIC MOD
@@ -13,3 +15,6 @@ namespace Debugger
 
 
 #define LC_DEBUGGER_BREAKPOINT() __debugbreak()
+
+
+#endif // LC_VERSION
\ No newline at end of file
diff --git a/Engine/Source/Developer/Windows/LiveCoding/Private/External/LC_DuplexPipe.cpp b/Engine/Source/Developer/Windows/LiveCoding/Private/External/LC_DuplexPipe.cpp
--- a/Engine/Source/Developer/Windows/LiveCoding/Private/External/LC_DuplexPipe.cpp
+++ b/Engine/Source/Developer/Windows/LiveCoding/Private/External/LC_DuplexPipe.cpp
@@ -1,5 +1,7 @@
 // Copyright 2011-2020 Molecular Matters GmbH, all rights reserved.
 
+#if LC_VERSION == 1
+
 // BEGIN EPIC MOD
 //#include PCH_INCLUDE
 // END EPIC MOD
@@ -94,3 +96,6 @@ bool DuplexPipe::Read(void* buffer, size_t size) const
 
 	return true;
 }
+
+
+#endif // LC_VERSION
\ No newline at end of file
diff --git a/Engine/Source/Developer/Windows/LiveCoding/Private/External/LC_DuplexPipe.h b/Engine/Source/Developer/Windows/LiveCoding/Private/External/LC_DuplexPipe.h
--- a/Engine/Source/Developer/Windows/LiveCoding/Private/External/LC_DuplexPipe.h
+++ b/Engine/Source/Developer/Windows/LiveCoding/Private/External/LC_DuplexPipe.h
@@ -2,6 +2,8 @@
 
 #pragma once
 
+#if LC_VERSION == 1
+
 // BEGIN EPIC MOD
 #include "CoreTypes.h"
 // END EPIC MOD
@@ -72,3 +74,6 @@ private:
 	void Send(const void* buffer, size_t size) const;
 	bool Read(void* buffer, size_t size) const;
 };
+
+
+#endif // LC_VERSION
\ No newline at end of file
diff --git a/Engine/Source/Developer/Windows/LiveCoding/Private/External/LC_DuplexPipeClient.cpp b/Engine/Source/Developer/Windows/LiveCoding/Private/External/LC_DuplexPipeClient.cpp
--- a/Engine/Source/Developer/Windows/LiveCoding/Private/External/LC_DuplexPipeClient.cpp
+++ b/Engine/Source/Developer/Windows/LiveCoding/Private/External/LC_DuplexPipeClient.cpp
@@ -1,5 +1,7 @@
 // Copyright 2011-2019 Molecular Matters GmbH, all rights reserved.
 
+#if LC_VERSION == 1
+
 // BEGIN EPIC MOD
 //#include PCH_INCLUDE
 // END EPIC MOD
@@ -37,3 +39,6 @@ bool DuplexPipeClient::Connect(const wchar_t* name)
 
 	return true;
 }
+
+
+#endif // LC_VERSION
\ No newline at end of file
diff --git a/Engine/Source/Developer/Windows/LiveCoding/Private/External/LC_DuplexPipeClient.h b/Engine/Source/Developer/Windows/LiveCoding/Private/External/LC_DuplexPipeClient.h
--- a/Engine/Source/Developer/Windows/LiveCoding/Private/External/LC_DuplexPipeClient.h
+++ b/Engine/Source/Developer/Windows/LiveCoding/Private/External/LC_DuplexPipeClient.h
@@ -2,6 +2,8 @@
 
 #pragma once
 
+#if LC_VERSION == 1
+
 // BEGIN EPIC MOD
 #include "CoreTypes.h"
 // END EPIC MOD
@@ -13,3 +15,6 @@ class DuplexPipeClient : public DuplexPipe
 public:
 	bool Connect(const wchar_t* name);
 };
+
+
+#endif // LC_VERSION
\ No newline at end of file
diff --git a/Engine/Source/Developer/Windows/LiveCoding/Private/External/LC_EntryPoint.cpp b/Engine/Source/Developer/Windows/LiveCoding/Private/External/LC_EntryPoint.cpp
--- a/Engine/Source/Developer/Windows/LiveCoding/Private/External/LC_EntryPoint.cpp
+++ b/Engine/Source/Developer/Windows/LiveCoding/Private/External/LC_EntryPoint.cpp
@@ -1,5 +1,9 @@
 // Copyright 2011-2020 Molecular Matters GmbH, all rights reserved.
 
+#if LC_VERSION == 1
+
 // BEGIN EPIC MOD
 #include "LC_EntryPoint.h"
-// END EPIC MOD
\ No newline at end of file
+// END EPIC MOD
+
+#endif // LC_VERSION
\ No newline at end of file
diff --git a/Engine/Source/Developer/Windows/LiveCoding/Private/External/LC_EntryPoint.h b/Engine/Source/Developer/Windows/LiveCoding/Private/External/LC_EntryPoint.h
--- a/Engine/Source/Developer/Windows/LiveCoding/Private/External/LC_EntryPoint.h
+++ b/Engine/Source/Developer/Windows/LiveCoding/Private/External/LC_EntryPoint.h
@@ -2,8 +2,13 @@
 
 #pragma once
 
+#if LC_VERSION == 1
+
 #include "CoreTypes.h"
 #include "Windows/MinimalWindowsApi.h"
 
 void Startup(Windows::HINSTANCE hInstance);
 void Shutdown();
+
+
+#endif // LC_VERSION
\ No newline at end of file
diff --git a/Engine/Source/Developer/Windows/LiveCoding/Private/External/LC_Environment.cpp b/Engine/Source/Developer/Windows/LiveCoding/Private/External/LC_Environment.cpp
--- a/Engine/Source/Developer/Windows/LiveCoding/Private/External/LC_Environment.cpp
+++ b/Engine/Source/Developer/Windows/LiveCoding/Private/External/LC_Environment.cpp
@@ -1,5 +1,7 @@
 // Copyright 2011-2020 Molecular Matters GmbH, all rights reserved.
 
+#if LC_VERSION == 1
+
 // BEGIN EPIC MOD
 //#include PCH_INCLUDE
 // END EPIC MOD
@@ -62,3 +64,6 @@ std::wstring environment::GetVariable(const wchar_t* variable, const wchar_t* de
 
 	return value;
 }
+
+
+#endif // LC_VERSION
\ No newline at end of file
diff --git a/Engine/Source/Developer/Windows/LiveCoding/Private/External/LC_Environment.h b/Engine/Source/Developer/Windows/LiveCoding/Private/External/LC_Environment.h
--- a/Engine/Source/Developer/Windows/LiveCoding/Private/External/LC_Environment.h
+++ b/Engine/Source/Developer/Windows/LiveCoding/Private/External/LC_Environment.h
@@ -2,6 +2,8 @@
 
 #pragma once
 
+#if LC_VERSION == 1
+
 // BEGIN EPIC MOD
 #include "CoreTypes.h"
 #include <string>
@@ -18,3 +20,6 @@ namespace environment
 	// Gets a variable from the environment of the calling process.
 	std::wstring GetVariable(const wchar_t* variable, const wchar_t* defaultValue);
 }
+
+
+#endif // LC_VERSION
\ No newline at end of file
diff --git a/Engine/Source/Developer/Windows/LiveCoding/Private/External/LC_Event.cpp b/Engine/Source/Developer/Windows/LiveCoding/Private/External/LC_Event.cpp
--- a/Engine/Source/Developer/Windows/LiveCoding/Private/External/LC_Event.cpp
+++ b/Engine/Source/Developer/Windows/LiveCoding/Private/External/LC_Event.cpp
@@ -1,5 +1,7 @@
 // Copyright 2011-2020 Molecular Matters GmbH, all rights reserved.
 
+#if LC_VERSION == 1
+
 // BEGIN EPIC MOD
 //#include PCH_INCLUDE
 // END EPIC MOD
@@ -86,3 +88,6 @@ bool Event::TryWait(void)
 {
 	return WaitTimeout(0u);
 }
+
+
+#endif // LC_VERSION
\ No newline at end of file
diff --git a/Engine/Source/Developer/Windows/LiveCoding/Private/External/LC_Event.h b/Engine/Source/Developer/Windows/LiveCoding/Private/External/LC_Event.h
--- a/Engine/Source/Developer/Windows/LiveCoding/Private/External/LC_Event.h
+++ b/Engine/Source/Developer/Windows/LiveCoding/Private/External/LC_Event.h
@@ -2,6 +2,8 @@
 
 #pragma once
 
+#if LC_VERSION == 1
+
 // BEGIN EPIC MOD
 #include "CoreTypes.h"
 #include "Windows/MinimalWindowsApi.h"
@@ -45,3 +47,6 @@ private:
 	Windows::HANDLE m_event;
 	// END EPIC MOD
 };
+
+
+#endif // LC_VERSION
\ No newline at end of file
diff --git a/Engine/Source/Developer/Windows/LiveCoding/Private/External/LC_Executable.cpp b/Engine/Source/Developer/Windows/LiveCoding/Private/External/LC_Executable.cpp
--- a/Engine/Source/Developer/Windows/LiveCoding/Private/External/LC_Executable.cpp
+++ b/Engine/Source/Developer/Windows/LiveCoding/Private/External/LC_Executable.cpp
@@ -1,5 +1,7 @@
 // Copyright 2011-2020 Molecular Matters GmbH, all rights reserved.
 
+#if LC_VERSION == 1
+
 // BEGIN EPIC MOD
 //#include PCH_INCLUDE
 // END EPIC MOD
@@ -434,3 +436,6 @@ void executable::CallDllEntryPoint(void* moduleBase, uint32_t entryPointRva)
 	DllEntryPoint* entryPoint = pointer::Offset<DllEntryPoint*>(moduleBase, entryPointRva);
 	entryPoint(static_cast<HINSTANCE>(moduleBase), DLL_PROCESS_ATTACH, NULL);
 }
+
+
+#endif // LC_VERSION
\ No newline at end of file
diff --git a/Engine/Source/Developer/Windows/LiveCoding/Private/External/LC_Executable.h b/Engine/Source/Developer/Windows/LiveCoding/Private/External/LC_Executable.h
--- a/Engine/Source/Developer/Windows/LiveCoding/Private/External/LC_Executable.h
+++ b/Engine/Source/Developer/Windows/LiveCoding/Private/External/LC_Executable.h
@@ -2,6 +2,8 @@
 
 #pragma once
 
+#if LC_VERSION == 1
+
 // BEGIN EPIC MOD
 #include "CoreTypes.h"
 // END EPIC MOD
@@ -160,3 +162,6 @@ namespace std
 		}
 	};
 }
+
+
+#endif // LC_VERSION
\ No newline at end of file
diff --git a/Engine/Source/Developer/Windows/LiveCoding/Private/External/LC_FilesystemTypes.h b/Engine/Source/Developer/Windows/LiveCoding/Private/External/LC_FilesystemTypes.h
--- a/Engine/Source/Developer/Windows/LiveCoding/Private/External/LC_FilesystemTypes.h
+++ b/Engine/Source/Developer/Windows/LiveCoding/Private/External/LC_FilesystemTypes.h
@@ -2,6 +2,8 @@
 
 #pragma once
 
+#if LC_VERSION == 1
+
 #include "LC_FixedSizeString.h"
 
 
@@ -38,3 +40,6 @@ namespace Filesystem
 
 	typedef FixedSizeString<wchar_t, MAX_PATH> Path;
 }
+
+
+#endif // LC_VERSION
\ No newline at end of file
diff --git a/Engine/Source/Developer/Windows/LiveCoding/Private/External/LC_FixedSizeString.cpp b/Engine/Source/Developer/Windows/LiveCoding/Private/External/LC_FixedSizeString.cpp
--- a/Engine/Source/Developer/Windows/LiveCoding/Private/External/LC_FixedSizeString.cpp
+++ b/Engine/Source/Developer/Windows/LiveCoding/Private/External/LC_FixedSizeString.cpp
@@ -1,5 +1,7 @@
 // Copyright 2011-2020 Molecular Matters GmbH, all rights reserved.
 
+#if LC_VERSION == 1
+
 // BEGIN EPIC MOD
 //#include PCH_INCLUDE
 // END EPIC MOD
@@ -122,4 +124,6 @@ FixedSizeString<T, N> FixedSizeString<T, N>::ToUpper(void) const LC_RESTRICT
 	return result;
 }
 
-#endif // LC_FIXEDSIZESTRING_CPP
\ No newline at end of file
+#endif // LC_FIXEDSIZESTRING_CPP
+
+#endif // LC_VERSION
\ No newline at end of file
diff --git a/Engine/Source/Developer/Windows/LiveCoding/Private/External/LC_FixedSizeString.h b/Engine/Source/Developer/Windows/LiveCoding/Private/External/LC_FixedSizeString.h
--- a/Engine/Source/Developer/Windows/LiveCoding/Private/External/LC_FixedSizeString.h
+++ b/Engine/Source/Developer/Windows/LiveCoding/Private/External/LC_FixedSizeString.h
@@ -2,6 +2,8 @@
 
 #pragma once
 
+#if LC_VERSION == 1
+
 // BEGIN EPIC MOD
 #include "LC_Platform.h"
 #include "LC_Foundation_Windows.h"
@@ -41,3 +43,6 @@ private:
 	unsigned int m_length;
 	T m_data[N];
 };
+
+
+#endif // LC_VERSION
\ No newline at end of file
diff --git a/Engine/Source/Developer/Windows/LiveCoding/Private/External/LC_FixedSizeString_Windows.cpp b/Engine/Source/Developer/Windows/LiveCoding/Private/External/LC_FixedSizeString_Windows.cpp
--- a/Engine/Source/Developer/Windows/LiveCoding/Private/External/LC_FixedSizeString_Windows.cpp
+++ b/Engine/Source/Developer/Windows/LiveCoding/Private/External/LC_FixedSizeString_Windows.cpp
@@ -1,5 +1,7 @@
 // Copyright 2011-2020 Molecular Matters GmbH, all rights reserved.
 
+#if LC_VERSION == 1
+
 // BEGIN EPIC MOD
 //#include PCH_INCLUDE
 // END EPIC MOD
@@ -12,3 +14,6 @@
 // explicit template instantiation of all used types
 template class FixedSizeString<wchar_t, MAX_PATH>;
 template class FixedSizeString<char, INET_ADDRSTRLEN>;
+
+
+#endif // LC_VERSION
\ No newline at end of file
diff --git a/Engine/Source/Developer/Windows/LiveCoding/Private/External/LC_Foundation_Windows.h b/Engine/Source/Developer/Windows/LiveCoding/Private/External/LC_Foundation_Windows.h
--- a/Engine/Source/Developer/Windows/LiveCoding/Private/External/LC_Foundation_Windows.h
+++ b/Engine/Source/Developer/Windows/LiveCoding/Private/External/LC_Foundation_Windows.h
@@ -2,6 +2,8 @@
 
 #pragma once
 
+#if LC_VERSION == 1
+
 // BEGIN EPIC MOD
 #include "CoreTypes.h"
 #include "Windows/WindowsHWrapper.h"
@@ -27,4 +29,6 @@
 // BEGIN EPIC MOD
 #define LC_COMPILER_FENCE						std::atomic_signal_fence(std::memory_order_acq_rel)
 #define LC_MEMORY_FENCE							std::atomic_thread_fence(std::memory_order_acq_rel)
-// END EPIC MOD
\ No newline at end of file
+// END EPIC MOD
+
+#endif // LC_VERSION
\ No newline at end of file
diff --git a/Engine/Source/Developer/Windows/LiveCoding/Private/External/LC_Hashing.h b/Engine/Source/Developer/Windows/LiveCoding/Private/External/LC_Hashing.h
--- a/Engine/Source/Developer/Windows/LiveCoding/Private/External/LC_Hashing.h
+++ b/Engine/Source/Developer/Windows/LiveCoding/Private/External/LC_Hashing.h
@@ -2,6 +2,8 @@
 
 #pragma once
 
+#if LC_VERSION == 1
+
 #include "xxhash.h"
 // BEGIN EPIC MOD
 #include "LC_Platform.h"
@@ -15,3 +17,6 @@ namespace Hashing
 		return XXH32(input, length, seed);
 	}
 }
+
+
+#endif // LC_VERSION
\ No newline at end of file
diff --git a/Engine/Source/Developer/Windows/LiveCoding/Private/External/LC_HeartBeat.cpp b/Engine/Source/Developer/Windows/LiveCoding/Private/External/LC_HeartBeat.cpp
--- a/Engine/Source/Developer/Windows/LiveCoding/Private/External/LC_HeartBeat.cpp
+++ b/Engine/Source/Developer/Windows/LiveCoding/Private/External/LC_HeartBeat.cpp
@@ -1,5 +1,7 @@
 // Copyright 2011-2020 Molecular Matters GmbH, all rights reserved.
 
+#if LC_VERSION == 1
+
 // BEGIN EPIC MOD
 //#include PCH_INCLUDE
 // END EPIC MOD
@@ -56,3 +58,6 @@ uint64_t HeartBeat::ReadBeat(void) const
 	InterprocessMutex::ScopedLock lock(m_mutex);
 	return Process::ReadNamedSharedMemory<uint64_t>(m_memory);
 }
+
+
+#endif // LC_VERSION
\ No newline at end of file
diff --git a/Engine/Source/Developer/Windows/LiveCoding/Private/External/LC_HeartBeat.h b/Engine/Source/Developer/Windows/LiveCoding/Private/External/LC_HeartBeat.h
--- a/Engine/Source/Developer/Windows/LiveCoding/Private/External/LC_HeartBeat.h
+++ b/Engine/Source/Developer/Windows/LiveCoding/Private/External/LC_HeartBeat.h
@@ -2,6 +2,8 @@
 
 #pragma once
 
+#if LC_VERSION == 1
+
 // BEGIN EPIC MOD
 #include "CoreTypes.h"
 // END EPIC MOD
@@ -31,3 +33,6 @@ private:
 	InterprocessMutex* m_mutex;
 	Process::NamedSharedMemory* m_memory;
 };
+
+
+#endif // LC_VERSION
\ No newline at end of file
diff --git a/Engine/Source/Developer/Windows/LiveCoding/Private/External/LC_InterprocessMutex.cpp b/Engine/Source/Developer/Windows/LiveCoding/Private/External/LC_InterprocessMutex.cpp
--- a/Engine/Source/Developer/Windows/LiveCoding/Private/External/LC_InterprocessMutex.cpp
+++ b/Engine/Source/Developer/Windows/LiveCoding/Private/External/LC_InterprocessMutex.cpp
@@ -1,5 +1,7 @@
 // Copyright 2011-2020 Molecular Matters GmbH, all rights reserved.
 
+#if LC_VERSION == 1
+
 // BEGIN EPIC MOD
 //#include PCH_INCLUDE
 // END EPIC MOD
@@ -55,3 +57,6 @@ void InterprocessMutex::Unlock(void)
 {
 	::ReleaseMutex(m_mutex);
 }
+
+
+#endif // LC_VERSION
\ No newline at end of file
diff --git a/Engine/Source/Developer/Windows/LiveCoding/Private/External/LC_InterprocessMutex.h b/Engine/Source/Developer/Windows/LiveCoding/Private/External/LC_InterprocessMutex.h
--- a/Engine/Source/Developer/Windows/LiveCoding/Private/External/LC_InterprocessMutex.h
+++ b/Engine/Source/Developer/Windows/LiveCoding/Private/External/LC_InterprocessMutex.h
@@ -2,6 +2,8 @@
 
 #pragma once
 
+#if LC_VERSION == 1
+
 // BEGIN EPIC MOD
 #include "Windows/MinimalWindowsApi.h"
 #include "LC_Platform.h"
@@ -45,3 +47,6 @@ private:
 	Windows::HANDLE m_mutex;
 	// END EPIC MOD
 };
+
+
+#endif // LC_VERSION
\ No newline at end of file
diff --git a/Engine/Source/Developer/Windows/LiveCoding/Private/External/LC_Logging.cpp b/Engine/Source/Developer/Windows/LiveCoding/Private/External/LC_Logging.cpp
--- a/Engine/Source/Developer/Windows/LiveCoding/Private/External/LC_Logging.cpp
+++ b/Engine/Source/Developer/Windows/LiveCoding/Private/External/LC_Logging.cpp
@@ -1,5 +1,7 @@
 // Copyright 2011-2020 Molecular Matters GmbH, all rights reserved.
 
+#if LC_VERSION == 1
+
 // BEGIN EPIC MOD
 //#include PCH_INCLUDE
 // END EPIC MOD
@@ -178,3 +180,6 @@ void Logging::SetOutputHandler(Logging::OutputHandlerType handler)
 	}
 }
 // END EPIC MOD - Redirecting log output
+
+
+#endif // LC_VERSION
\ No newline at end of file
diff --git a/Engine/Source/Developer/Windows/LiveCoding/Private/External/LC_Logging.h b/Engine/Source/Developer/Windows/LiveCoding/Private/External/LC_Logging.h
--- a/Engine/Source/Developer/Windows/LiveCoding/Private/External/LC_Logging.h
+++ b/Engine/Source/Developer/Windows/LiveCoding/Private/External/LC_Logging.h
@@ -2,6 +2,8 @@
 
 #pragma once
 
+#if LC_VERSION == 1
+
 // BEGIN EPIC MOD
 #include "CoreTypes.h"
 // END EPIC MOD
@@ -86,3 +88,6 @@ namespace Logging
 #define LC_ERROR_TELEMETRY(_format, ...)		Logging::Log(Logging::Channel::TELEMETRY, Logging::Type::LOG_ERROR, "ERROR: %s" _format "\n", Logging::GetIndentation(Logging::Channel::TELEMETRY), __VA_ARGS__)
 
 #define LC_SUCCESS_USER(_format, ...)			Logging::Log(Logging::Channel::USER, Logging::Type::LOG_SUCCESS, "%s" _format "\n", Logging::GetIndentation(Logging::Channel::USER), __VA_ARGS__)
+
+
+#endif // LC_VERSION
\ No newline at end of file
diff --git a/Engine/Source/Developer/Windows/LiveCoding/Private/External/LC_Memory.cpp b/Engine/Source/Developer/Windows/LiveCoding/Private/External/LC_Memory.cpp
--- a/Engine/Source/Developer/Windows/LiveCoding/Private/External/LC_Memory.cpp
+++ b/Engine/Source/Developer/Windows/LiveCoding/Private/External/LC_Memory.cpp
@@ -1,6 +1,11 @@
 // Copyright 2011-2019 Molecular Matters GmbH, all rights reserved.
 
+#if LC_VERSION == 1
+
 // BEGIN EPIC MOD
 //#include PCH_INCLUDE
 // END EPIC MOD
 #include "LC_Memory.h"
+
+
+#endif // LC_VERSION
\ No newline at end of file
diff --git a/Engine/Source/Developer/Windows/LiveCoding/Private/External/LC_Memory.h b/Engine/Source/Developer/Windows/LiveCoding/Private/External/LC_Memory.h
--- a/Engine/Source/Developer/Windows/LiveCoding/Private/External/LC_Memory.h
+++ b/Engine/Source/Developer/Windows/LiveCoding/Private/External/LC_Memory.h
@@ -2,6 +2,8 @@
 
 #pragma once
 
+#if LC_VERSION == 1
+
 // BEGIN EPIC MOD
 #include "CoreTypes.h"
 // END EPIC MOD
@@ -49,3 +51,6 @@ namespace memory
 		return true;
 	}
 }
+
+
+#endif // LC_VERSION
\ No newline at end of file
diff --git a/Engine/Source/Developer/Windows/LiveCoding/Private/External/LC_MemoryMappedFile.cpp b/Engine/Source/Developer/Windows/LiveCoding/Private/External/LC_MemoryMappedFile.cpp
--- a/Engine/Source/Developer/Windows/LiveCoding/Private/External/LC_MemoryMappedFile.cpp
+++ b/Engine/Source/Developer/Windows/LiveCoding/Private/External/LC_MemoryMappedFile.cpp
@@ -1,5 +1,7 @@
 // Copyright 2011-2020 Molecular Matters GmbH, all rights reserved.
 
+#if LC_VERSION == 1
+
 // BEGIN EPIC MOD
 //#include PCH_INCLUDE
 // BEGIN EPIC MOD
@@ -149,3 +151,6 @@ uint64_t Filesystem::GetMemoryMappedFileSizeOnDisk(const MemoryMappedFile* file)
 
 	return static_cast<uint64_t>(integer.QuadPart);
 }
+
+
+#endif // LC_VERSION
\ No newline at end of file
diff --git a/Engine/Source/Developer/Windows/LiveCoding/Private/External/LC_MemoryMappedFile.h b/Engine/Source/Developer/Windows/LiveCoding/Private/External/LC_MemoryMappedFile.h
--- a/Engine/Source/Developer/Windows/LiveCoding/Private/External/LC_MemoryMappedFile.h
+++ b/Engine/Source/Developer/Windows/LiveCoding/Private/External/LC_MemoryMappedFile.h
@@ -2,6 +2,8 @@
 
 #pragma once
 
+#if LC_VERSION == 1
+
 #include "LC_FilesystemTypes.h"
 
 
@@ -18,3 +20,6 @@ namespace Filesystem
 
 	uint64_t GetMemoryMappedFileSizeOnDisk(const MemoryMappedFile* file);
 }
+
+
+#endif // LC_VERSION
\ No newline at end of file
diff --git a/Engine/Source/Developer/Windows/LiveCoding/Private/External/LC_MemoryStream.cpp b/Engine/Source/Developer/Windows/LiveCoding/Private/External/LC_MemoryStream.cpp
--- a/Engine/Source/Developer/Windows/LiveCoding/Private/External/LC_MemoryStream.cpp
+++ b/Engine/Source/Developer/Windows/LiveCoding/Private/External/LC_MemoryStream.cpp
@@ -1,5 +1,7 @@
 // Copyright 2011-2020 Molecular Matters GmbH, all rights reserved.
 
+#if LC_VERSION == 1
+
 // BEGIN EPIC MOD
 //#include PCH_INCLUDE
 // END EPIC MOD
@@ -54,3 +56,6 @@ void memoryStream::Writer::Write(const void* data, size_t size)
 	memcpy(m_data + m_offset, data, size);
 	m_offset += size;
 }
+
+
+#endif // LC_VERSION
\ No newline at end of file
diff --git a/Engine/Source/Developer/Windows/LiveCoding/Private/External/LC_MemoryStream.h b/Engine/Source/Developer/Windows/LiveCoding/Private/External/LC_MemoryStream.h
--- a/Engine/Source/Developer/Windows/LiveCoding/Private/External/LC_MemoryStream.h
+++ b/Engine/Source/Developer/Windows/LiveCoding/Private/External/LC_MemoryStream.h
@@ -2,6 +2,8 @@
 
 #pragma once
 
+#if LC_VERSION == 1
+
 // BEGIN EPIC MOD
 #include "CoreTypes.h"
 #include "LC_Platform.h"
@@ -74,3 +76,6 @@ namespace memoryStream
 		LC_DISABLE_MOVE_ASSIGNMENT(Writer);
 	};
 }
+
+
+#endif // LC_VERSION
\ No newline at end of file
diff --git a/Engine/Source/Developer/Windows/LiveCoding/Private/External/LC_NamedSharedMemory.cpp b/Engine/Source/Developer/Windows/LiveCoding/Private/External/LC_NamedSharedMemory.cpp
--- a/Engine/Source/Developer/Windows/LiveCoding/Private/External/LC_NamedSharedMemory.cpp
+++ b/Engine/Source/Developer/Windows/LiveCoding/Private/External/LC_NamedSharedMemory.cpp
@@ -1,5 +1,7 @@
 // Copyright 2011-2020 Molecular Matters GmbH, all rights reserved.
 
+#if LC_VERSION == 1
+
 // BEGIN EPIC MOD
 //#include PCH_INCLUDE
 // END EPIC MOD
@@ -85,3 +87,6 @@ void Process::WriteNamedSharedMemory(NamedSharedMemory* memory, const void* buff
 {
 	memcpy(memory->memoryView, buffer, size);
 }
+
+
+#endif // LC_VERSION
\ No newline at end of file
diff --git a/Engine/Source/Developer/Windows/LiveCoding/Private/External/LC_NamedSharedMemory.h b/Engine/Source/Developer/Windows/LiveCoding/Private/External/LC_NamedSharedMemory.h
--- a/Engine/Source/Developer/Windows/LiveCoding/Private/External/LC_NamedSharedMemory.h
+++ b/Engine/Source/Developer/Windows/LiveCoding/Private/External/LC_NamedSharedMemory.h
@@ -2,6 +2,8 @@
 
 #pragma once
 
+#if LC_VERSION == 1
+
 // BEGIN EPIC MOD
 #include "Windows/MinimalWindowsApi.h"
 // END EPIC MOD
@@ -47,3 +49,6 @@ namespace Process
 		WriteNamedSharedMemory(memory, &value, sizeof(T));
 	}
 }
+
+
+#endif // LC_VERSION
\ No newline at end of file
diff --git a/Engine/Source/Developer/Windows/LiveCoding/Private/External/LC_NamedSharedMemoryTypes.h b/Engine/Source/Developer/Windows/LiveCoding/Private/External/LC_NamedSharedMemoryTypes.h
--- a/Engine/Source/Developer/Windows/LiveCoding/Private/External/LC_NamedSharedMemoryTypes.h
+++ b/Engine/Source/Developer/Windows/LiveCoding/Private/External/LC_NamedSharedMemoryTypes.h
@@ -2,9 +2,13 @@
 
 #pragma once
 
+#if LC_VERSION == 1
 
 namespace Process
 {
 	// opaque type
 	struct NamedSharedMemory;
 }
+
+
+#endif // LC_VERSION
\ No newline at end of file
diff --git a/Engine/Source/Developer/Windows/LiveCoding/Private/External/LC_Platform.h b/Engine/Source/Developer/Windows/LiveCoding/Private/External/LC_Platform.h
--- a/Engine/Source/Developer/Windows/LiveCoding/Private/External/LC_Platform.h
+++ b/Engine/Source/Developer/Windows/LiveCoding/Private/External/LC_Platform.h
@@ -2,6 +2,8 @@
 
 #pragma once
 
+#if LC_VERSION == 1
+
 // BEGIN EPIC MOD
 #include "CoreTypes.h"
 #include "Misc/AssertionMacros.h"
@@ -43,4 +45,6 @@
 #else
 #	define LC_IDENTIFIER(_name)		"_" _name
 #endif
-// END EPIC MOD
\ No newline at end of file
+// END EPIC MOD
+
+#endif // LC_VERSION
\ No newline at end of file
diff --git a/Engine/Source/Developer/Windows/LiveCoding/Private/External/LC_PointerUtil.h b/Engine/Source/Developer/Windows/LiveCoding/Private/External/LC_PointerUtil.h
--- a/Engine/Source/Developer/Windows/LiveCoding/Private/External/LC_PointerUtil.h
+++ b/Engine/Source/Developer/Windows/LiveCoding/Private/External/LC_PointerUtil.h
@@ -2,6 +2,8 @@
 
 #pragma once
 
+#if LC_VERSION == 1
+
 // BEGIN EPIC MOD
 #include "CoreTypes.h"
 // END EPIC MOD
@@ -169,3 +171,6 @@ namespace pointer
 		return static_cast<T>(As<const char*>(to) - As<const char*>(from));
 	}
 }
+
+
+#endif // LC_VERSION
\ No newline at end of file
diff --git a/Engine/Source/Developer/Windows/LiveCoding/Private/External/LC_Preprocessor.h b/Engine/Source/Developer/Windows/LiveCoding/Private/External/LC_Preprocessor.h
--- a/Engine/Source/Developer/Windows/LiveCoding/Private/External/LC_Preprocessor.h
+++ b/Engine/Source/Developer/Windows/LiveCoding/Private/External/LC_Preprocessor.h
@@ -2,6 +2,7 @@
 
 #pragma once
 
+#if LC_VERSION == 1
 
 // joins two tokens, even when the tokens are macros
 #define LC_PP_JOIN_HELPER_HELPER(_0, _1)		_0##_1
@@ -14,3 +15,6 @@
 
 // generates a unique name
 #define LC_PP_UNIQUE_NAME(_name)				LC_PP_JOIN(_name, __COUNTER__)
+
+
+#endif // LC_VERSION
\ No newline at end of file
diff --git a/Engine/Source/Developer/Windows/LiveCoding/Private/External/LC_PrimitiveNames.cpp b/Engine/Source/Developer/Windows/LiveCoding/Private/External/LC_PrimitiveNames.cpp
--- a/Engine/Source/Developer/Windows/LiveCoding/Private/External/LC_PrimitiveNames.cpp
+++ b/Engine/Source/Developer/Windows/LiveCoding/Private/External/LC_PrimitiveNames.cpp
@@ -1,5 +1,7 @@
 // Copyright 2011-2020 Molecular Matters GmbH, all rights reserved.
 
+#if LC_VERSION == 1
+
 // BEGIN EPIC MOD
 //#include PCH_INCLUDE
 // END EPIC MOD
@@ -187,3 +189,6 @@ std::wstring primitiveNames::Restart(Process::Id processId)
 #undef LPP_RESTART_REQUESTED
 #undef LPP_RESTART_PREPARED
 #undef LPP_RESTART
+
+
+#endif // LC_VERSION
\ No newline at end of file
diff --git a/Engine/Source/Developer/Windows/LiveCoding/Private/External/LC_PrimitiveNames.h b/Engine/Source/Developer/Windows/LiveCoding/Private/External/LC_PrimitiveNames.h
--- a/Engine/Source/Developer/Windows/LiveCoding/Private/External/LC_PrimitiveNames.h
+++ b/Engine/Source/Developer/Windows/LiveCoding/Private/External/LC_PrimitiveNames.h
@@ -2,6 +2,8 @@
 
 #pragma once
 
+#if LC_VERSION == 1
+
 // BEGIN EPIC MOD
 #include "CoreTypes.h"
 // END EPIC MOD
@@ -35,3 +37,6 @@ namespace primitiveNames
 	std::wstring PreparedRestart(Process::Id processId);
 	std::wstring Restart(Process::Id processId);
 }
+
+
+#endif // LC_VERSION
\ No newline at end of file
diff --git a/Engine/Source/Developer/Windows/LiveCoding/Private/External/LC_Process.h b/Engine/Source/Developer/Windows/LiveCoding/Private/External/LC_Process.h
--- a/Engine/Source/Developer/Windows/LiveCoding/Private/External/LC_Process.h
+++ b/Engine/Source/Developer/Windows/LiveCoding/Private/External/LC_Process.h
@@ -2,6 +2,8 @@
 
 #pragma once
 
+#if LC_VERSION == 1
+
 // BEGIN EPIC MOD
 #include "CoreTypes.h"
 //#include "LC_Types.h"
@@ -167,3 +169,6 @@ namespace Process
 	Environment* CreateEnvironmentFromMap(const TMap<FString, FString>& Pairs);
 	// END EPIC MOD
 }
+
+
+#endif // LC_VERSION
\ No newline at end of file
diff --git a/Engine/Source/Developer/Windows/LiveCoding/Private/External/LC_ProcessTypes.h b/Engine/Source/Developer/Windows/LiveCoding/Private/External/LC_ProcessTypes.h
--- a/Engine/Source/Developer/Windows/LiveCoding/Private/External/LC_ProcessTypes.h
+++ b/Engine/Source/Developer/Windows/LiveCoding/Private/External/LC_ProcessTypes.h
@@ -2,6 +2,8 @@
 
 #pragma once
 
+#if LC_VERSION == 1
+
 #include "LC_FilesystemTypes.h"
 #include "LC_StrongTypedef.h"
 
@@ -41,3 +43,6 @@ namespace Process
 
 	static const Handle INVALID_HANDLE = Handle(INVALID_HANDLE_VALUE);
 }
+
+
+#endif // LC_VERSION
\ No newline at end of file
diff --git a/Engine/Source/Developer/Windows/LiveCoding/Private/External/LC_Restart.cpp b/Engine/Source/Developer/Windows/LiveCoding/Private/External/LC_Restart.cpp
--- a/Engine/Source/Developer/Windows/LiveCoding/Private/External/LC_Restart.cpp
+++ b/Engine/Source/Developer/Windows/LiveCoding/Private/External/LC_Restart.cpp
@@ -1,5 +1,7 @@
 // Copyright 2011-2020 Molecular Matters GmbH, all rights reserved.
 
+#if LC_VERSION == 1
+
 // BEGIN EPIC MOD
 //#include PCH_INCLUDE
 // END EPIC MOD
@@ -84,3 +86,6 @@ void restart::Execute(lpp::RestartBehaviour behaviour, unsigned int exitCode)
 			break;
 	}
 }
+
+
+#endif // LC_VERSION
\ No newline at end of file
diff --git a/Engine/Source/Developer/Windows/LiveCoding/Private/External/LC_Restart.h b/Engine/Source/Developer/Windows/LiveCoding/Private/External/LC_Restart.h
--- a/Engine/Source/Developer/Windows/LiveCoding/Private/External/LC_Restart.h
+++ b/Engine/Source/Developer/Windows/LiveCoding/Private/External/LC_Restart.h
@@ -2,6 +2,8 @@
 
 #pragma once
 
+#if LC_VERSION == 1
+
 #include "LPP_API.h"
 
 
@@ -13,3 +15,6 @@ namespace restart
 	int WasRequested(void);
 	void Execute(lpp::RestartBehaviour behaviour, unsigned int exitCode);
 }
+
+
+#endif // LC_VERSION
\ No newline at end of file
diff --git a/Engine/Source/Developer/Windows/LiveCoding/Private/External/LC_RunMode.h b/Engine/Source/Developer/Windows/LiveCoding/Private/External/LC_RunMode.h
--- a/Engine/Source/Developer/Windows/LiveCoding/Private/External/LC_RunMode.h
+++ b/Engine/Source/Developer/Windows/LiveCoding/Private/External/LC_RunMode.h
@@ -2,6 +2,7 @@
 
 #pragma once
 
+#if LC_VERSION == 1
 
 struct RunMode
 {
@@ -13,3 +14,6 @@ struct RunMode
 		INVALID
 	};
 };
+
+
+#endif // LC_VERSION
\ No newline at end of file
diff --git a/Engine/Source/Developer/Windows/LiveCoding/Private/External/LC_Semaphore.cpp b/Engine/Source/Developer/Windows/LiveCoding/Private/External/LC_Semaphore.cpp
--- a/Engine/Source/Developer/Windows/LiveCoding/Private/External/LC_Semaphore.cpp
+++ b/Engine/Source/Developer/Windows/LiveCoding/Private/External/LC_Semaphore.cpp
@@ -1,5 +1,7 @@
 // Copyright 2011-2020 Molecular Matters GmbH, all rights reserved.
 
+#if LC_VERSION == 1
+
 // BEGIN EPIC MOD
 //#include PCH_INCLUDE
 // END EPIC MOD
@@ -68,3 +70,6 @@ bool Semaphore::TryWait(void)
 {
 	return WaitTimeout(0u);
 }
+
+
+#endif // LC_VERSION
\ No newline at end of file
diff --git a/Engine/Source/Developer/Windows/LiveCoding/Private/External/LC_Semaphore.h b/Engine/Source/Developer/Windows/LiveCoding/Private/External/LC_Semaphore.h
--- a/Engine/Source/Developer/Windows/LiveCoding/Private/External/LC_Semaphore.h
+++ b/Engine/Source/Developer/Windows/LiveCoding/Private/External/LC_Semaphore.h
@@ -2,6 +2,8 @@
 
 #pragma once
 
+#if LC_VERSION == 1
+
 // BEGIN EPIC MOD
 #include "Windows/MinimalWindowsApi.h"
 // END EPIC MOD
@@ -30,3 +32,6 @@ private:
 	Windows::HANDLE m_sema;
 	// END EPIC MOD
 };
+
+
+#endif // LC_VERSION
\ No newline at end of file
diff --git a/Engine/Source/Developer/Windows/LiveCoding/Private/External/LC_Shortcut.cpp b/Engine/Source/Developer/Windows/LiveCoding/Private/External/LC_Shortcut.cpp
--- a/Engine/Source/Developer/Windows/LiveCoding/Private/External/LC_Shortcut.cpp
+++ b/Engine/Source/Developer/Windows/LiveCoding/Private/External/LC_Shortcut.cpp
@@ -1,5 +1,7 @@
 // Copyright 2011-2020 Molecular Matters GmbH, all rights reserved.
 
+#if LC_VERSION == 1
+
 // BEGIN EPIC MOD
 //#include PCH_INCLUDE
 // END EPIC MOD
@@ -86,3 +88,6 @@ std::wstring shortcut::ConvertShortcutToText(int shortcutValue)
 
 	return text;
 }
+
+
+#endif // LC_VERSION
\ No newline at end of file
diff --git a/Engine/Source/Developer/Windows/LiveCoding/Private/External/LC_Shortcut.h b/Engine/Source/Developer/Windows/LiveCoding/Private/External/LC_Shortcut.h
--- a/Engine/Source/Developer/Windows/LiveCoding/Private/External/LC_Shortcut.h
+++ b/Engine/Source/Developer/Windows/LiveCoding/Private/External/LC_Shortcut.h
@@ -2,6 +2,8 @@
 
 #pragma once
 
+#if LC_VERSION == 1
+
 // BEGIN EPIC MOD
 #include "CoreTypes.h"
 #include <string>
@@ -16,3 +18,6 @@ namespace shortcut
 	int GetVirtualKeyCode(int shortcutValue);
 	std::wstring ConvertShortcutToText(int shortcutValue);
 }
+
+
+#endif // LC_VERSION
\ No newline at end of file
diff --git a/Engine/Source/Developer/Windows/LiveCoding/Private/External/LC_StringUtil.cpp b/Engine/Source/Developer/Windows/LiveCoding/Private/External/LC_StringUtil.cpp
--- a/Engine/Source/Developer/Windows/LiveCoding/Private/External/LC_StringUtil.cpp
+++ b/Engine/Source/Developer/Windows/LiveCoding/Private/External/LC_StringUtil.cpp
@@ -1,5 +1,7 @@
 // Copyright 2011-2020 Molecular Matters GmbH, all rights reserved.
 
+#if LC_VERSION == 1
+
 // BEGIN EPIC MOD
 //#include PCH_INCLUDE
 // END EPIC MOD
@@ -508,3 +510,6 @@ namespace string
 		return safeName;
 	}
 }
+
+
+#endif // LC_VERSION
\ No newline at end of file
diff --git a/Engine/Source/Developer/Windows/LiveCoding/Private/External/LC_StringUtil.h b/Engine/Source/Developer/Windows/LiveCoding/Private/External/LC_StringUtil.h
--- a/Engine/Source/Developer/Windows/LiveCoding/Private/External/LC_StringUtil.h
+++ b/Engine/Source/Developer/Windows/LiveCoding/Private/External/LC_StringUtil.h
@@ -2,6 +2,8 @@
 
 #pragma once
 
+#if LC_VERSION == 1
+
 // BEGIN EPIC MOD
 #include "CoreTypes.h"
 #include <string>
@@ -108,3 +110,6 @@ namespace string
 	}
 #endif
 }
+
+
+#endif // LC_VERSION
\ No newline at end of file
diff --git a/Engine/Source/Developer/Windows/LiveCoding/Private/External/LC_StrongTypedef.h b/Engine/Source/Developer/Windows/LiveCoding/Private/External/LC_StrongTypedef.h
--- a/Engine/Source/Developer/Windows/LiveCoding/Private/External/LC_StrongTypedef.h
+++ b/Engine/Source/Developer/Windows/LiveCoding/Private/External/LC_StrongTypedef.h
@@ -2,6 +2,8 @@
 
 #pragma once
 
+#if LC_VERSION == 1
+
 //BEGIN EPIC MOD
 #include <algorithm>
 //END EPIC MOD
@@ -97,3 +99,6 @@ namespace std
 	};
 }
 #endif
+
+
+#endif // LC_VERSION
\ No newline at end of file
diff --git a/Engine/Source/Developer/Windows/LiveCoding/Private/External/LC_SyncPoint.cpp b/Engine/Source/Developer/Windows/LiveCoding/Private/External/LC_SyncPoint.cpp
--- a/Engine/Source/Developer/Windows/LiveCoding/Private/External/LC_SyncPoint.cpp
+++ b/Engine/Source/Developer/Windows/LiveCoding/Private/External/LC_SyncPoint.cpp
@@ -1,5 +1,7 @@
 // Copyright 2011-2020 Molecular Matters GmbH, all rights reserved.
 
+#if LC_VERSION == 1
+
 // BEGIN EPIC MOD
 //#include PCH_INCLUDE
 // END EPIC MOD
@@ -61,3 +63,6 @@ void syncPoint::EnterTarget(void)
 		g_leaveUserSyncPoint.Wait();
 	}
 }
+
+
+#endif // LC_VERSION
\ No newline at end of file
diff --git a/Engine/Source/Developer/Windows/LiveCoding/Private/External/LC_SyncPoint.h b/Engine/Source/Developer/Windows/LiveCoding/Private/External/LC_SyncPoint.h
--- a/Engine/Source/Developer/Windows/LiveCoding/Private/External/LC_SyncPoint.h
+++ b/Engine/Source/Developer/Windows/LiveCoding/Private/External/LC_SyncPoint.h
@@ -2,6 +2,8 @@
 
 #pragma once
 
+#if LC_VERSION == 1
+
 // BEGIN EPIC MOD
 #include "CoreTypes.h"
 // END EPIC MOD
@@ -17,3 +19,6 @@ namespace syncPoint
 	// Called by the LPP API when user code in a target process wants to enter the sync point
 	void EnterTarget(void);
 }
+
+
+#endif // LC_VERSION
\ No newline at end of file
diff --git a/Engine/Source/Developer/Windows/LiveCoding/Private/External/LC_Thread.cpp b/Engine/Source/Developer/Windows/LiveCoding/Private/External/LC_Thread.cpp
--- a/Engine/Source/Developer/Windows/LiveCoding/Private/External/LC_Thread.cpp
+++ b/Engine/Source/Developer/Windows/LiveCoding/Private/External/LC_Thread.cpp
@@ -1,5 +1,7 @@
 // Copyright 2011-2020 Molecular Matters GmbH, all rights reserved.
 
+#if LC_VERSION == 1
+
 // BEGIN EPIC MOD
 //#include PCH_INCLUDE
 // END EPIC MOD
@@ -180,3 +182,6 @@ Thread::Id Thread::GetId(Handle handle)
 // BEGIN EPIC MODS
 #pragma warning(pop)
 // END EPIC MODS
+
+
+#endif // LC_VERSION
\ No newline at end of file
diff --git a/Engine/Source/Developer/Windows/LiveCoding/Private/External/LC_Thread.h b/Engine/Source/Developer/Windows/LiveCoding/Private/External/LC_Thread.h
--- a/Engine/Source/Developer/Windows/LiveCoding/Private/External/LC_Thread.h
+++ b/Engine/Source/Developer/Windows/LiveCoding/Private/External/LC_Thread.h
@@ -2,6 +2,8 @@
 
 #pragma once
 
+#if LC_VERSION == 1
+
 // BEGIN EPIC MOD
 #include "CoreTypes.h"
 // END EPIC MOD
@@ -165,3 +167,6 @@ inline Thread::Handle Thread::CreateFromMemberFunction(const char* threadName, u
 
 	return handle;
 }
+
+
+#endif // LC_VERSION
\ No newline at end of file
diff --git a/Engine/Source/Developer/Windows/LiveCoding/Private/External/LC_ThreadTypes.h b/Engine/Source/Developer/Windows/LiveCoding/Private/External/LC_ThreadTypes.h
--- a/Engine/Source/Developer/Windows/LiveCoding/Private/External/LC_ThreadTypes.h
+++ b/Engine/Source/Developer/Windows/LiveCoding/Private/External/LC_ThreadTypes.h
@@ -2,6 +2,8 @@
 
 #pragma once
 
+#if LC_VERSION == 1
+
 // BEGIN EPIC MOD
 #include "CoreTypes.h"
 // END EPIC MOD
@@ -25,3 +27,6 @@ namespace Thread
 
 	static const Handle INVALID_HANDLE = Handle(INVALID_HANDLE_VALUE);
 }
+
+
+#endif // LC_VERSION
\ No newline at end of file
diff --git a/Engine/Source/Developer/Windows/LiveCoding/Private/External/LC_UtcTime.cpp b/Engine/Source/Developer/Windows/LiveCoding/Private/External/LC_UtcTime.cpp
--- a/Engine/Source/Developer/Windows/LiveCoding/Private/External/LC_UtcTime.cpp
+++ b/Engine/Source/Developer/Windows/LiveCoding/Private/External/LC_UtcTime.cpp
@@ -1,5 +1,7 @@
 // Copyright 2011-2020 Molecular Matters GmbH, all rights reserved.
 
+#if LC_VERSION == 1
+
 // BEGIN EPIC MOD
 //#include PCH_INCLUDE
 // END EPIC MOD
@@ -19,3 +21,6 @@ uint64_t utcTime::GetCurrent(void)
 
 	return asLargeInteger.QuadPart;
 }
+
+
+#endif // LC_VERSION
\ No newline at end of file
diff --git a/Engine/Source/Developer/Windows/LiveCoding/Private/External/LC_UtcTime.h b/Engine/Source/Developer/Windows/LiveCoding/Private/External/LC_UtcTime.h
--- a/Engine/Source/Developer/Windows/LiveCoding/Private/External/LC_UtcTime.h
+++ b/Engine/Source/Developer/Windows/LiveCoding/Private/External/LC_UtcTime.h
@@ -2,6 +2,8 @@
 
 #pragma once
 
+#if LC_VERSION == 1
+
 // BEGIN EPIC MOD
 #include "CoreTypes.h"
 #include <inttypes.h>
@@ -11,3 +13,6 @@ namespace utcTime
 {
 	uint64_t GetCurrent(void);
 }
+
+
+#endif // LC_VERSION
\ No newline at end of file
diff --git a/Engine/Source/Developer/Windows/LiveCoding/Private/External/LC_VirtualMemory.cpp b/Engine/Source/Developer/Windows/LiveCoding/Private/External/LC_VirtualMemory.cpp
--- a/Engine/Source/Developer/Windows/LiveCoding/Private/External/LC_VirtualMemory.cpp
+++ b/Engine/Source/Developer/Windows/LiveCoding/Private/External/LC_VirtualMemory.cpp
@@ -1,5 +1,7 @@
 // Copyright 2011-2020 Molecular Matters GmbH, all rights reserved.
 
+#if LC_VERSION == 1
+
 // BEGIN EPIC MOD
 //#include PCH_INCLUDE
 // END EPIC MOD
@@ -61,3 +63,6 @@ uint32_t VirtualMemory::GetPageSize(void)
 	static const uint32_t pageSize = ::GetPageSize();
 	return pageSize;
 }
+
+
+#endif // LC_VERSION
\ No newline at end of file
diff --git a/Engine/Source/Developer/Windows/LiveCoding/Private/External/LC_VirtualMemory.h b/Engine/Source/Developer/Windows/LiveCoding/Private/External/LC_VirtualMemory.h
--- a/Engine/Source/Developer/Windows/LiveCoding/Private/External/LC_VirtualMemory.h
+++ b/Engine/Source/Developer/Windows/LiveCoding/Private/External/LC_VirtualMemory.h
@@ -2,6 +2,8 @@
 
 #pragma once
 
+#if LC_VERSION == 1
+
 // BEGIN EPIC MOD
 #include "CoreTypes.h"
 // END EPIC MOD
@@ -30,3 +32,6 @@ namespace VirtualMemory
 	// Returns the page size used by the virtual memory system.
 	uint32_t GetPageSize(void);
 }
+
+
+#endif // LC_VERSION
\ No newline at end of file
diff --git a/Engine/Source/Developer/Windows/LiveCoding/Private/External/LC_VirtualMemoryTypes.h b/Engine/Source/Developer/Windows/LiveCoding/Private/External/LC_VirtualMemoryTypes.h
--- a/Engine/Source/Developer/Windows/LiveCoding/Private/External/LC_VirtualMemoryTypes.h
+++ b/Engine/Source/Developer/Windows/LiveCoding/Private/External/LC_VirtualMemoryTypes.h
@@ -2,6 +2,8 @@
 
 #pragma once
 
+#if LC_VERSION == 1
+
 // BEGIN EPIC MOD
 #include "Windows/WindowsHWrapper.h"
 // END EPIC MOD
@@ -17,3 +19,6 @@ namespace VirtualMemory
 		};
 	};
 }
+
+
+#endif // LC_VERSION
\ No newline at end of file
diff --git a/Engine/Source/Developer/Windows/LiveCoding/Private/External/LC_WindowsInternalFunctions.cpp b/Engine/Source/Developer/Windows/LiveCoding/Private/External/LC_WindowsInternalFunctions.cpp
--- a/Engine/Source/Developer/Windows/LiveCoding/Private/External/LC_WindowsInternalFunctions.cpp
+++ b/Engine/Source/Developer/Windows/LiveCoding/Private/External/LC_WindowsInternalFunctions.cpp
@@ -1,6 +1,9 @@
 // Copyright 2011-2020 Molecular Matters GmbH, all rights reserved.
 
 // BEGIN EPIC MOD
+
+#if LC_VERSION == 1
+
 //#include PCH_INCLUDE
 // END EPIC MOD
 #include "LC_WindowsInternalFunctions.h"
@@ -13,3 +16,6 @@ WindowsInternals::Function<decltype(NtWriteVirtualMemory)> WindowsInternals::NtW
 WindowsInternals::Function<decltype(NtQuerySystemInformation)> WindowsInternals::NtQuerySystemInformation("ntdll.dll", "NtQuerySystemInformation");
 WindowsInternals::Function<decltype(NtQueryInformationProcess)> WindowsInternals::NtQueryInformationProcess("ntdll.dll", "NtQueryInformationProcess");
 WindowsInternals::Function<decltype(NtContinue)> WindowsInternals::NtContinue("ntdll.dll", "NtContinue");
+
+
+#endif // LC_VERSION
\ No newline at end of file
diff --git a/Engine/Source/Developer/Windows/LiveCoding/Private/External/LC_WindowsInternalFunctions.h b/Engine/Source/Developer/Windows/LiveCoding/Private/External/LC_WindowsInternalFunctions.h
--- a/Engine/Source/Developer/Windows/LiveCoding/Private/External/LC_WindowsInternalFunctions.h
+++ b/Engine/Source/Developer/Windows/LiveCoding/Private/External/LC_WindowsInternalFunctions.h
@@ -2,6 +2,8 @@
 
 #pragma once
 
+#if LC_VERSION == 1
+
 #include "LC_WindowsInternals.h"
 // BEGIN EPIC MODS
 #include "LC_Logging.h"
@@ -117,3 +119,6 @@ namespace WindowsInternals
 	extern Function<decltype(NtQueryInformationProcess)> NtQueryInformationProcess;
 	extern Function<decltype(NtContinue)> NtContinue;
 }
+
+
+#endif // LC_VERSION
\ No newline at end of file
diff --git a/Engine/Source/Developer/Windows/LiveCoding/Private/External/LC_WindowsInternals.h b/Engine/Source/Developer/Windows/LiveCoding/Private/External/LC_WindowsInternals.h
--- a/Engine/Source/Developer/Windows/LiveCoding/Private/External/LC_WindowsInternals.h
+++ b/Engine/Source/Developer/Windows/LiveCoding/Private/External/LC_WindowsInternals.h
@@ -2,6 +2,8 @@
 
 #pragma once
 
+#if LC_VERSION == 1
+
 // BEGIN EPIC MOD
 #include "Windows/WindowsHWrapper.h"
 // END EPIC MOD
@@ -261,3 +263,5 @@ namespace WindowsInternals
 		PVOID Reserved3;
 	};
 }
+
+#endif // LC_VERSION
\ No newline at end of file
diff --git a/Engine/Source/Developer/Windows/LiveCoding/Private/External/LPP_API.h b/Engine/Source/Developer/Windows/LiveCoding/Private/External/LPP_API.h
--- a/Engine/Source/Developer/Windows/LiveCoding/Private/External/LPP_API.h
+++ b/Engine/Source/Developer/Windows/LiveCoding/Private/External/LPP_API.h
@@ -2,6 +2,7 @@
 
 #pragma once
 
+#if LC_VERSION == 1
 
 /******************************************************************************/
 /* HOOKS                                                                      */
@@ -494,3 +495,5 @@ LPP_NS_END
 #undef LPP_CALL2
 #undef LPP_CALL3
 #undef LPP_CALL4
+
+#endif // LC_VERSION
\ No newline at end of file
