[LiveCodingServer]
* Fixed compile error caused by missing define LC_VERSION
* Added ExtraRootPath
* Added #if LC_VERSION == 1 around all livecode 1 code

diff --git a/Engine/Source/Developer/Windows/LiveCoding/Private/External/LC_Process.cpp b/Engine/Source/Developer/Windows/LiveCoding/Private/External/LC_Process.cpp
--- a/Engine/Source/Developer/Windows/LiveCoding/Private/External/LC_Process.cpp
+++ b/Engine/Source/Developer/Windows/LiveCoding/Private/External/LC_Process.cpp
@@ -1,5 +1,7 @@
 // Copyright 2011-2020 Molecular Matters GmbH, all rights reserved.
 
+#if LC_VERSION == 1
+
 // BEGIN EPIC MOD
 //#include PCH_INCLUDE
 // END EPIC MOD
@@ -755,3 +757,5 @@ void Process::DumpMemory(Handle handle, const void* address, size_t size)
 // BEGIN EPIC MODS
 #pragma warning(pop)
 // END EPIC MODS
+
+#endif // LC_VERSION
\ No newline at end of file
diff --git a/Engine/Source/Developer/Windows/LiveCodingServer/LiveCodingServer.Build.cs b/Engine/Source/Developer/Windows/LiveCodingServer/LiveCodingServer.Build.cs
--- a/Engine/Source/Developer/Windows/LiveCodingServer/LiveCodingServer.Build.cs
+++ b/Engine/Source/Developer/Windows/LiveCodingServer/LiveCodingServer.Build.cs
@@ -18,6 +18,8 @@ public class LiveCodingServer : ModuleRules
 
         if (Target.Platform == UnrealTargetPlatform.Win64)
 		{
+			PrivateDefinitions.Add("LC_VERSION=1");
+
 			PrivateIncludePaths.Add("Developer/Windows/LiveCoding/Private");
 			PrivateIncludePaths.Add("Developer/Windows/LiveCoding/Private/External");
 
@@ -27,6 +29,7 @@ public class LiveCodingServer : ModuleRules
 				throw new System.Exception("Unable to find DIA SDK directory");
 			}
 
+			ExtraRootPath = ("DiaSdkDir", DiaSdkDir);
 			PrivateIncludePaths.Add(Path.Combine(DiaSdkDir, "include"));
 			PublicAdditionalLibraries.Add(Path.Combine(DiaSdkDir, "lib", "amd64", "diaguids.lib"));
 			RuntimeDependencies.Add("$(TargetOutputDir)/msdia140.dll", Path.Combine(DiaSdkDir, "bin", "amd64", "msdia140.dll"));
diff --git a/Engine/Source/Developer/Windows/LiveCodingServer/Private/External/LC_AllocatorStats.cpp b/Engine/Source/Developer/Windows/LiveCodingServer/Private/External/LC_AllocatorStats.cpp
--- a/Engine/Source/Developer/Windows/LiveCodingServer/Private/External/LC_AllocatorStats.cpp
+++ b/Engine/Source/Developer/Windows/LiveCodingServer/Private/External/LC_AllocatorStats.cpp
@@ -1,5 +1,7 @@
 // Copyright 2011-2020 Molecular Matters GmbH, all rights reserved.
 
+#if LC_VERSION == 1
+
 // BEGIN EPIC MOD
 //#include PCH_INCLUDE
 // END EPIC MOD
@@ -65,3 +67,6 @@ uint64_t AllocatorStats::GetMemorySize(void) const
 // BEGIN EPIC MOD
 #include "Windows/HideWindowsPlatformAtomics.h"
 // END EPIC MOD
+
+
+#endif
\ No newline at end of file
diff --git a/Engine/Source/Developer/Windows/LiveCodingServer/Private/External/LC_AllocatorStats.h b/Engine/Source/Developer/Windows/LiveCodingServer/Private/External/LC_AllocatorStats.h
--- a/Engine/Source/Developer/Windows/LiveCodingServer/Private/External/LC_AllocatorStats.h
+++ b/Engine/Source/Developer/Windows/LiveCodingServer/Private/External/LC_AllocatorStats.h
@@ -2,6 +2,8 @@
 
 #pragma once
 
+#if LC_VERSION == 1
+
 // BEGIN EPIC MOD
 #include "CoreTypes.h"
 #include <stdint.h>
@@ -27,3 +29,6 @@ private:
 	uint64_t m_allocationCount;
 	uint64_t m_memorySize;
 };
+
+
+#endif
\ No newline at end of file
diff --git a/Engine/Source/Developer/Windows/LiveCodingServer/Private/External/LC_Allocators.cpp b/Engine/Source/Developer/Windows/LiveCodingServer/Private/External/LC_Allocators.cpp
--- a/Engine/Source/Developer/Windows/LiveCodingServer/Private/External/LC_Allocators.cpp
+++ b/Engine/Source/Developer/Windows/LiveCodingServer/Private/External/LC_Allocators.cpp
@@ -1,5 +1,7 @@
 // Copyright 2011-2020 Molecular Matters GmbH, all rights reserved.
 
+#if LC_VERSION == 1
+
 // BEGIN EPIC MOD
 //#include PCH_INCLUDE
 // END EPIC MOD
@@ -20,3 +22,6 @@ Mallocator g_rawElfAllocator("Raw ELF files", 8u);
 Mallocator g_relocatableFileAllocator("Relocatable files", 8u);
 Mallocator g_executableFileAllocator("Executable files", 8u);
 Mallocator g_tlogFileAllocator("Tlog files", 8u);
+
+
+#endif
\ No newline at end of file
diff --git a/Engine/Source/Developer/Windows/LiveCodingServer/Private/External/LC_Allocators.h b/Engine/Source/Developer/Windows/LiveCodingServer/Private/External/LC_Allocators.h
--- a/Engine/Source/Developer/Windows/LiveCodingServer/Private/External/LC_Allocators.h
+++ b/Engine/Source/Developer/Windows/LiveCodingServer/Private/External/LC_Allocators.h
@@ -2,6 +2,8 @@
 
 #pragma once
 
+#if LC_VERSION == 1
+
 // BEGIN EPIC MOD
 #include "CoreTypes.h"
 // END EPIC MOD
@@ -61,3 +63,6 @@ extern Mallocator g_rawElfAllocator;
 extern Mallocator g_relocatableFileAllocator;
 extern Mallocator g_executableFileAllocator;
 extern Mallocator g_tlogFileAllocator;
+
+
+#endif
\ No newline at end of file
diff --git a/Engine/Source/Developer/Windows/LiveCodingServer/Private/External/LC_Amalgamation.cpp b/Engine/Source/Developer/Windows/LiveCodingServer/Private/External/LC_Amalgamation.cpp
--- a/Engine/Source/Developer/Windows/LiveCodingServer/Private/External/LC_Amalgamation.cpp
+++ b/Engine/Source/Developer/Windows/LiveCodingServer/Private/External/LC_Amalgamation.cpp
@@ -1,5 +1,7 @@
 // Copyright 2011-2020 Molecular Matters GmbH, all rights reserved.
 
+#if LC_VERSION == 1
+
 // BEGIN EPIC MOD
 //#include PCH_INCLUDE
 // END EPIC MOD
@@ -872,3 +874,6 @@ void amalgamation::DeleteDatabase(const symbols::ObjPath& objPath)
 	const std::wstring& databasePath = GenerateDatabasePath(objPath);
 	Filesystem::DeleteIfExists(databasePath.c_str());
 }
+
+
+#endif
\ No newline at end of file
diff --git a/Engine/Source/Developer/Windows/LiveCodingServer/Private/External/LC_Amalgamation.h b/Engine/Source/Developer/Windows/LiveCodingServer/Private/External/LC_Amalgamation.h
--- a/Engine/Source/Developer/Windows/LiveCodingServer/Private/External/LC_Amalgamation.h
+++ b/Engine/Source/Developer/Windows/LiveCodingServer/Private/External/LC_Amalgamation.h
@@ -2,6 +2,8 @@
 
 #pragma once
 
+#if LC_VERSION == 1
+
 // BEGIN EPIC MOD
 #include "CoreTypes.h"
 // END EPIC MOD
@@ -47,3 +49,6 @@ namespace amalgamation
 	// deletes a compiland's dependency database from disk
 	void DeleteDatabase(const symbols::ObjPath& objPath);
 }
+
+
+#endif
\ No newline at end of file
diff --git a/Engine/Source/Developer/Windows/LiveCodingServer/Private/External/LC_AppSettings.cpp b/Engine/Source/Developer/Windows/LiveCodingServer/Private/External/LC_AppSettings.cpp
--- a/Engine/Source/Developer/Windows/LiveCodingServer/Private/External/LC_AppSettings.cpp
+++ b/Engine/Source/Developer/Windows/LiveCodingServer/Private/External/LC_AppSettings.cpp
@@ -1,5 +1,7 @@
 // Copyright 2011-2020 Molecular Matters GmbH, all rights reserved.
 
+#if LC_VERSION == 1
+
 // BEGIN EPIC MOD
 //#include PCH_INCLUDE
 // END EPIC MOD
@@ -892,3 +894,6 @@ SettingBool* appSettings::g_ue4EnableNatVisSupport = nullptr;
 
 SettingString* appSettings::g_fastBuildDatabasePath = nullptr;
 SettingString* appSettings::g_fastBuildDllName = nullptr;
+
+
+#endif
\ No newline at end of file
diff --git a/Engine/Source/Developer/Windows/LiveCodingServer/Private/External/LC_AppSettings.h b/Engine/Source/Developer/Windows/LiveCodingServer/Private/External/LC_AppSettings.h
--- a/Engine/Source/Developer/Windows/LiveCodingServer/Private/External/LC_AppSettings.h
+++ b/Engine/Source/Developer/Windows/LiveCodingServer/Private/External/LC_AppSettings.h
@@ -2,6 +2,8 @@
 
 #pragma once
 
+#if LC_VERSION == 1
+
 // BEGIN EPIC MOD
 #include "CoreTypes.h"
 // END EPIC MOD
@@ -124,3 +126,6 @@ namespace appSettings
 	extern SettingString* g_fastBuildDatabasePath;
 	extern SettingString* g_fastBuildDllName;
 }
+
+
+#endif
\ No newline at end of file
diff --git a/Engine/Source/Developer/Windows/LiveCodingServer/Private/External/LC_Assembler.cpp b/Engine/Source/Developer/Windows/LiveCodingServer/Private/External/LC_Assembler.cpp
--- a/Engine/Source/Developer/Windows/LiveCodingServer/Private/External/LC_Assembler.cpp
+++ b/Engine/Source/Developer/Windows/LiveCodingServer/Private/External/LC_Assembler.cpp
@@ -1,5 +1,7 @@
 // Copyright 2011-2020 Molecular Matters GmbH, all rights reserved.
 
+#if LC_VERSION == 1
+
 // BEGIN EPIC MOD
 //#include PCH_INCLUDE
 // END EPIC MOD
@@ -155,3 +157,6 @@ Assembler::SizedInstruction<Assembler::InstructionSize::ADD_RAX_INDIRECT_RIP32>
 {
 	return MakeSizedInstructionWithOffset(Opcode::ADD_RAX_INDIRECT_RIP32, ripOffset);
 }
+
+
+#endif
\ No newline at end of file
diff --git a/Engine/Source/Developer/Windows/LiveCodingServer/Private/External/LC_Assembler.h b/Engine/Source/Developer/Windows/LiveCodingServer/Private/External/LC_Assembler.h
--- a/Engine/Source/Developer/Windows/LiveCodingServer/Private/External/LC_Assembler.h
+++ b/Engine/Source/Developer/Windows/LiveCodingServer/Private/External/LC_Assembler.h
@@ -2,6 +2,8 @@
 
 #pragma once
 
+#if LC_VERSION == 1
+
 // BEGIN EPIC MOD
 #include "CoreTypes.h"
 #include <cinttypes>
@@ -114,3 +116,6 @@ namespace Assembler
 	// Generates an indirect RIP-relative ADD using a 32-bit offset.
 	SizedInstruction<InstructionSize::ADD_RAX_INDIRECT_RIP32> MakeADD32IndirectRIPRelative(int32_t ripOffset);
 }
+
+
+#endif
\ No newline at end of file
diff --git a/Engine/Source/Developer/Windows/LiveCodingServer/Private/External/LC_COMThread.cpp b/Engine/Source/Developer/Windows/LiveCodingServer/Private/External/LC_COMThread.cpp
--- a/Engine/Source/Developer/Windows/LiveCodingServer/Private/External/LC_COMThread.cpp
+++ b/Engine/Source/Developer/Windows/LiveCodingServer/Private/External/LC_COMThread.cpp
@@ -1,5 +1,7 @@
 // Copyright 2011-2020 Molecular Matters GmbH, all rights reserved.
 
+#if LC_VERSION == 1
+
 // BEGIN EPIC MOD
 //#include PCH_INCLUDE
 // END EPIC MOD
@@ -193,4 +195,6 @@ Thread::ReturnValue COMThread::ThreadFunction(void)
 
 // BEGIN EPIC MOD
 #include "Windows/HideWindowsPlatformAtomics.h"
-// END EPIC MOD
\ No newline at end of file
+// END EPIC MOD
+
+#endif
\ No newline at end of file
diff --git a/Engine/Source/Developer/Windows/LiveCodingServer/Private/External/LC_COMThread.h b/Engine/Source/Developer/Windows/LiveCodingServer/Private/External/LC_COMThread.h
--- a/Engine/Source/Developer/Windows/LiveCodingServer/Private/External/LC_COMThread.h
+++ b/Engine/Source/Developer/Windows/LiveCodingServer/Private/External/LC_COMThread.h
@@ -2,6 +2,8 @@
 
 #pragma once
 
+#if LC_VERSION == 1
+
 // BEGIN EPIC MOD
 #include <type_traits>
 // END EPIC MOD
@@ -82,3 +84,6 @@ private:
 	Event m_leaveThreadEvent;
 	Thread::Handle m_internalThread;
 };
+
+
+#endif
\ No newline at end of file
diff --git a/Engine/Source/Developer/Windows/LiveCodingServer/Private/External/LC_ChangeNotification.cpp b/Engine/Source/Developer/Windows/LiveCodingServer/Private/External/LC_ChangeNotification.cpp
--- a/Engine/Source/Developer/Windows/LiveCodingServer/Private/External/LC_ChangeNotification.cpp
+++ b/Engine/Source/Developer/Windows/LiveCodingServer/Private/External/LC_ChangeNotification.cpp
@@ -1,5 +1,7 @@
 // Copyright 2011-2020 Molecular Matters GmbH, all rights reserved.
 
+#if LC_VERSION == 1
+
 // BEGIN EPIC MOD
 //#include PCH_INCLUDE
 // END EPIC MOD
@@ -104,3 +106,6 @@ bool ChangeNotification::WaitForNotification(unsigned int timeoutMs)
 			return false;
 	}
 }
+
+
+#endif
\ No newline at end of file
diff --git a/Engine/Source/Developer/Windows/LiveCodingServer/Private/External/LC_ChangeNotification.h b/Engine/Source/Developer/Windows/LiveCodingServer/Private/External/LC_ChangeNotification.h
--- a/Engine/Source/Developer/Windows/LiveCodingServer/Private/External/LC_ChangeNotification.h
+++ b/Engine/Source/Developer/Windows/LiveCodingServer/Private/External/LC_ChangeNotification.h
@@ -2,6 +2,8 @@
 
 #pragma once
 
+#if LC_VERSION == 1
+
 // BEGIN EPIC MOD
 #include "CoreTypes.h"
 #include "Windows/MinimalWindowsAPI.h"
@@ -27,3 +29,6 @@ private:
 	Windows::HANDLE m_handle;
 	// END EPIC MOD
 };
+
+
+#endif
\ No newline at end of file
diff --git a/Engine/Source/Developer/Windows/LiveCodingServer/Private/External/LC_CodeCave.cpp b/Engine/Source/Developer/Windows/LiveCodingServer/Private/External/LC_CodeCave.cpp
--- a/Engine/Source/Developer/Windows/LiveCodingServer/Private/External/LC_CodeCave.cpp
+++ b/Engine/Source/Developer/Windows/LiveCodingServer/Private/External/LC_CodeCave.cpp
@@ -1,5 +1,7 @@
 // Copyright 2011-2020 Molecular Matters GmbH, all rights reserved.
 
+#if LC_VERSION == 1
+
 // BEGIN EPIC MOD
 //#include PCH_INCLUDE
 // END EPIC MOD
@@ -92,3 +94,6 @@ void CodeCave::Uninstall(void)
 
 	Process::Resume(m_processHandle);
 }
+
+
+#endif
\ No newline at end of file
diff --git a/Engine/Source/Developer/Windows/LiveCodingServer/Private/External/LC_CodeCave.h b/Engine/Source/Developer/Windows/LiveCodingServer/Private/External/LC_CodeCave.h
--- a/Engine/Source/Developer/Windows/LiveCodingServer/Private/External/LC_CodeCave.h
+++ b/Engine/Source/Developer/Windows/LiveCodingServer/Private/External/LC_CodeCave.h
@@ -2,6 +2,8 @@
 
 #pragma once
 
+#if LC_VERSION == 1
+
 // BEGIN EPIC MOD
 #include "CoreTypes.h"
 // END EPIC MOD
@@ -37,3 +39,6 @@ private:
 
 	types::vector<PerThreadData> m_perThreadData;
 };
+
+
+#endif
\ No newline at end of file
diff --git a/Engine/Source/Developer/Windows/LiveCodingServer/Private/External/LC_Coff.cpp b/Engine/Source/Developer/Windows/LiveCodingServer/Private/External/LC_Coff.cpp
--- a/Engine/Source/Developer/Windows/LiveCodingServer/Private/External/LC_Coff.cpp
+++ b/Engine/Source/Developer/Windows/LiveCodingServer/Private/External/LC_Coff.cpp
@@ -1,5 +1,7 @@
 // Copyright 2011-2020 Molecular Matters GmbH, all rights reserved.
 
+#if LC_VERSION == 1
+
 // BEGIN EPIC MOD
 //#include PCH_INCLUDE
 // END EPIC MOD
@@ -1856,3 +1858,6 @@ namespace coff
 		return string::Matches(sectionName, ".msvcjmc");
 	}
 }
+
+
+#endif
\ No newline at end of file
diff --git a/Engine/Source/Developer/Windows/LiveCodingServer/Private/External/LC_Coff.h b/Engine/Source/Developer/Windows/LiveCodingServer/Private/External/LC_Coff.h
--- a/Engine/Source/Developer/Windows/LiveCodingServer/Private/External/LC_Coff.h
+++ b/Engine/Source/Developer/Windows/LiveCodingServer/Private/External/LC_Coff.h
@@ -2,6 +2,8 @@
 
 #pragma once
 
+#if LC_VERSION == 1
+
 // BEGIN EPIC MOD
 #include "CoreTypes.h"
 // END EPIC MOD
@@ -463,3 +465,6 @@ namespace coff
 	bool IsInterestingSymbol(const ImmutableString& name);
 	bool IsMSVCJustMyCodeSection(const char* sectionName);
 }
+
+
+#endif
\ No newline at end of file
diff --git a/Engine/Source/Developer/Windows/LiveCodingServer/Private/External/LC_CoffCache.cpp b/Engine/Source/Developer/Windows/LiveCodingServer/Private/External/LC_CoffCache.cpp
--- a/Engine/Source/Developer/Windows/LiveCodingServer/Private/External/LC_CoffCache.cpp
+++ b/Engine/Source/Developer/Windows/LiveCodingServer/Private/External/LC_CoffCache.cpp
@@ -1,5 +1,7 @@
 // Copyright 2011-2020 Molecular Matters GmbH, all rights reserved.
 
+#if LC_VERSION == 1
+
 // BEGIN EPIC MOD
 //#include PCH_INCLUDE
 // END EPIC MOD
@@ -86,3 +88,6 @@ bool CoffCache<T>::Contains(const ImmutableString& coffIdentifier) const
 // explicit template instantiation
 template class CoffCache<coff::CoffDB>;
 template class CoffCache<coff::ExternalSymbolDB>;
+
+
+#endif
\ No newline at end of file
diff --git a/Engine/Source/Developer/Windows/LiveCodingServer/Private/External/LC_CoffCache.h b/Engine/Source/Developer/Windows/LiveCodingServer/Private/External/LC_CoffCache.h
--- a/Engine/Source/Developer/Windows/LiveCodingServer/Private/External/LC_CoffCache.h
+++ b/Engine/Source/Developer/Windows/LiveCodingServer/Private/External/LC_CoffCache.h
@@ -2,6 +2,8 @@
 
 #pragma once
 
+#if LC_VERSION == 1
+
 // BEGIN EPIC MOD
 #include "CoreTypes.h"
 #include "LC_Coff.h"
@@ -32,3 +34,6 @@ private:
 	typedef types::StringMap<T*> Cache;
 	Cache m_cache;
 };
+
+
+#endif
\ No newline at end of file
diff --git a/Engine/Source/Developer/Windows/LiveCodingServer/Private/External/LC_CoffDetail.cpp b/Engine/Source/Developer/Windows/LiveCodingServer/Private/External/LC_CoffDetail.cpp
--- a/Engine/Source/Developer/Windows/LiveCodingServer/Private/External/LC_CoffDetail.cpp
+++ b/Engine/Source/Developer/Windows/LiveCodingServer/Private/External/LC_CoffDetail.cpp
@@ -1,5 +1,7 @@
 // Copyright 2011-2020 Molecular Matters GmbH, all rights reserved.
 
+#if LC_VERSION == 1
+
 // BEGIN EPIC MOD
 //#include PCH_INCLUDE
 // END EPIC MOD
@@ -109,3 +111,6 @@ namespace coffDetail
 		return relocationCount;
 	}
 }
+
+
+#endif
\ No newline at end of file
diff --git a/Engine/Source/Developer/Windows/LiveCodingServer/Private/External/LC_CoffDetail.h b/Engine/Source/Developer/Windows/LiveCodingServer/Private/External/LC_CoffDetail.h
--- a/Engine/Source/Developer/Windows/LiveCodingServer/Private/External/LC_CoffDetail.h
+++ b/Engine/Source/Developer/Windows/LiveCodingServer/Private/External/LC_CoffDetail.h
@@ -2,6 +2,8 @@
 
 #pragma once
 
+#if LC_VERSION == 1
+
 // BEGIN EPIC MOD
 #include "CoreTypes.h"
 // END EPIC MOD
@@ -299,3 +301,6 @@ namespace coffDetail
 		return size + (size & 1u);
 	}
 }
+
+
+#endif
\ No newline at end of file
diff --git a/Engine/Source/Developer/Windows/LiveCodingServer/Private/External/LC_Compiler.cpp b/Engine/Source/Developer/Windows/LiveCodingServer/Private/External/LC_Compiler.cpp
--- a/Engine/Source/Developer/Windows/LiveCodingServer/Private/External/LC_Compiler.cpp
+++ b/Engine/Source/Developer/Windows/LiveCodingServer/Private/External/LC_Compiler.cpp
@@ -1,5 +1,7 @@
 // Copyright 2011-2020 Molecular Matters GmbH, all rights reserved.
 
+#if LC_VERSION == 1
+
 // BEGIN EPIC MOD
 //#include PCH_INCLUDE
 // END EPIC MOD
@@ -374,4 +376,6 @@ void compiler::AddEnvironmentToCache(const wchar_t* absolutePathToCompilerExe, P
 	CriticalSection::ScopedLock lock(&g_compilerCacheCS);
 	g_compilerEnvironmentCache.Insert(absolutePathToCompilerExe, *environment);
 }
-// END EPIC MOD
\ No newline at end of file
+// END EPIC MOD
+
+#endif
\ No newline at end of file
diff --git a/Engine/Source/Developer/Windows/LiveCodingServer/Private/External/LC_Compiler.h b/Engine/Source/Developer/Windows/LiveCodingServer/Private/External/LC_Compiler.h
--- a/Engine/Source/Developer/Windows/LiveCodingServer/Private/External/LC_Compiler.h
+++ b/Engine/Source/Developer/Windows/LiveCodingServer/Private/External/LC_Compiler.h
@@ -2,6 +2,8 @@
 
 #pragma once
 
+#if LC_VERSION == 1
+
 // BEGIN EPIC MOD
 #include "CoreTypes.h"
 // END EPIC MOD
@@ -48,3 +50,6 @@ namespace compiler
 
 	LinkerType::Enum DetermineLinkerType(const wchar_t* linkerPath);
 }
+
+
+#endif
\ No newline at end of file
diff --git a/Engine/Source/Developer/Windows/LiveCodingServer/Private/External/LC_CompilerOptions.cpp b/Engine/Source/Developer/Windows/LiveCodingServer/Private/External/LC_CompilerOptions.cpp
--- a/Engine/Source/Developer/Windows/LiveCodingServer/Private/External/LC_CompilerOptions.cpp
+++ b/Engine/Source/Developer/Windows/LiveCodingServer/Private/External/LC_CompilerOptions.cpp
@@ -1,5 +1,7 @@
 // Copyright 2011-2020 Molecular Matters GmbH, all rights reserved.
 
+#if LC_VERSION == 1
+
 // BEGIN EPIC MOD
 //#include PCH_INCLUDE
 // END EPIC MOD
@@ -59,3 +61,6 @@ bool compilerOptions::UsesMinimalRebuild(const char* options)
 
 	return string::Contains(options, "-Gm");
 }
+
+
+#endif
\ No newline at end of file
diff --git a/Engine/Source/Developer/Windows/LiveCodingServer/Private/External/LC_CompilerOptions.h b/Engine/Source/Developer/Windows/LiveCodingServer/Private/External/LC_CompilerOptions.h
--- a/Engine/Source/Developer/Windows/LiveCodingServer/Private/External/LC_CompilerOptions.h
+++ b/Engine/Source/Developer/Windows/LiveCodingServer/Private/External/LC_CompilerOptions.h
@@ -2,6 +2,8 @@
 
 #pragma once
 
+#if LC_VERSION == 1
+
 // BEGIN EPIC MOD
 #include "CoreTypes.h"
 #include <string>
@@ -16,3 +18,6 @@ namespace compilerOptions
 	bool UsesC7DebugFormat(const char* options);
 	bool UsesMinimalRebuild(const char* options);
 }
+
+
+#endif
\ No newline at end of file
diff --git a/Engine/Source/Developer/Windows/LiveCodingServer/Private/External/LC_DiaSymbolName.cpp b/Engine/Source/Developer/Windows/LiveCodingServer/Private/External/LC_DiaSymbolName.cpp
--- a/Engine/Source/Developer/Windows/LiveCodingServer/Private/External/LC_DiaSymbolName.cpp
+++ b/Engine/Source/Developer/Windows/LiveCodingServer/Private/External/LC_DiaSymbolName.cpp
@@ -1,5 +1,7 @@
 // Copyright 2011-2020 Molecular Matters GmbH, all rights reserved.
 
+#if LC_VERSION == 1
+
 // BEGIN EPIC MOD
 //#include PCH_INCLUDE
 // END EPIC MOD
@@ -26,3 +28,6 @@ dia::SymbolName::SymbolName(SymbolName&& other)
 {
 	other.m_name = nullptr;
 }
+
+
+#endif
\ No newline at end of file
diff --git a/Engine/Source/Developer/Windows/LiveCodingServer/Private/External/LC_DiaSymbolName.h b/Engine/Source/Developer/Windows/LiveCodingServer/Private/External/LC_DiaSymbolName.h
--- a/Engine/Source/Developer/Windows/LiveCodingServer/Private/External/LC_DiaSymbolName.h
+++ b/Engine/Source/Developer/Windows/LiveCodingServer/Private/External/LC_DiaSymbolName.h
@@ -2,6 +2,8 @@
 
 #pragma once
 
+#if LC_VERSION == 1
+
 // BEGIN EPIC MOD
 #include "CoreTypes.h"
 #include "LC_Platform.h"
@@ -35,3 +37,6 @@ namespace dia
 		LC_DISABLE_MOVE_ASSIGNMENT(SymbolName);
 	};
 }
+
+
+#endif
\ No newline at end of file
diff --git a/Engine/Source/Developer/Windows/LiveCodingServer/Private/External/LC_DiaUtil.cpp b/Engine/Source/Developer/Windows/LiveCodingServer/Private/External/LC_DiaUtil.cpp
--- a/Engine/Source/Developer/Windows/LiveCodingServer/Private/External/LC_DiaUtil.cpp
+++ b/Engine/Source/Developer/Windows/LiveCodingServer/Private/External/LC_DiaUtil.cpp
@@ -1,5 +1,7 @@
 // Copyright 2011-2020 Molecular Matters GmbH, all rights reserved.
 
+#if LC_VERSION == 1
+
 // BEGIN EPIC MOD
 //#include PCH_INCLUDE
 // END EPIC MOD
@@ -376,3 +378,6 @@ dia::SymbolName dia::FindSourceFileByRVA(IDiaSession* session, uint32_t rva)
 
 	return SymbolName(nullptr);
 }
+
+
+#endif
\ No newline at end of file
diff --git a/Engine/Source/Developer/Windows/LiveCodingServer/Private/External/LC_DiaUtil.h b/Engine/Source/Developer/Windows/LiveCodingServer/Private/External/LC_DiaUtil.h
--- a/Engine/Source/Developer/Windows/LiveCodingServer/Private/External/LC_DiaUtil.h
+++ b/Engine/Source/Developer/Windows/LiveCodingServer/Private/External/LC_DiaUtil.h
@@ -2,6 +2,8 @@
 
 #pragma once
 
+#if LC_VERSION == 1
+
 // BEGIN EPIC MOD
 #include "CoreTypes.h"
 // END EPIC MOD
@@ -45,3 +47,6 @@ namespace dia
 	uint32_t FindLineNumberByRVA(IDiaSession* session, uint32_t rva);
 	SymbolName FindSourceFileByRVA(IDiaSession* session, uint32_t rva);
 }
+
+
+#endif
\ No newline at end of file
diff --git a/Engine/Source/Developer/Windows/LiveCodingServer/Private/External/LC_DiaVariant.cpp b/Engine/Source/Developer/Windows/LiveCodingServer/Private/External/LC_DiaVariant.cpp
--- a/Engine/Source/Developer/Windows/LiveCodingServer/Private/External/LC_DiaVariant.cpp
+++ b/Engine/Source/Developer/Windows/LiveCodingServer/Private/External/LC_DiaVariant.cpp
@@ -1,5 +1,7 @@
 // Copyright 2011-2020 Molecular Matters GmbH, all rights reserved.
 
+#if LC_VERSION == 1
+
 // BEGIN EPIC MOD
 //#include PCH_INCLUDE
 // END EPIC MOD
@@ -42,3 +44,6 @@ dia::Variant::Variant(Variant&& other)
 {
 	other.m_str = nullptr;
 }
+
+
+#endif
\ No newline at end of file
diff --git a/Engine/Source/Developer/Windows/LiveCodingServer/Private/External/LC_DiaVariant.h b/Engine/Source/Developer/Windows/LiveCodingServer/Private/External/LC_DiaVariant.h
--- a/Engine/Source/Developer/Windows/LiveCodingServer/Private/External/LC_DiaVariant.h
+++ b/Engine/Source/Developer/Windows/LiveCodingServer/Private/External/LC_DiaVariant.h
@@ -2,6 +2,8 @@
 
 #pragma once
 
+#if LC_VERSION == 1
+
 // BEGIN EPIC MOD
 #include "CoreTypes.h"
 #include "LC_Platform.h"
@@ -36,3 +38,6 @@ namespace dia
 		LC_DISABLE_MOVE_ASSIGNMENT(Variant);
 	};
 }
+
+
+#endif
\ No newline at end of file
diff --git a/Engine/Source/Developer/Windows/LiveCodingServer/Private/External/LC_DirectoryCache.cpp b/Engine/Source/Developer/Windows/LiveCodingServer/Private/External/LC_DirectoryCache.cpp
--- a/Engine/Source/Developer/Windows/LiveCodingServer/Private/External/LC_DirectoryCache.cpp
+++ b/Engine/Source/Developer/Windows/LiveCodingServer/Private/External/LC_DirectoryCache.cpp
@@ -1,5 +1,7 @@
 // Copyright 2011-2020 Molecular Matters GmbH, all rights reserved.
 
+#if LC_VERSION == 1
+
 // BEGIN EPIC MOD
 //#include PCH_INCLUDE
 // END EPIC MOD
@@ -62,3 +64,6 @@ void DirectoryCache::RestartNotifications(void)
 		dir->hadChange = false;
 	}
 }
+
+
+#endif
\ No newline at end of file
diff --git a/Engine/Source/Developer/Windows/LiveCodingServer/Private/External/LC_DirectoryCache.h b/Engine/Source/Developer/Windows/LiveCodingServer/Private/External/LC_DirectoryCache.h
--- a/Engine/Source/Developer/Windows/LiveCodingServer/Private/External/LC_DirectoryCache.h
+++ b/Engine/Source/Developer/Windows/LiveCodingServer/Private/External/LC_DirectoryCache.h
@@ -2,6 +2,8 @@
 
 #pragma once
 
+#if LC_VERSION == 1
+
 // BEGIN EPIC MOD
 #include "CoreTypes.h"
 #include "LC_Types.h"
@@ -28,3 +30,6 @@ public:
 private:
 	types::unordered_map<std::wstring, Directory*> m_directories;
 };
+
+
+#endif
\ No newline at end of file
diff --git a/Engine/Source/Developer/Windows/LiveCodingServer/Private/External/LC_Disassembler.cpp b/Engine/Source/Developer/Windows/LiveCodingServer/Private/External/LC_Disassembler.cpp
--- a/Engine/Source/Developer/Windows/LiveCodingServer/Private/External/LC_Disassembler.cpp
+++ b/Engine/Source/Developer/Windows/LiveCodingServer/Private/External/LC_Disassembler.cpp
@@ -1,5 +1,7 @@
 // Copyright 2011-2020 Molecular Matters GmbH, all rights reserved.
 
+#if LC_VERSION == 1
+
 // BEGIN EPIC MOD
 //#include PCH_INCLUDE
 // END EPIC MOD
@@ -71,3 +73,6 @@ bool Disassembler::IsLEA(uint16_t opcode)
 {
 	return (opcode == I_LEA);
 }
+
+
+#endif
\ No newline at end of file
diff --git a/Engine/Source/Developer/Windows/LiveCodingServer/Private/External/LC_Disassembler.h b/Engine/Source/Developer/Windows/LiveCodingServer/Private/External/LC_Disassembler.h
--- a/Engine/Source/Developer/Windows/LiveCodingServer/Private/External/LC_Disassembler.h
+++ b/Engine/Source/Developer/Windows/LiveCodingServer/Private/External/LC_Disassembler.h
@@ -2,6 +2,8 @@
 
 #pragma once
 
+#if LC_VERSION == 1
+
 // BEGIN EPIC MOD
 #include "CoreTypes.h"
 // END EPIC MOD
@@ -37,3 +39,6 @@ namespace Disassembler
 		};
 	};
 }
+
+
+#endif
\ No newline at end of file
diff --git a/Engine/Source/Developer/Windows/LiveCodingServer/Private/External/LC_DuplexPipeServer.cpp b/Engine/Source/Developer/Windows/LiveCodingServer/Private/External/LC_DuplexPipeServer.cpp
--- a/Engine/Source/Developer/Windows/LiveCodingServer/Private/External/LC_DuplexPipeServer.cpp
+++ b/Engine/Source/Developer/Windows/LiveCodingServer/Private/External/LC_DuplexPipeServer.cpp
@@ -1,5 +1,7 @@
 // Copyright 2011-2020 Molecular Matters GmbH, all rights reserved.
 
+#if LC_VERSION == 1
+
 // BEGIN EPIC MOD
 //#include PCH_INCLUDE
 // END EPIC MOD
@@ -62,3 +64,6 @@ void DuplexPipeServer::Disconnect(void)
 	::FlushFileBuffers(m_pipe);
 	::DisconnectNamedPipe(m_pipe);
 }
+
+
+#endif
\ No newline at end of file
diff --git a/Engine/Source/Developer/Windows/LiveCodingServer/Private/External/LC_DuplexPipeServer.h b/Engine/Source/Developer/Windows/LiveCodingServer/Private/External/LC_DuplexPipeServer.h
--- a/Engine/Source/Developer/Windows/LiveCodingServer/Private/External/LC_DuplexPipeServer.h
+++ b/Engine/Source/Developer/Windows/LiveCodingServer/Private/External/LC_DuplexPipeServer.h
@@ -2,6 +2,8 @@
 
 #pragma once
 
+#if LC_VERSION == 1
+
 // BEGIN EPIC MOD
 #include "CoreTypes.h"
 // END EPIC MOD
@@ -15,3 +17,6 @@ public:
 	bool WaitForClient(void);
 	void Disconnect(void);
 };
+
+
+#endif
\ No newline at end of file
diff --git a/Engine/Source/Developer/Windows/LiveCodingServer/Private/External/LC_ExecutablePatcher.cpp b/Engine/Source/Developer/Windows/LiveCodingServer/Private/External/LC_ExecutablePatcher.cpp
--- a/Engine/Source/Developer/Windows/LiveCodingServer/Private/External/LC_ExecutablePatcher.cpp
+++ b/Engine/Source/Developer/Windows/LiveCodingServer/Private/External/LC_ExecutablePatcher.cpp
@@ -1,5 +1,7 @@
 // Copyright 2011-2020 Molecular Matters GmbH, all rights reserved.
 
+#if LC_VERSION == 1
+
 // BEGIN EPIC MOD
 //#include PCH_INCLUDE
 // END EPIC MOD
@@ -84,3 +86,6 @@ void ExecutablePatcher::RestoreEntryPoint(Process::Handle processHandle, void* m
 		Process::WriteProcessMemory(processHandle, address, m_originalCode[i]);
 	}
 }
+
+
+#endif
\ No newline at end of file
diff --git a/Engine/Source/Developer/Windows/LiveCodingServer/Private/External/LC_ExecutablePatcher.h b/Engine/Source/Developer/Windows/LiveCodingServer/Private/External/LC_ExecutablePatcher.h
--- a/Engine/Source/Developer/Windows/LiveCodingServer/Private/External/LC_ExecutablePatcher.h
+++ b/Engine/Source/Developer/Windows/LiveCodingServer/Private/External/LC_ExecutablePatcher.h
@@ -2,6 +2,8 @@
 
 #pragma once
 
+#if LC_VERSION == 1
+
 // BEGIN EPIC MOD
 #include "CoreTypes.h"
 // END EPIC MOD
@@ -43,3 +45,6 @@ public:
 private:
 	uint8_t m_originalCode[INJECTED_CODE_SIZE] = {};
 };
+
+
+#endif
\ No newline at end of file
diff --git a/Engine/Source/Developer/Windows/LiveCodingServer/Private/External/LC_FileAttributeCache.cpp b/Engine/Source/Developer/Windows/LiveCodingServer/Private/External/LC_FileAttributeCache.cpp
--- a/Engine/Source/Developer/Windows/LiveCodingServer/Private/External/LC_FileAttributeCache.cpp
+++ b/Engine/Source/Developer/Windows/LiveCodingServer/Private/External/LC_FileAttributeCache.cpp
@@ -1,5 +1,7 @@
 // Copyright 2011-2020 Molecular Matters GmbH, all rights reserved.
 
+#if LC_VERSION == 1
+
 // BEGIN EPIC MOD
 //#include PCH_INCLUDE
 // END EPIC MOD
@@ -37,3 +39,6 @@ size_t FileAttributeCache::GetEntryCount(void) const
 {
 	return m_data.size();
 }
+
+
+#endif
\ No newline at end of file
diff --git a/Engine/Source/Developer/Windows/LiveCodingServer/Private/External/LC_FileAttributeCache.h b/Engine/Source/Developer/Windows/LiveCodingServer/Private/External/LC_FileAttributeCache.h
--- a/Engine/Source/Developer/Windows/LiveCodingServer/Private/External/LC_FileAttributeCache.h
+++ b/Engine/Source/Developer/Windows/LiveCodingServer/Private/External/LC_FileAttributeCache.h
@@ -2,6 +2,8 @@
 
 #pragma once
 
+#if LC_VERSION == 1
+
 // BEGIN EPIC MOD
 #include "CoreTypes.h"
 #include "LC_Types.h"
@@ -35,3 +37,6 @@ private:
 	typedef types::unordered_map_with_hash<std::wstring, Data, Hasher> Cache;
 	Cache m_data;
 };
+
+
+#endif
\ No newline at end of file
diff --git a/Engine/Source/Developer/Windows/LiveCodingServer/Private/External/LC_Filesystem.cpp b/Engine/Source/Developer/Windows/LiveCodingServer/Private/External/LC_Filesystem.cpp
--- a/Engine/Source/Developer/Windows/LiveCodingServer/Private/External/LC_Filesystem.cpp
+++ b/Engine/Source/Developer/Windows/LiveCodingServer/Private/External/LC_Filesystem.cpp
@@ -1,5 +1,7 @@
 // Copyright 2011-2020 Molecular Matters GmbH, all rights reserved.
 
+#if LC_VERSION == 1
+
 // BEGIN EPIC MOD
 //#include PCH_INCLUDE
 // END EPIC MOD
@@ -487,3 +489,6 @@ types::vector<Filesystem::Path> Filesystem::EnumerateFiles(const wchar_t* direct
 
 	return foundFiles;
 }
+
+
+#endif
\ No newline at end of file
diff --git a/Engine/Source/Developer/Windows/LiveCodingServer/Private/External/LC_Filesystem.h b/Engine/Source/Developer/Windows/LiveCodingServer/Private/External/LC_Filesystem.h
--- a/Engine/Source/Developer/Windows/LiveCodingServer/Private/External/LC_Filesystem.h
+++ b/Engine/Source/Developer/Windows/LiveCodingServer/Private/External/LC_Filesystem.h
@@ -2,6 +2,8 @@
 
 #pragma once
 
+#if LC_VERSION == 1
+
 // BEGIN EPIC MOD
 #include "CoreTypes.h"
 // END EPIC MOD
@@ -79,3 +81,6 @@ namespace Filesystem
 	// Returns an array of filenames relative to the given directory, e.g. file1.txt, file2.txt, dir\file.txt, dir\subDir\file.txt.
 	types::vector<Path> EnumerateFiles(const wchar_t* directory);
 }
+
+
+#endif
\ No newline at end of file
diff --git a/Engine/Source/Developer/Windows/LiveCodingServer/Private/External/LC_FunctionPatcher.cpp b/Engine/Source/Developer/Windows/LiveCodingServer/Private/External/LC_FunctionPatcher.cpp
--- a/Engine/Source/Developer/Windows/LiveCodingServer/Private/External/LC_FunctionPatcher.cpp
+++ b/Engine/Source/Developer/Windows/LiveCodingServer/Private/External/LC_FunctionPatcher.cpp
@@ -1,5 +1,7 @@
 // Copyright 2011-2020 Molecular Matters GmbH, all rights reserved.
 
+#if LC_VERSION == 1
+
 // BEGIN EPIC MOD
 //#include PCH_INCLUDE
 // END EPIC MOD
@@ -440,3 +442,6 @@ bool functions::IsValidRecord(const LibraryRecord& record)
 {
 	return (record.wholeInstructionSize != 0u);
 }
+
+
+#endif
\ No newline at end of file
diff --git a/Engine/Source/Developer/Windows/LiveCodingServer/Private/External/LC_FunctionPatcher.h b/Engine/Source/Developer/Windows/LiveCodingServer/Private/External/LC_FunctionPatcher.h
--- a/Engine/Source/Developer/Windows/LiveCodingServer/Private/External/LC_FunctionPatcher.h
+++ b/Engine/Source/Developer/Windows/LiveCodingServer/Private/External/LC_FunctionPatcher.h
@@ -2,6 +2,8 @@
 
 #pragma once
 
+#if LC_VERSION == 1
+
 // BEGIN EPIC MOD
 #include "CoreTypes.h"
 // END EPIC MOD
@@ -80,3 +82,6 @@ namespace functions
 	bool IsValidRecord(const Record& record);
 	bool IsValidRecord(const LibraryRecord& record);
 }
+
+
+#endif
\ No newline at end of file
diff --git a/Engine/Source/Developer/Windows/LiveCodingServer/Private/External/LC_GrowingMemoryBlock.cpp b/Engine/Source/Developer/Windows/LiveCodingServer/Private/External/LC_GrowingMemoryBlock.cpp
--- a/Engine/Source/Developer/Windows/LiveCodingServer/Private/External/LC_GrowingMemoryBlock.cpp
+++ b/Engine/Source/Developer/Windows/LiveCodingServer/Private/External/LC_GrowingMemoryBlock.cpp
@@ -1,5 +1,7 @@
 // Copyright 2011-2020 Molecular Matters GmbH, all rights reserved.
 
+#if LC_VERSION == 1
+
 // BEGIN EPIC MOD
 //#include PCH_INCLUDE
 // END EPIC MOD
@@ -52,3 +54,6 @@ bool GrowingMemoryBlock::Insert(const void* data, size_t size)
 
 	return true;
 }
+
+
+#endif
\ No newline at end of file
diff --git a/Engine/Source/Developer/Windows/LiveCodingServer/Private/External/LC_GrowingMemoryBlock.h b/Engine/Source/Developer/Windows/LiveCodingServer/Private/External/LC_GrowingMemoryBlock.h
--- a/Engine/Source/Developer/Windows/LiveCodingServer/Private/External/LC_GrowingMemoryBlock.h
+++ b/Engine/Source/Developer/Windows/LiveCodingServer/Private/External/LC_GrowingMemoryBlock.h
@@ -2,6 +2,8 @@
 
 #pragma once
 
+#if LC_VERSION == 1
+
 // BEGIN EPIC MOD
 #include "CoreTypes.h"
 #include "LC_Platform.h"
@@ -36,3 +38,6 @@ private:
 	LC_DISABLE_ASSIGNMENT(GrowingMemoryBlock);
 	LC_DISABLE_MOVE_ASSIGNMENT(GrowingMemoryBlock);
 };
+
+
+#endif
\ No newline at end of file
diff --git a/Engine/Source/Developer/Windows/LiveCodingServer/Private/External/LC_Hook.cpp b/Engine/Source/Developer/Windows/LiveCodingServer/Private/External/LC_Hook.cpp
--- a/Engine/Source/Developer/Windows/LiveCodingServer/Private/External/LC_Hook.cpp
+++ b/Engine/Source/Developer/Windows/LiveCodingServer/Private/External/LC_Hook.cpp
@@ -1,5 +1,7 @@
 // Copyright 2011-2020 Molecular Matters GmbH, all rights reserved.
 
+#if LC_VERSION == 1
+
 // BEGIN EPIC MOD
 //#include PCH_INCLUDE
 // END EPIC MOD
@@ -29,3 +31,6 @@ uint32_t hook::FindLastInSection(const symbols::ImageSectionDB* imageSectionDb,
 
 	return 0u;
 }
+
+
+#endif
\ No newline at end of file
diff --git a/Engine/Source/Developer/Windows/LiveCodingServer/Private/External/LC_Hook.h b/Engine/Source/Developer/Windows/LiveCodingServer/Private/External/LC_Hook.h
--- a/Engine/Source/Developer/Windows/LiveCodingServer/Private/External/LC_Hook.h
+++ b/Engine/Source/Developer/Windows/LiveCodingServer/Private/External/LC_Hook.h
@@ -2,6 +2,8 @@
 
 #pragma once
 
+#if LC_VERSION == 1
+
 // BEGIN EPIC MOD
 #include "CoreTypes.h"
 #include <stdint.h>
@@ -60,3 +62,6 @@ namespace hook
 		}
 	}
 }
+
+
+#endif
\ No newline at end of file
diff --git a/Engine/Source/Developer/Windows/LiveCodingServer/Private/External/LC_ImmutableString.cpp b/Engine/Source/Developer/Windows/LiveCodingServer/Private/External/LC_ImmutableString.cpp
--- a/Engine/Source/Developer/Windows/LiveCodingServer/Private/External/LC_ImmutableString.cpp
+++ b/Engine/Source/Developer/Windows/LiveCodingServer/Private/External/LC_ImmutableString.cpp
@@ -1,5 +1,7 @@
 // Copyright 2011-2020 Molecular Matters GmbH, all rights reserved.
 
+#if LC_VERSION == 1
+
 // BEGIN EPIC MOD
 //#include PCH_INCLUDE
 // END EPIC MOD
@@ -197,3 +199,6 @@ void ImmutableString::Free(void)
 		LC_FREE(&g_immutableStringAllocator, const_cast<char*>(m_data.longString.str), (GetLength() + 1u) * sizeof(char));
 	}
 }
+
+
+#endif
\ No newline at end of file
diff --git a/Engine/Source/Developer/Windows/LiveCodingServer/Private/External/LC_ImmutableString.h b/Engine/Source/Developer/Windows/LiveCodingServer/Private/External/LC_ImmutableString.h
--- a/Engine/Source/Developer/Windows/LiveCodingServer/Private/External/LC_ImmutableString.h
+++ b/Engine/Source/Developer/Windows/LiveCodingServer/Private/External/LC_ImmutableString.h
@@ -2,6 +2,8 @@
 
 #pragma once
 
+#if LC_VERSION == 1
+
 // BEGIN EPIC MOD
 #include "CoreTypes.h"
 // END EPIC MOD
@@ -184,3 +186,6 @@ private:
 
 	Data m_data;
 };
+
+
+#endif
\ No newline at end of file
diff --git a/Engine/Source/Developer/Windows/LiveCodingServer/Private/External/LC_Key.cpp b/Engine/Source/Developer/Windows/LiveCodingServer/Private/External/LC_Key.cpp
--- a/Engine/Source/Developer/Windows/LiveCodingServer/Private/External/LC_Key.cpp
+++ b/Engine/Source/Developer/Windows/LiveCodingServer/Private/External/LC_Key.cpp
@@ -1,5 +1,7 @@
 // Copyright 2011-2020 Molecular Matters GmbH, all rights reserved.
 
+#if LC_VERSION == 1
+
 // BEGIN EPIC MOD
 //#include PCH_INCLUDE
 // END EPIC MOD
@@ -50,3 +52,6 @@ bool input::Key::WentDown(void) const
 {
 	return (m_isPressed && !m_wasPressed);
 }
+
+
+#endif
\ No newline at end of file
diff --git a/Engine/Source/Developer/Windows/LiveCodingServer/Private/External/LC_Key.h b/Engine/Source/Developer/Windows/LiveCodingServer/Private/External/LC_Key.h
--- a/Engine/Source/Developer/Windows/LiveCodingServer/Private/External/LC_Key.h
+++ b/Engine/Source/Developer/Windows/LiveCodingServer/Private/External/LC_Key.h
@@ -2,6 +2,8 @@
 
 #pragma once
 
+#if LC_VERSION == 1
+
 // BEGIN EPIC MOD
 #include "CoreTypes.h"
 // END EPIC MOD
@@ -26,3 +28,6 @@ namespace input
 		bool m_wasPressed;
 	};
 }
+
+
+#endif
\ No newline at end of file
diff --git a/Engine/Source/Developer/Windows/LiveCodingServer/Private/External/LC_LiveModule.cpp b/Engine/Source/Developer/Windows/LiveCodingServer/Private/External/LC_LiveModule.cpp
--- a/Engine/Source/Developer/Windows/LiveCodingServer/Private/External/LC_LiveModule.cpp
+++ b/Engine/Source/Developer/Windows/LiveCodingServer/Private/External/LC_LiveModule.cpp
@@ -1,5 +1,7 @@
 // Copyright 2011-2020 Molecular Matters GmbH, all rights reserved.
 
+#if LC_VERSION == 1
+
 // BEGIN EPIC MOD
 //#include PCH_INCLUDE
 // END EPIC MOD
@@ -4933,4 +4935,6 @@ bool LiveModule::IsModifiedSource(const wchar_t* sourceFile) const
 	}
 	return true;
 }
-// END EPIC MOD
\ No newline at end of file
+// END EPIC MOD
+
+#endif
\ No newline at end of file
diff --git a/Engine/Source/Developer/Windows/LiveCodingServer/Private/External/LC_LiveModule.h b/Engine/Source/Developer/Windows/LiveCodingServer/Private/External/LC_LiveModule.h
--- a/Engine/Source/Developer/Windows/LiveCodingServer/Private/External/LC_LiveModule.h
+++ b/Engine/Source/Developer/Windows/LiveCodingServer/Private/External/LC_LiveModule.h
@@ -2,6 +2,8 @@
 
 #pragma once
 
+#if LC_VERSION == 1
+
 // BEGIN EPIC MOD
 #include "CoreTypes.h"
 // END EPIC MOD
@@ -156,3 +158,6 @@ private:
 	types::StringMap<uint32_t> m_objFileToCompilandId;
 	// END EPIC MOD
 };
+
+
+#endif
\ No newline at end of file
diff --git a/Engine/Source/Developer/Windows/LiveCodingServer/Private/External/LC_LiveProcess.cpp b/Engine/Source/Developer/Windows/LiveCodingServer/Private/External/LC_LiveProcess.cpp
--- a/Engine/Source/Developer/Windows/LiveCodingServer/Private/External/LC_LiveProcess.cpp
+++ b/Engine/Source/Developer/Windows/LiveCodingServer/Private/External/LC_LiveProcess.cpp
@@ -1,5 +1,7 @@
 // Copyright 2011-2020 Molecular Matters GmbH, all rights reserved.
 
+#if LC_VERSION == 1
+
 // BEGIN EPIC MOD
 //#include PCH_INCLUDE
 // END EPIC MOD
@@ -383,4 +385,6 @@ Windows::HMODULE LiveProcess::GetLazyLoadedModuleBase(const std::wstring& module
 		return iter->second.m_moduleBase;
 	}
 }
-// END EPIC MOD
\ No newline at end of file
+// END EPIC MOD
+
+#endif
\ No newline at end of file
diff --git a/Engine/Source/Developer/Windows/LiveCodingServer/Private/External/LC_LiveProcess.h b/Engine/Source/Developer/Windows/LiveCodingServer/Private/External/LC_LiveProcess.h
--- a/Engine/Source/Developer/Windows/LiveCodingServer/Private/External/LC_LiveProcess.h
+++ b/Engine/Source/Developer/Windows/LiveCodingServer/Private/External/LC_LiveProcess.h
@@ -2,6 +2,8 @@
 
 #pragma once
 
+#if LC_VERSION == 1
+
 #include "LC_ProcessTypes.h"
 #include "LC_ThreadTypes.h"
 #include "LC_Executable.h"
@@ -203,3 +205,6 @@ private:
 	LC_DISABLE_ASSIGNMENT(LiveProcess);
 	LC_DISABLE_MOVE_ASSIGNMENT(LiveProcess);
 };
+
+
+#endif
\ No newline at end of file
diff --git a/Engine/Source/Developer/Windows/LiveCodingServer/Private/External/LC_Mallocator.cpp b/Engine/Source/Developer/Windows/LiveCodingServer/Private/External/LC_Mallocator.cpp
--- a/Engine/Source/Developer/Windows/LiveCodingServer/Private/External/LC_Mallocator.cpp
+++ b/Engine/Source/Developer/Windows/LiveCodingServer/Private/External/LC_Mallocator.cpp
@@ -1,5 +1,7 @@
 // Copyright 2011-2020 Molecular Matters GmbH, all rights reserved.
 
+#if LC_VERSION == 1
+
 // BEGIN EPIC MOD
 //#include PCH_INCLUDE
 // END EPIC MOD
@@ -58,3 +60,6 @@ const AllocatorStats& Mallocator::GetStats(void) const
 {
 	return m_stats;
 }
+
+
+#endif
\ No newline at end of file
diff --git a/Engine/Source/Developer/Windows/LiveCodingServer/Private/External/LC_Mallocator.h b/Engine/Source/Developer/Windows/LiveCodingServer/Private/External/LC_Mallocator.h
--- a/Engine/Source/Developer/Windows/LiveCodingServer/Private/External/LC_Mallocator.h
+++ b/Engine/Source/Developer/Windows/LiveCodingServer/Private/External/LC_Mallocator.h
@@ -2,6 +2,8 @@
 
 #pragma once
 
+#if LC_VERSION == 1
+
 // BEGIN EPIC MOD
 #include "CoreTypes.h"
 #include "Windows/MinimalWindowsAPI.h"
@@ -30,3 +32,6 @@ private:
 	size_t m_alignment;
 	AllocatorStats m_stats;
 };
+
+
+#endif
\ No newline at end of file
diff --git a/Engine/Source/Developer/Windows/LiveCodingServer/Private/External/LC_MemoryBlock.cpp b/Engine/Source/Developer/Windows/LiveCodingServer/Private/External/LC_MemoryBlock.cpp
--- a/Engine/Source/Developer/Windows/LiveCodingServer/Private/External/LC_MemoryBlock.cpp
+++ b/Engine/Source/Developer/Windows/LiveCodingServer/Private/External/LC_MemoryBlock.cpp
@@ -1,5 +1,7 @@
 // Copyright 2011-2020 Molecular Matters GmbH, all rights reserved.
 
+#if LC_VERSION == 1
+
 // BEGIN EPIC MOD
 //#include PCH_INCLUDE
 // END EPIC MOD
@@ -18,3 +20,6 @@ MemoryBlock::~MemoryBlock(void)
 {
 	free(m_data);
 }
+
+
+#endif
\ No newline at end of file
diff --git a/Engine/Source/Developer/Windows/LiveCodingServer/Private/External/LC_MemoryBlock.h b/Engine/Source/Developer/Windows/LiveCodingServer/Private/External/LC_MemoryBlock.h
--- a/Engine/Source/Developer/Windows/LiveCodingServer/Private/External/LC_MemoryBlock.h
+++ b/Engine/Source/Developer/Windows/LiveCodingServer/Private/External/LC_MemoryBlock.h
@@ -2,6 +2,8 @@
 
 #pragma once
 
+#if LC_VERSION == 1
+
 // BEGIN EPIC MOD
 #include "CoreTypes.h"
 #include "LC_Platform.h"
@@ -33,3 +35,6 @@ private:
 	LC_DISABLE_ASSIGNMENT(MemoryBlock);
 	LC_DISABLE_MOVE_ASSIGNMENT(MemoryBlock);
 };
+
+
+#endif
\ No newline at end of file
diff --git a/Engine/Source/Developer/Windows/LiveCodingServer/Private/External/LC_MicroAllocator.cpp b/Engine/Source/Developer/Windows/LiveCodingServer/Private/External/LC_MicroAllocator.cpp
--- a/Engine/Source/Developer/Windows/LiveCodingServer/Private/External/LC_MicroAllocator.cpp
+++ b/Engine/Source/Developer/Windows/LiveCodingServer/Private/External/LC_MicroAllocator.cpp
@@ -1,5 +1,7 @@
 // Copyright 2011-2020 Molecular Matters GmbH, all rights reserved.
 
+#if LC_VERSION == 1
+
 // BEGIN EPIC MOD
 //#include PCH_INCLUDE
 // END EPIC MOD
@@ -87,3 +89,6 @@ void MicroAllocator::PrintStats(void) const
 	stats.Merge(m_blockAllocator.GetStats());
 	stats.Print(m_name);
 }
+
+
+#endif
\ No newline at end of file
diff --git a/Engine/Source/Developer/Windows/LiveCodingServer/Private/External/LC_MicroAllocator.h b/Engine/Source/Developer/Windows/LiveCodingServer/Private/External/LC_MicroAllocator.h
--- a/Engine/Source/Developer/Windows/LiveCodingServer/Private/External/LC_MicroAllocator.h
+++ b/Engine/Source/Developer/Windows/LiveCodingServer/Private/External/LC_MicroAllocator.h
@@ -2,6 +2,8 @@
 
 #pragma once
 
+#if LC_VERSION == 1
+
 // BEGIN EPIC MOD
 #include "CoreTypes.h"
 // END EPIC MOD
@@ -32,3 +34,6 @@ private:
 	Mallocator m_blockAllocator;
 	const char* m_name;
 };
+
+
+#endif
\ No newline at end of file
diff --git a/Engine/Source/Developer/Windows/LiveCodingServer/Private/External/LC_ModuleCache.cpp b/Engine/Source/Developer/Windows/LiveCodingServer/Private/External/LC_ModuleCache.cpp
--- a/Engine/Source/Developer/Windows/LiveCodingServer/Private/External/LC_ModuleCache.cpp
+++ b/Engine/Source/Developer/Windows/LiveCodingServer/Private/External/LC_ModuleCache.cpp
@@ -1,5 +1,7 @@
 // Copyright 2011-2020 Molecular Matters GmbH, all rights reserved.
 
+#if LC_VERSION == 1
+
 // BEGIN EPIC MOD
 //#include PCH_INCLUDE
 // END EPIC MOD
@@ -164,3 +166,6 @@ types::vector<void*> ModuleCache::GatherModuleBases(Process::Id processId) const
 
 	return result;
 }
+
+
+#endif
\ No newline at end of file
diff --git a/Engine/Source/Developer/Windows/LiveCodingServer/Private/External/LC_ModuleCache.h b/Engine/Source/Developer/Windows/LiveCodingServer/Private/External/LC_ModuleCache.h
--- a/Engine/Source/Developer/Windows/LiveCodingServer/Private/External/LC_ModuleCache.h
+++ b/Engine/Source/Developer/Windows/LiveCodingServer/Private/External/LC_ModuleCache.h
@@ -2,6 +2,8 @@
 
 #pragma once
 
+#if LC_VERSION == 1
+
 // BEGIN EPIC MOD
 #include "CoreTypes.h"
 // END EPIC MOD
@@ -104,3 +106,6 @@ private:
 	mutable CriticalSection m_cs;
 	types::vector<Data> m_cache;
 };
+
+
+#endif
\ No newline at end of file
diff --git a/Engine/Source/Developer/Windows/LiveCodingServer/Private/External/LC_ModulePatch.cpp b/Engine/Source/Developer/Windows/LiveCodingServer/Private/External/LC_ModulePatch.cpp
--- a/Engine/Source/Developer/Windows/LiveCodingServer/Private/External/LC_ModulePatch.cpp
+++ b/Engine/Source/Developer/Windows/LiveCodingServer/Private/External/LC_ModulePatch.cpp
@@ -1,5 +1,7 @@
 // Copyright 2011-2020 Molecular Matters GmbH, all rights reserved.
 
+#if LC_VERSION == 1
+
 // BEGIN EPIC MOD
 //#include PCH_INCLUDE
 // END EPIC MOD
@@ -147,3 +149,6 @@ const ModulePatch::Data& ModulePatch::GetData(void) const
 {
 	return m_data;
 }
+
+
+#endif
\ No newline at end of file
diff --git a/Engine/Source/Developer/Windows/LiveCodingServer/Private/External/LC_ModulePatch.h b/Engine/Source/Developer/Windows/LiveCodingServer/Private/External/LC_ModulePatch.h
--- a/Engine/Source/Developer/Windows/LiveCodingServer/Private/External/LC_ModulePatch.h
+++ b/Engine/Source/Developer/Windows/LiveCodingServer/Private/External/LC_ModulePatch.h
@@ -2,6 +2,8 @@
 
 #pragma once
 
+#if LC_VERSION == 1
+
 // BEGIn EPIC MOD
 #include "CoreTypes.h"
 // END EPIC MOD
@@ -83,3 +85,6 @@ private:
 
 	Data m_data;
 };
+
+
+#endif
\ No newline at end of file
diff --git a/Engine/Source/Developer/Windows/LiveCodingServer/Private/External/LC_NameMangling.cpp b/Engine/Source/Developer/Windows/LiveCodingServer/Private/External/LC_NameMangling.cpp
--- a/Engine/Source/Developer/Windows/LiveCodingServer/Private/External/LC_NameMangling.cpp
+++ b/Engine/Source/Developer/Windows/LiveCodingServer/Private/External/LC_NameMangling.cpp
@@ -1,5 +1,7 @@
 // Copyright 2011-2020 Molecular Matters GmbH, all rights reserved.
 
+#if LC_VERSION == 1
+
 // BEGIN EPIC MOD
 //#include PCH_INCLUDE
 // END EPC MOD
@@ -30,3 +32,6 @@ std::wstring nameMangling::UndecorateSymbolWide(const char* symbolName, uint16_t
 
 	return string::ToWideString(buffer);
 }
+
+
+#endif
\ No newline at end of file
diff --git a/Engine/Source/Developer/Windows/LiveCodingServer/Private/External/LC_NameMangling.h b/Engine/Source/Developer/Windows/LiveCodingServer/Private/External/LC_NameMangling.h
--- a/Engine/Source/Developer/Windows/LiveCodingServer/Private/External/LC_NameMangling.h
+++ b/Engine/Source/Developer/Windows/LiveCodingServer/Private/External/LC_NameMangling.h
@@ -2,6 +2,8 @@
 
 #pragma once
 
+#if LC_VERSION == 1
+
 // BEGIN EPIC MOD
 #include "CoreTypes.h"
 #include <string>
@@ -31,3 +33,6 @@ namespace nameMangling
 	std::string UndecorateSymbol(const char* symbolName, uint16_t flags);
 	std::wstring UndecorateSymbolWide(const char* symbolName, uint16_t flags);
 }
+
+
+#endif
\ No newline at end of file
diff --git a/Engine/Source/Developer/Windows/LiveCodingServer/Private/External/LC_OutputFile.cpp b/Engine/Source/Developer/Windows/LiveCodingServer/Private/External/LC_OutputFile.cpp
--- a/Engine/Source/Developer/Windows/LiveCodingServer/Private/External/LC_OutputFile.cpp
+++ b/Engine/Source/Developer/Windows/LiveCodingServer/Private/External/LC_OutputFile.cpp
@@ -1,5 +1,7 @@
 // Copyright 2011-2020 Molecular Matters GmbH, all rights reserved.
 
+#if LC_VERSION == 1
+
 // BEGIN EPIC MOD
 //#include PCH_INCLUDE
 // END EPIC MOD
@@ -46,4 +48,6 @@ void OutputFile::WriteToFile(const char* text)
 		const size_t length = strlen(text);
 		::WriteFile(m_logFile, text, static_cast<DWORD>(length), &bytesWritten, nullptr);
 	}
-}
\ No newline at end of file
+}
+
+#endif
\ No newline at end of file
diff --git a/Engine/Source/Developer/Windows/LiveCodingServer/Private/External/LC_OutputFile.h b/Engine/Source/Developer/Windows/LiveCodingServer/Private/External/LC_OutputFile.h
--- a/Engine/Source/Developer/Windows/LiveCodingServer/Private/External/LC_OutputFile.h
+++ b/Engine/Source/Developer/Windows/LiveCodingServer/Private/External/LC_OutputFile.h
@@ -2,6 +2,8 @@
 
 #pragma once
 
+#if LC_VERSION == 1
+
 // BEGIN EPIC MOD
 #include "CoreTypes.h"
 #include "Windows/MinimalWindowsAPI.h"
@@ -22,3 +24,6 @@ private:
 	Windows::HANDLE m_logFile;
 	// END EPIC MOD
 };
+
+
+#endif
\ No newline at end of file
diff --git a/Engine/Source/Developer/Windows/LiveCodingServer/Private/External/LC_Patch.cpp b/Engine/Source/Developer/Windows/LiveCodingServer/Private/External/LC_Patch.cpp
--- a/Engine/Source/Developer/Windows/LiveCodingServer/Private/External/LC_Patch.cpp
+++ b/Engine/Source/Developer/Windows/LiveCodingServer/Private/External/LC_Patch.cpp
@@ -1,5 +1,7 @@
 // Copyright 2011-2020 Molecular Matters GmbH, all rights reserved.
 
+#if LC_VERSION == 1
+
 // BEGIN EPIC MOD
 //#include PCH_INCLUDE
 // END EPIC MOD
@@ -76,3 +78,6 @@ void patch::InstallRelativeNearJump(Process::Handle processHandle, void* address
 // BEGIN EPIC MODS
 #pragma warning(pop)
 // END EPIC MODS
+
+
+#endif
\ No newline at end of file
diff --git a/Engine/Source/Developer/Windows/LiveCodingServer/Private/External/LC_Patch.h b/Engine/Source/Developer/Windows/LiveCodingServer/Private/External/LC_Patch.h
--- a/Engine/Source/Developer/Windows/LiveCodingServer/Private/External/LC_Patch.h
+++ b/Engine/Source/Developer/Windows/LiveCodingServer/Private/External/LC_Patch.h
@@ -2,6 +2,8 @@
 
 #pragma once
 
+#if LC_VERSION == 1
+
 // BEGIN EPIC MOD
 #include "CoreTypes.h"
 // END EPIC MOD
@@ -16,3 +18,6 @@ namespace patch
 	void InstallRelativeShortJump(Process::Handle processHandle, void* address, void* destination);
 	void InstallRelativeNearJump(Process::Handle processHandle, void* address, void* destination);
 }
+
+
+#endif
\ No newline at end of file
diff --git a/Engine/Source/Developer/Windows/LiveCodingServer/Private/External/LC_PoolAllocator.cpp b/Engine/Source/Developer/Windows/LiveCodingServer/Private/External/LC_PoolAllocator.cpp
--- a/Engine/Source/Developer/Windows/LiveCodingServer/Private/External/LC_PoolAllocator.cpp
+++ b/Engine/Source/Developer/Windows/LiveCodingServer/Private/External/LC_PoolAllocator.cpp
@@ -1,5 +1,7 @@
 // Copyright 2011-2020 Molecular Matters GmbH, all rights reserved.
 
+#if LC_VERSION == 1
+
 // BEGIN EPIC MOD
 //#include PCH_INCLUDE
 // END EPIC MOD
@@ -177,3 +179,6 @@ const AllocatorStats& PoolAllocator<T>::GetStats(void) const
 // explicit template instantiation
 template class PoolAllocator<PoolAllocatorSingleThreadPolicy>;
 template class PoolAllocator<PoolAllocatorMultiThreadPolicy>;
+
+
+#endif
\ No newline at end of file
diff --git a/Engine/Source/Developer/Windows/LiveCodingServer/Private/External/LC_PoolAllocator.h b/Engine/Source/Developer/Windows/LiveCodingServer/Private/External/LC_PoolAllocator.h
--- a/Engine/Source/Developer/Windows/LiveCodingServer/Private/External/LC_PoolAllocator.h
+++ b/Engine/Source/Developer/Windows/LiveCodingServer/Private/External/LC_PoolAllocator.h
@@ -2,6 +2,8 @@
 
 #pragma once
 
+#if LC_VERSION == 1
+
 // BEGIN EPIC MOD
 #include "CoreTypes.h"
 // END EPIC MOD
@@ -67,3 +69,6 @@ private:
 	// intrusive linked list
 	BlockHeader* m_blockList;
 };
+
+
+#endif
\ No newline at end of file
diff --git a/Engine/Source/Developer/Windows/LiveCodingServer/Private/External/LC_RelocationPatcher.cpp b/Engine/Source/Developer/Windows/LiveCodingServer/Private/External/LC_RelocationPatcher.cpp
--- a/Engine/Source/Developer/Windows/LiveCodingServer/Private/External/LC_RelocationPatcher.cpp
+++ b/Engine/Source/Developer/Windows/LiveCodingServer/Private/External/LC_RelocationPatcher.cpp
@@ -1,5 +1,7 @@
 // Copyright 2011-2020 Molecular Matters GmbH, all rights reserved.
 
+#if LC_VERSION == 1
+
 // BEGIN EPIC MOD
 //#include PCH_INCLUDE
 // END EPIC MOD
@@ -556,3 +558,6 @@ bool relocations::IsValidRecord(const Record& record)
 {
 	return (record.relocationType != coff::Relocation::Type::UNKNOWN);
 }
+
+
+#endif
\ No newline at end of file
diff --git a/Engine/Source/Developer/Windows/LiveCodingServer/Private/External/LC_RelocationPatcher.h b/Engine/Source/Developer/Windows/LiveCodingServer/Private/External/LC_RelocationPatcher.h
--- a/Engine/Source/Developer/Windows/LiveCodingServer/Private/External/LC_RelocationPatcher.h
+++ b/Engine/Source/Developer/Windows/LiveCodingServer/Private/External/LC_RelocationPatcher.h
@@ -2,6 +2,8 @@
 
 #pragma once
 
+#if LC_VERSION == 1
+
 // BEGIN EPIC MOD
 #include "CoreTypes.h"
 // END EPIC MOD
@@ -84,3 +86,6 @@ namespace relocations
 
 	bool IsValidRecord(const Record& record);
 };
+
+
+#endif
\ No newline at end of file
diff --git a/Engine/Source/Developer/Windows/LiveCodingServer/Private/External/LC_Scheduler.cpp b/Engine/Source/Developer/Windows/LiveCodingServer/Private/External/LC_Scheduler.cpp
--- a/Engine/Source/Developer/Windows/LiveCodingServer/Private/External/LC_Scheduler.cpp
+++ b/Engine/Source/Developer/Windows/LiveCodingServer/Private/External/LC_Scheduler.cpp
@@ -1,5 +1,7 @@
 // Copyright 2011-2019 Molecular Matters GmbH, all rights reserved.
 
+#if LC_VERSION == 1
+
 // BEGIN EPIC MOD
 //#include PCH_INCLUDE
 // END EPIC MOD
@@ -164,3 +166,6 @@ void scheduler::WaitForTask(TaskBase* task)
 // BEGIN EPIC MODS
 #pragma warning(pop)
 // END EPIC MODS
+
+
+#endif
\ No newline at end of file
diff --git a/Engine/Source/Developer/Windows/LiveCodingServer/Private/External/LC_Scheduler.h b/Engine/Source/Developer/Windows/LiveCodingServer/Private/External/LC_Scheduler.h
--- a/Engine/Source/Developer/Windows/LiveCodingServer/Private/External/LC_Scheduler.h
+++ b/Engine/Source/Developer/Windows/LiveCodingServer/Private/External/LC_Scheduler.h
@@ -2,6 +2,8 @@
 
 #pragma once
 
+#if LC_VERSION == 1
+
 // BEGIN EPIC MOD
 #include "CoreTypes.h"
 // END EPIC MOD
@@ -51,3 +53,6 @@ namespace scheduler
 
 	void WaitForTask(TaskBase* task);
 }
+
+
+#endif
\ No newline at end of file
diff --git a/Engine/Source/Developer/Windows/LiveCodingServer/Private/External/LC_SchedulerQueue.cpp b/Engine/Source/Developer/Windows/LiveCodingServer/Private/External/LC_SchedulerQueue.cpp
--- a/Engine/Source/Developer/Windows/LiveCodingServer/Private/External/LC_SchedulerQueue.cpp
+++ b/Engine/Source/Developer/Windows/LiveCodingServer/Private/External/LC_SchedulerQueue.cpp
@@ -1,5 +1,7 @@
 // Copyright 2011-2020 Molecular Matters GmbH, all rights reserved.
 
+#if LC_VERSION == 1
+
 // BEGIN EPIC MOD
 //#include PCH_INCLUDE
 // END EPIC MOD
@@ -77,3 +79,6 @@ scheduler::TaskBase* scheduler::TaskQueue::TryPopTask(void)
 
 	return task;
 }
+
+
+#endif
\ No newline at end of file
diff --git a/Engine/Source/Developer/Windows/LiveCodingServer/Private/External/LC_SchedulerQueue.h b/Engine/Source/Developer/Windows/LiveCodingServer/Private/External/LC_SchedulerQueue.h
--- a/Engine/Source/Developer/Windows/LiveCodingServer/Private/External/LC_SchedulerQueue.h
+++ b/Engine/Source/Developer/Windows/LiveCodingServer/Private/External/LC_SchedulerQueue.h
@@ -2,6 +2,8 @@
 
 #pragma once
 
+#if LC_VERSION == 1
+
 // BEGIN EPIC MOD
 #include "CoreTypes.h"
 // END EPIC MOD
@@ -45,3 +47,6 @@ namespace scheduler
 		CriticalSection m_cs;
 	};
 }
+
+
+#endif
\ No newline at end of file
diff --git a/Engine/Source/Developer/Windows/LiveCodingServer/Private/External/LC_SchedulerTask.cpp b/Engine/Source/Developer/Windows/LiveCodingServer/Private/External/LC_SchedulerTask.cpp
--- a/Engine/Source/Developer/Windows/LiveCodingServer/Private/External/LC_SchedulerTask.cpp
+++ b/Engine/Source/Developer/Windows/LiveCodingServer/Private/External/LC_SchedulerTask.cpp
@@ -1,5 +1,7 @@
 // Copyright 2011-2020 Molecular Matters GmbH, all rights reserved.
 
+#if LC_VERSION == 1
+
 // BEGIN EPIC MOD
 //#include PCH_INCLUDE
 // END EPIC MOD
@@ -58,3 +60,6 @@ void scheduler::TaskBase::OnChildDetach(void)
 	CriticalSection::ScopedLock lock(&m_taskCS);
 	--m_openTasks;
 }
+
+
+#endif
\ No newline at end of file
diff --git a/Engine/Source/Developer/Windows/LiveCodingServer/Private/External/LC_SchedulerTask.h b/Engine/Source/Developer/Windows/LiveCodingServer/Private/External/LC_SchedulerTask.h
--- a/Engine/Source/Developer/Windows/LiveCodingServer/Private/External/LC_SchedulerTask.h
+++ b/Engine/Source/Developer/Windows/LiveCodingServer/Private/External/LC_SchedulerTask.h
@@ -2,6 +2,8 @@
 
 #pragma once
 
+#if LC_VERSION == 1
+
 // BEGIN EPIC MOD
 #include "CoreTypes.h"
 // END EPIC MOD
@@ -82,3 +84,6 @@ namespace scheduler
 		LC_DISABLE_MOVE_ASSIGNMENT(Task);
 	};
 }
+
+
+#endif
\ No newline at end of file
diff --git a/Engine/Source/Developer/Windows/LiveCodingServer/Private/External/LC_SchedulerWorkerThread.cpp b/Engine/Source/Developer/Windows/LiveCodingServer/Private/External/LC_SchedulerWorkerThread.cpp
--- a/Engine/Source/Developer/Windows/LiveCodingServer/Private/External/LC_SchedulerWorkerThread.cpp
+++ b/Engine/Source/Developer/Windows/LiveCodingServer/Private/External/LC_SchedulerWorkerThread.cpp
@@ -1,5 +1,7 @@
 // Copyright 2011-2020 Molecular Matters GmbH, all rights reserved.
 
+#if LC_VERSION == 1
+
 // BEGIN EPIC MOD
 //#include PCH_INCLUDE
 // END EPIC MOD
@@ -40,3 +42,6 @@ Thread::ReturnValue scheduler::WorkerThread::ThreadFunction(TaskQueue* queue)
 
 	return Thread::ReturnValue(0u);
 }
+
+
+#endif
\ No newline at end of file
diff --git a/Engine/Source/Developer/Windows/LiveCodingServer/Private/External/LC_SchedulerWorkerThread.h b/Engine/Source/Developer/Windows/LiveCodingServer/Private/External/LC_SchedulerWorkerThread.h
--- a/Engine/Source/Developer/Windows/LiveCodingServer/Private/External/LC_SchedulerWorkerThread.h
+++ b/Engine/Source/Developer/Windows/LiveCodingServer/Private/External/LC_SchedulerWorkerThread.h
@@ -2,6 +2,8 @@
 
 #pragma once
 
+#if LC_VERSION == 1
+
 // BEGIN EPIC MOD
 #include "CoreTypes.h"
 // END EPIC MOD
@@ -25,3 +27,6 @@ namespace scheduler
 		Thread::Handle m_thread;
 	};
 }
+
+
+#endif
\ No newline at end of file
diff --git a/Engine/Source/Developer/Windows/LiveCodingServer/Private/External/LC_ServerCommandThread.cpp b/Engine/Source/Developer/Windows/LiveCodingServer/Private/External/LC_ServerCommandThread.cpp
--- a/Engine/Source/Developer/Windows/LiveCodingServer/Private/External/LC_ServerCommandThread.cpp
+++ b/Engine/Source/Developer/Windows/LiveCodingServer/Private/External/LC_ServerCommandThread.cpp
@@ -1,5 +1,7 @@
 // Copyright 2011-2020 Molecular Matters GmbH, all rights reserved.
 
+#if LC_VERSION == 1
+
 // BEGIN EPIC MOD
 //#include PCH_INCLUDE
 // END EPIC MOD
@@ -2339,3 +2341,6 @@ bool ServerCommandThread::actions::ApplySettingString::Execute(const CommandType
 // BEGIN EPIC MODS
 #pragma warning(pop)
 // END EPIC MODS
+
+
+#endif
\ No newline at end of file
diff --git a/Engine/Source/Developer/Windows/LiveCodingServer/Private/External/LC_ServerCommandThread.h b/Engine/Source/Developer/Windows/LiveCodingServer/Private/External/LC_ServerCommandThread.h
--- a/Engine/Source/Developer/Windows/LiveCodingServer/Private/External/LC_ServerCommandThread.h
+++ b/Engine/Source/Developer/Windows/LiveCodingServer/Private/External/LC_ServerCommandThread.h
@@ -2,6 +2,8 @@
 
 #pragma once
 
+#if LC_VERSION == 1
+
 #include "LC_ThreadTypes.h"
 #include "LC_ProcessTypes.h"
 #include "LC_Commands.h"
@@ -205,3 +207,6 @@ private:
 #endif
 // END EPIC MOD
 };
+
+
+#endif
\ No newline at end of file
diff --git a/Engine/Source/Developer/Windows/LiveCodingServer/Private/External/LC_Settings.cpp b/Engine/Source/Developer/Windows/LiveCodingServer/Private/External/LC_Settings.cpp
--- a/Engine/Source/Developer/Windows/LiveCodingServer/Private/External/LC_Settings.cpp
+++ b/Engine/Source/Developer/Windows/LiveCodingServer/Private/External/LC_Settings.cpp
@@ -1,5 +1,7 @@
 // Copyright 2011-2020 Molecular Matters GmbH, all rights reserved.
 
+#if LC_VERSION == 1
+
 // BEGIN EPIC MOD
 //#include PCH_INCLUDE
 // END EPIC MOD
@@ -445,3 +447,6 @@ int SettingShortcut::GetValue(void) const
 {
 	return m_value;
 }
+
+
+#endif
\ No newline at end of file
diff --git a/Engine/Source/Developer/Windows/LiveCodingServer/Private/External/LC_Settings.h b/Engine/Source/Developer/Windows/LiveCodingServer/Private/External/LC_Settings.h
--- a/Engine/Source/Developer/Windows/LiveCodingServer/Private/External/LC_Settings.h
+++ b/Engine/Source/Developer/Windows/LiveCodingServer/Private/External/LC_Settings.h
@@ -2,6 +2,8 @@
 
 #pragma once
 
+#if LC_VERSION == 1
+
 // BEGIN EPIC MOD
 #include "CoreTypes.h"
 #include "LC_Types.h"
@@ -177,3 +179,6 @@ private:
 	const wchar_t* m_description;
 	int m_value;
 };
+
+
+#endif
\ No newline at end of file
diff --git a/Engine/Source/Developer/Windows/LiveCodingServer/Private/External/LC_SymbolInfo.cpp b/Engine/Source/Developer/Windows/LiveCodingServer/Private/External/LC_SymbolInfo.cpp
--- a/Engine/Source/Developer/Windows/LiveCodingServer/Private/External/LC_SymbolInfo.cpp
+++ b/Engine/Source/Developer/Windows/LiveCodingServer/Private/External/LC_SymbolInfo.cpp
@@ -1,5 +1,7 @@
 // Copyright 2011-2020 Molecular Matters GmbH, all rights reserved.
 
+#if LC_VERSION == 1
+
 // BEGIN EPIC MOD
 //#include PCH_INCLUDE
 // END EPIC MOD
@@ -16,3 +18,6 @@ SymbolInfo::SymbolInfo(const char* const function, const char* const filename, u
 	strcpy_s(m_function, function);
 	strcpy_s(m_filename, filename);
 }
+
+
+#endif
\ No newline at end of file
diff --git a/Engine/Source/Developer/Windows/LiveCodingServer/Private/External/LC_SymbolInfo.h b/Engine/Source/Developer/Windows/LiveCodingServer/Private/External/LC_SymbolInfo.h
--- a/Engine/Source/Developer/Windows/LiveCodingServer/Private/External/LC_SymbolInfo.h
+++ b/Engine/Source/Developer/Windows/LiveCodingServer/Private/External/LC_SymbolInfo.h
@@ -2,6 +2,8 @@
 
 #pragma once
 
+#if LC_VERSION == 1
+
 // BEGIN EPIC MOD
 #include "CoreTypes.h"
 #include "LC_Platform.h"
@@ -37,3 +39,6 @@ private:
 	// END EPIC MOD
 	unsigned int m_line;
 };
+
+
+#endif
\ No newline at end of file
diff --git a/Engine/Source/Developer/Windows/LiveCodingServer/Private/External/LC_SymbolPatterns.cpp b/Engine/Source/Developer/Windows/LiveCodingServer/Private/External/LC_SymbolPatterns.cpp
--- a/Engine/Source/Developer/Windows/LiveCodingServer/Private/External/LC_SymbolPatterns.cpp
+++ b/Engine/Source/Developer/Windows/LiveCodingServer/Private/External/LC_SymbolPatterns.cpp
@@ -1,5 +1,7 @@
 // Copyright 2011-2020 Molecular Matters GmbH, all rights reserved.
 
+#if LC_VERSION == 1
+
 // BEGIN EPIC MOD
 //#include PCH_INCLUDE
 // END EPIC MOD
@@ -244,3 +246,6 @@ namespace symbolPatterns
 	};
 	// END EPIC MOD
 }
+
+
+#endif
\ No newline at end of file
diff --git a/Engine/Source/Developer/Windows/LiveCodingServer/Private/External/LC_SymbolPatterns.h b/Engine/Source/Developer/Windows/LiveCodingServer/Private/External/LC_SymbolPatterns.h
--- a/Engine/Source/Developer/Windows/LiveCodingServer/Private/External/LC_SymbolPatterns.h
+++ b/Engine/Source/Developer/Windows/LiveCodingServer/Private/External/LC_SymbolPatterns.h
@@ -2,6 +2,8 @@
 
 #pragma once
 
+#if LC_VERSION == 1
+
 // BEGIN EPIC MOD
 #include "CoreTypes.h"
 #include "LC_Platform.h"
@@ -47,3 +49,6 @@ namespace symbolPatterns
 	extern const char* const UE_REGISTER_PATTERNS[6];
 	// END EPIC MOD
 }
+
+
+#endif
\ No newline at end of file
diff --git a/Engine/Source/Developer/Windows/LiveCodingServer/Private/External/LC_SymbolReconstruction.cpp b/Engine/Source/Developer/Windows/LiveCodingServer/Private/External/LC_SymbolReconstruction.cpp
--- a/Engine/Source/Developer/Windows/LiveCodingServer/Private/External/LC_SymbolReconstruction.cpp
+++ b/Engine/Source/Developer/Windows/LiveCodingServer/Private/External/LC_SymbolReconstruction.cpp
@@ -1,5 +1,7 @@
 // Copyright 2011-2020 Molecular Matters GmbH, all rights reserved.
 
+#if LC_VERSION == 1
+
 // BEGIN EPIC MOD
 //#include PCH_INCLUDE
 // END EPIC MOD
@@ -919,3 +921,6 @@ walkOpenSymbols:
 		goto walkOpenSymbols;
 	}
 }
+
+
+#endif
\ No newline at end of file
diff --git a/Engine/Source/Developer/Windows/LiveCodingServer/Private/External/LC_SymbolReconstruction.h b/Engine/Source/Developer/Windows/LiveCodingServer/Private/External/LC_SymbolReconstruction.h
--- a/Engine/Source/Developer/Windows/LiveCodingServer/Private/External/LC_SymbolReconstruction.h
+++ b/Engine/Source/Developer/Windows/LiveCodingServer/Private/External/LC_SymbolReconstruction.h
@@ -2,6 +2,8 @@
 
 #pragma once
 
+#if LC_VERSION == 1
+
 // BEGIN EPIC MOD
 #include "CoreTypes.h"
 // END EPIC MOD
@@ -26,3 +28,6 @@ namespace symbols
 		symbols::SymbolDB* symbolDB
 	);
 }
+
+
+#endif
\ No newline at end of file
diff --git a/Engine/Source/Developer/Windows/LiveCodingServer/Private/External/LC_SymbolResolution.cpp b/Engine/Source/Developer/Windows/LiveCodingServer/Private/External/LC_SymbolResolution.cpp
--- a/Engine/Source/Developer/Windows/LiveCodingServer/Private/External/LC_SymbolResolution.cpp
+++ b/Engine/Source/Developer/Windows/LiveCodingServer/Private/External/LC_SymbolResolution.cpp
@@ -1,5 +1,7 @@
 // Copyright 2011-2020 Molecular Matters GmbH, all rights reserved.
 
+#if LC_VERSION == 1
+
 // BEGIN EPIC MOD
 //#include PCH_INCLUDE
 // END EPIC MOD
@@ -81,3 +83,6 @@ namespace symbolResolution
 		return SymbolInfo("unknown", "unknown", 0);
 	}
 }
+
+
+#endif
\ No newline at end of file
diff --git a/Engine/Source/Developer/Windows/LiveCodingServer/Private/External/LC_SymbolResolution.h b/Engine/Source/Developer/Windows/LiveCodingServer/Private/External/LC_SymbolResolution.h
--- a/Engine/Source/Developer/Windows/LiveCodingServer/Private/External/LC_SymbolResolution.h
+++ b/Engine/Source/Developer/Windows/LiveCodingServer/Private/External/LC_SymbolResolution.h
@@ -2,6 +2,8 @@
 
 #pragma once
 
+#if LC_VERSION == 1
+
 #include "LC_SymbolInfo.h"
 
 
@@ -10,3 +12,6 @@ namespace symbolResolution
 	// Resolves the symbol information for a given address
 	SymbolInfo ResolveSymbolsForAddress(const void* const address);
 }
+
+
+#endif
\ No newline at end of file
diff --git a/Engine/Source/Developer/Windows/LiveCodingServer/Private/External/LC_Symbols.cpp b/Engine/Source/Developer/Windows/LiveCodingServer/Private/External/LC_Symbols.cpp
--- a/Engine/Source/Developer/Windows/LiveCodingServer/Private/External/LC_Symbols.cpp
+++ b/Engine/Source/Developer/Windows/LiveCodingServer/Private/External/LC_Symbols.cpp
@@ -1,5 +1,7 @@
 // Copyright 2011-2020 Molecular Matters GmbH, all rights reserved.
 
+#if LC_VERSION == 1
+
 // BEGIN EPIC MOD
 //#include PCH_INCLUDE
 // END EPIC MOD
@@ -2399,4 +2401,6 @@ namespace symbols
 
 // BEGIN EPIC MOD
 #include "Windows/HideWindowsPlatformAtomics.h"
-// END EPIC MOD
\ No newline at end of file
+// END EPIC MOD
+
+#endif
\ No newline at end of file
diff --git a/Engine/Source/Developer/Windows/LiveCodingServer/Private/External/LC_Symbols.h b/Engine/Source/Developer/Windows/LiveCodingServer/Private/External/LC_Symbols.h
--- a/Engine/Source/Developer/Windows/LiveCodingServer/Private/External/LC_Symbols.h
+++ b/Engine/Source/Developer/Windows/LiveCodingServer/Private/External/LC_Symbols.h
@@ -2,6 +2,8 @@
 
 #pragma once
 
+#if LC_VERSION == 1
+
 // BEGIN EPIC MOD
 #include "CoreTypes.h"
 // END EPIC MOD
@@ -422,3 +424,6 @@ namespace symbols
 	bool IsUEStaticsSymbol(const ImmutableString& symbolName);
 	// END EPIC MOD
 }
+
+
+#endif
\ No newline at end of file
diff --git a/Engine/Source/Developer/Windows/LiveCodingServer/Private/External/LC_Telemetry.cpp b/Engine/Source/Developer/Windows/LiveCodingServer/Private/External/LC_Telemetry.cpp
--- a/Engine/Source/Developer/Windows/LiveCodingServer/Private/External/LC_Telemetry.cpp
+++ b/Engine/Source/Developer/Windows/LiveCodingServer/Private/External/LC_Telemetry.cpp
@@ -1,5 +1,7 @@
 // Copyright 2011-2020 Molecular Matters GmbH, all rights reserved.
 
+#if LC_VERSION == 1
+
 // BEGIN EPIC MOD
 //#include PCH_INCLUDE
 // END EPIC MOD
@@ -127,3 +129,6 @@ void telemetry::Accumulator::Print(void)
 // BEGIN EPIC MOD
 #include "Windows/HideWindowsPlatformAtomics.h"
 // END EPIC MOD
+
+
+#endif
\ No newline at end of file
diff --git a/Engine/Source/Developer/Windows/LiveCodingServer/Private/External/LC_Telemetry.h b/Engine/Source/Developer/Windows/LiveCodingServer/Private/External/LC_Telemetry.h
--- a/Engine/Source/Developer/Windows/LiveCodingServer/Private/External/LC_Telemetry.h
+++ b/Engine/Source/Developer/Windows/LiveCodingServer/Private/External/LC_Telemetry.h
@@ -2,6 +2,8 @@
 
 #pragma once
 
+#if LC_VERSION == 1
+
 // BEGIN EPIC MOD
 #include "CoreTypes.h"
 // END EPIC MOD
@@ -51,3 +53,6 @@ namespace telemetry
 		uint64_t m_accumulated;
 	};
 }
+
+
+#endif
\ No newline at end of file
diff --git a/Engine/Source/Developer/Windows/LiveCodingServer/Private/External/LC_TimeDate.cpp b/Engine/Source/Developer/Windows/LiveCodingServer/Private/External/LC_TimeDate.cpp
--- a/Engine/Source/Developer/Windows/LiveCodingServer/Private/External/LC_TimeDate.cpp
+++ b/Engine/Source/Developer/Windows/LiveCodingServer/Private/External/LC_TimeDate.cpp
@@ -1,5 +1,7 @@
 // Copyright 2011-2020 Molecular Matters GmbH, all rights reserved.
 
+#if LC_VERSION == 1
+
 // BEGIN EPIC MOD
 //#include PCH_INCLUDE
 // END EPIC MOD
@@ -41,3 +43,6 @@ const char* TimeDate::ToDateString(DateDescription& desc) const
 	_snprintf_s(desc, _TRUNCATE, "%d-%02d-%02d", m_year, m_month, m_day);
 	return desc;
 }
+
+
+#endif
\ No newline at end of file
diff --git a/Engine/Source/Developer/Windows/LiveCodingServer/Private/External/LC_TimeDate.h b/Engine/Source/Developer/Windows/LiveCodingServer/Private/External/LC_TimeDate.h
--- a/Engine/Source/Developer/Windows/LiveCodingServer/Private/External/LC_TimeDate.h
+++ b/Engine/Source/Developer/Windows/LiveCodingServer/Private/External/LC_TimeDate.h
@@ -2,6 +2,8 @@
 
 #pragma once
 
+#if LC_VERSION == 1
+
 
 class TimeDate
 {
@@ -32,3 +34,6 @@ private:
 	unsigned char m_second;
 	unsigned short m_milliSecond;
 };
+
+
+#endif
\ No newline at end of file
diff --git a/Engine/Source/Developer/Windows/LiveCodingServer/Private/External/LC_TimeStamp.cpp b/Engine/Source/Developer/Windows/LiveCodingServer/Private/External/LC_TimeStamp.cpp
--- a/Engine/Source/Developer/Windows/LiveCodingServer/Private/External/LC_TimeStamp.cpp
+++ b/Engine/Source/Developer/Windows/LiveCodingServer/Private/External/LC_TimeStamp.cpp
@@ -1,5 +1,7 @@
 // Copyright 2011-2020 Molecular Matters GmbH, all rights reserved.
 
+#if LC_VERSION == 1
+
 // BEGIN EPIC MOD
 //#include PCH_INCLUDE
 // END EPIC MOD
@@ -58,3 +60,6 @@ namespace timeStamp
 		return count * g_millionOverFrequency;
 	}
 }
+
+
+#endif
\ No newline at end of file
diff --git a/Engine/Source/Developer/Windows/LiveCodingServer/Private/External/LC_TimeStamp.h b/Engine/Source/Developer/Windows/LiveCodingServer/Private/External/LC_TimeStamp.h
--- a/Engine/Source/Developer/Windows/LiveCodingServer/Private/External/LC_TimeStamp.h
+++ b/Engine/Source/Developer/Windows/LiveCodingServer/Private/External/LC_TimeStamp.h
@@ -2,6 +2,8 @@
 
 #pragma once
 
+#if LC_VERSION == 1
+
 // BEGIN EPIC MOD
 #include "LC_Platform.h"
 // END EPIC MOD
@@ -25,3 +27,6 @@ namespace timeStamp
 	// Converts counts returned by Get() into microseconds.
 	double ToMicroSeconds(uint64_t count);
 }
+
+
+#endif
\ No newline at end of file
diff --git a/Engine/Source/Developer/Windows/LiveCodingServer/Private/External/LC_Types.h b/Engine/Source/Developer/Windows/LiveCodingServer/Private/External/LC_Types.h
--- a/Engine/Source/Developer/Windows/LiveCodingServer/Private/External/LC_Types.h
+++ b/Engine/Source/Developer/Windows/LiveCodingServer/Private/External/LC_Types.h
@@ -2,6 +2,8 @@
 
 #pragma once
 
+#if LC_VERSION == 1
+
 // BEGIN EPIC MOD
 #include "CoreTypes.h"
 // END EPIC MOD
@@ -71,3 +73,6 @@ namespace types
 		Allocator<Key>
 	>;
 }
+
+
+#endif
\ No newline at end of file
diff --git a/Engine/Source/Developer/Windows/LiveCodingServer/Private/External/LC_UniqueId.cpp b/Engine/Source/Developer/Windows/LiveCodingServer/Private/External/LC_UniqueId.cpp
--- a/Engine/Source/Developer/Windows/LiveCodingServer/Private/External/LC_UniqueId.cpp
+++ b/Engine/Source/Developer/Windows/LiveCodingServer/Private/External/LC_UniqueId.cpp
@@ -1,5 +1,7 @@
 // Copyright 2011-2020 Molecular Matters GmbH, all rights reserved.
 
+#if LC_VERSION == 1
+
 // BEGIN EPIC MOD
 //#include PCH_INCLUDE
 // END EPIC MOD
@@ -80,3 +82,6 @@ namespace uniqueId
 		return g_uniqueId->Generate(path);
 	}
 }
+
+
+#endif
\ No newline at end of file
diff --git a/Engine/Source/Developer/Windows/LiveCodingServer/Private/External/LC_UniqueId.h b/Engine/Source/Developer/Windows/LiveCodingServer/Private/External/LC_UniqueId.h
--- a/Engine/Source/Developer/Windows/LiveCodingServer/Private/External/LC_UniqueId.h
+++ b/Engine/Source/Developer/Windows/LiveCodingServer/Private/External/LC_UniqueId.h
@@ -2,6 +2,8 @@
 
 #pragma once
 
+#if LC_VERSION == 1
+
 // BEGIN EPIC MOD
 #include "CoreTypes.h"
 #include <string>
@@ -14,3 +16,6 @@ namespace uniqueId
 
 	uint32_t Generate(const std::wstring& path);
 }
+
+
+#endif
\ No newline at end of file
diff --git a/Engine/Source/Developer/Windows/LiveCodingServer/Private/External/LC_VirtualDrive.cpp b/Engine/Source/Developer/Windows/LiveCodingServer/Private/External/LC_VirtualDrive.cpp
--- a/Engine/Source/Developer/Windows/LiveCodingServer/Private/External/LC_VirtualDrive.cpp
+++ b/Engine/Source/Developer/Windows/LiveCodingServer/Private/External/LC_VirtualDrive.cpp
@@ -1,5 +1,7 @@
 // Copyright 2011-2020 Molecular Matters GmbH, all rights reserved.
 
+#if LC_VERSION == 1
+
 // BEGIN EPIC MOD
 //#include PCH_INCLUDE
 // END EPIC MOD
@@ -29,3 +31,6 @@ void virtualDrive::Remove(const wchar_t* driveLetterPlusColon, const wchar_t* pa
 		LC_ERROR_USER("Cannot remove virtual drive %S for path %S. Error: 0x%X", driveLetterPlusColon, path, error);
 	}
 }
+
+
+#endif
\ No newline at end of file
diff --git a/Engine/Source/Developer/Windows/LiveCodingServer/Private/External/LC_VirtualDrive.h b/Engine/Source/Developer/Windows/LiveCodingServer/Private/External/LC_VirtualDrive.h
--- a/Engine/Source/Developer/Windows/LiveCodingServer/Private/External/LC_VirtualDrive.h
+++ b/Engine/Source/Developer/Windows/LiveCodingServer/Private/External/LC_VirtualDrive.h
@@ -2,6 +2,8 @@
 
 #pragma once
 
+#if LC_VERSION == 1
+
 // BEGIN EPIC MOD
 #include "CoreTypes.h"
 // END EPIC MOD
@@ -11,3 +13,6 @@ namespace virtualDrive
 	void Add(const wchar_t* driveLetterPlusColon, const wchar_t* path);
 	void Remove(const wchar_t* driveLetterPlusColon, const wchar_t* path);
 }
+
+
+#endif
\ No newline at end of file
diff --git a/Engine/Source/Developer/Windows/LiveCodingServer/Private/External/LC_VirtualMemoryRange.cpp b/Engine/Source/Developer/Windows/LiveCodingServer/Private/External/LC_VirtualMemoryRange.cpp
--- a/Engine/Source/Developer/Windows/LiveCodingServer/Private/External/LC_VirtualMemoryRange.cpp
+++ b/Engine/Source/Developer/Windows/LiveCodingServer/Private/External/LC_VirtualMemoryRange.cpp
@@ -1,5 +1,7 @@
 // Copyright 2011-2020 Molecular Matters GmbH, all rights reserved.
 
+#if LC_VERSION == 1
+
 // BEGIN EPIC MOD
 //#include PCH_INCLUDE
 // END EPIC MOD
@@ -124,3 +126,6 @@ void VirtualMemoryRange::AddPage(void* page)
 }
 // END EPIC MOD
 
+
+
+#endif
\ No newline at end of file
diff --git a/Engine/Source/Developer/Windows/LiveCodingServer/Private/External/LC_VirtualMemoryRange.h b/Engine/Source/Developer/Windows/LiveCodingServer/Private/External/LC_VirtualMemoryRange.h
--- a/Engine/Source/Developer/Windows/LiveCodingServer/Private/External/LC_VirtualMemoryRange.h
+++ b/Engine/Source/Developer/Windows/LiveCodingServer/Private/External/LC_VirtualMemoryRange.h
@@ -2,6 +2,8 @@
 
 #pragma once
 
+#if LC_VERSION == 1
+
 // BEGIN EPIC MOD
 #include "CoreTypes.h"
 // END EPIC MOD
@@ -35,3 +37,6 @@ private:
 	types::vector<PageData> m_pageData;
 	CriticalSection m_cs;
 };
+
+
+#endif
\ No newline at end of file
diff --git a/Engine/Source/Developer/Windows/LiveCodingServer/Private/External/LC_VisualStudioAutomation.cpp b/Engine/Source/Developer/Windows/LiveCodingServer/Private/External/LC_VisualStudioAutomation.cpp
--- a/Engine/Source/Developer/Windows/LiveCodingServer/Private/External/LC_VisualStudioAutomation.cpp
+++ b/Engine/Source/Developer/Windows/LiveCodingServer/Private/External/LC_VisualStudioAutomation.cpp
@@ -1,5 +1,7 @@
 // Copyright 2011-2020 Molecular Matters GmbH, all rights reserved.
 
+#if LC_VERSION == 1
+
 // BEGIN EPIC MOD
 //#include PCH_INCLUDE
 // END EPIC MOD
@@ -686,4 +688,6 @@ bool visualStudio::Break(const EnvDTE::DebuggerPtr& debugger)
 
 // BEGIN EPIC MOD
 #endif
-// END EPIC MOD
\ No newline at end of file
+// END EPIC MOD
+
+#endif
\ No newline at end of file
diff --git a/Engine/Source/Developer/Windows/LiveCodingServer/Private/External/LC_VisualStudioAutomation.h b/Engine/Source/Developer/Windows/LiveCodingServer/Private/External/LC_VisualStudioAutomation.h
--- a/Engine/Source/Developer/Windows/LiveCodingServer/Private/External/LC_VisualStudioAutomation.h
+++ b/Engine/Source/Developer/Windows/LiveCodingServer/Private/External/LC_VisualStudioAutomation.h
@@ -2,6 +2,8 @@
 
 #pragma once
 
+#if LC_VERSION == 1
+
 #include "LC_Types.h"
 #include "LC_ProcessTypes.h"
 #include "LC_ThreadTypes.h"
@@ -50,4 +52,6 @@ namespace visualStudio
 
 // BEGIN EPIC MOD
 #endif
-// END EPIC MOD
\ No newline at end of file
+// END EPIC MOD
+
+#endif
\ No newline at end of file
