Explicitly use new define for default binding space.

diff --git a/Engine/Source/Developer/Windows/ShaderFormatD3D/Private/D3DShaderCompilerDXC.cpp b/Engine/Source/Developer/Windows/ShaderFormatD3D/Private/D3DShaderCompilerDXC.cpp
--- a/Engine/Source/Developer/Windows/ShaderFormatD3D/Private/D3DShaderCompilerDXC.cpp
+++ b/Engine/Source/Developer/Windows/ShaderFormatD3D/Private/D3DShaderCompilerDXC.cpp
@@ -81,7 +81,7 @@ static uint32 GetAutoBindingSpace(const FShaderTarget& Target)
 	case SF_WorkGraphComputeNode:
 		return UE_HLSL_SPACE_WORK_GRAPH_LOCAL;
 	default:
-		return 0;
+		return UE_HLSL_SPACE_DEFAULT;
 	}
 }
 
