[UBT]
* Added vfs paths to LiveCoding manifest
* Added vfs support in LPP(1). LPP will convert virtual to local paths just before file access and linker process will be wrapped in UbaCli containing -vfs options for conversion of paths

diff --git a/Engine/Source/Developer/Windows/LiveCodingServer/Private/External/LC_LiveModule.cpp b/Engine/Source/Developer/Windows/LiveCodingServer/Private/External/LC_LiveModule.cpp
--- a/Engine/Source/Developer/Windows/LiveCodingServer/Private/External/LC_LiveModule.cpp
+++ b/Engine/Source/Developer/Windows/LiveCodingServer/Private/External/LC_LiveModule.cpp
@@ -539,8 +539,9 @@ namespace
 	};
 
 
-	template <typename T>
-	static types::vector<symbols::ObjPath> UpdateCoffCache(const T& compilands, CoffCache<coff::CoffDB>* coffCache, CacheUpdate::Enum updateType, const types::vector<symbols::ModifiedObjFile>& modifiedOrNewObjFiles)
+	// EPIC MOD - Added VfsToLocal
+	template <typename T, typename VfsToLocalFunc>
+	static types::vector<symbols::ObjPath> UpdateCoffCache(const VfsToLocalFunc& vfsToLocal, const T& compilands, CoffCache<coff::CoffDB>* coffCache, CacheUpdate::Enum updateType, const types::vector<symbols::ModifiedObjFile>& modifiedOrNewObjFiles)
 	{
 		LC_LOG_INDENT_DEV;
 
@@ -567,11 +568,13 @@ namespace
 			{
 				updatedCoffs.push_back(objPath);
 
-				auto task = scheduler::CreateTask(taskRoot, [objPath, wideObjPath, coffCache, compilandUniqueId]()
+				// EPIC MOD - Added VfsToLocal
+				auto task = scheduler::CreateTask(taskRoot, [objPath, wideObjPath, coffCache, compilandUniqueId, &vfsToLocal]()
 				{
 					LC_LOG_DEV("Updating COFF cache for file %s", objPath.c_str());
 
-					coff::ObjFile* objFile = coff::OpenObj(wideObjPath.c_str());
+					// EPIC MOD - Added VfsToLocal
+					coff::ObjFile* objFile = coff::OpenObj(vfsToLocal(wideObjPath).c_str());
 					if (objFile && objFile->memoryFile)
 					{
 						coff::CoffDB* database = coff::GatherDatabase(objFile, compilandUniqueId);
@@ -1351,7 +1354,7 @@ void LiveModule::UpdateDirectoryCache(DirectoryCache* cache)
 
 
 // BEGIN EPIC MOD
-LiveModule::ErrorType::Enum LiveModule::Update(FileAttributeCache* fileCache, DirectoryCache* directoryCache, UpdateType::Enum updateType, const types::vector<symbols::ModifiedObjFile>& modifiedOrNewObjFiles, const types::vector<std::wstring>& additionalLibraries)
+LiveModule::ErrorType::Enum LiveModule::Update(FileAttributeCache* fileCache, DirectoryCache* directoryCache, UpdateType::Enum updateType, const types::vector<symbols::ModifiedObjFile>& modifiedOrNewObjFiles, const types::vector<std::wstring>& additionalLibraries, const types::vector<VfsEntry>& vfsEntries)
 // END EPIC MOD
 {
 	telemetry::Scope updateScope("Update live module");
@@ -1360,6 +1363,14 @@ LiveModule::ErrorType::Enum LiveModule::Update(FileAttributeCache* fileCache, Di
 
 	// BEGIN EPIC MOD
 	GLiveCodingServer->GetStatusChangeDelegate().ExecuteIfBound(L"Checking modified files...");
+
+	auto VfsToLocal = [&](const std::wstring& path)
+		{
+			for (auto& entry : vfsEntries)
+				if (entry.virtualPath.compare(0, entry.virtualPath.size(), path, 0, entry.virtualPath.size()) == 0)
+					return entry.localPath + (path.data() + entry.virtualPath.size());
+			return path;
+		};
 	// END EPIC MOD
 
 	bool usesClang = false;
@@ -2025,15 +2036,14 @@ LiveModule::ErrorType::Enum LiveModule::Update(FileAttributeCache* fileCache, Di
 		LC_LOG_DEV("Finding set of .obj files");
 		LC_LOG_INDENT_DEV;
 
-		struct Helper
-		{
-			static void UpdateExternalSymbolsAndNeededFiles
+		// EPIC MOD - Changed to lambda to be able to access VfsToLocal
+		auto UpdateExternalSymbolsAndNeededFiles = [&]
 			(
 				const symbols::ObjPath& objPath, const symbols::Compiland* compiland, uint32_t compilandUniqueId, const types::StringMap<ImmutableString>& pchSymbolToCompilandName,
 				types::StringMap<CompilandInfo>& externalSymbols, types::StringSet& neededCompilands
 			)
 			{
-				coff::ObjFile* coffFile = coff::OpenObj(string::ToWideString(objPath).c_str());
+				coff::ObjFile* coffFile = coff::OpenObj(VfsToLocal(string::ToWideString(objPath)).c_str());
 				if (coffFile && coffFile->memoryFile)
 				{
 					coff::ExternalSymbolDB* externalSymbolDb = coff::GatherExternalSymbolDatabase(coffFile, compilandUniqueId);
@@ -2083,8 +2093,7 @@ LiveModule::ErrorType::Enum LiveModule::Update(FileAttributeCache* fileCache, Di
 						}
 					}
 				}
-			}
-		};
+			};
 
 		if (modifiedOrNewObjFiles.size() == 0u)
 		{
@@ -2102,7 +2111,7 @@ LiveModule::ErrorType::Enum LiveModule::Update(FileAttributeCache* fileCache, Di
 					LC_LOG_DEV("%s is recompiled", objPath.c_str());
 					neededCompilands.emplace(objPath);
 
-					Helper::UpdateExternalSymbolsAndNeededFiles(objPath, compiland, compiland->uniqueId, m_pchSymbolToCompilandName, externalSymbols, neededCompilands);
+					UpdateExternalSymbolsAndNeededFiles(objPath, compiland, compiland->uniqueId, m_pchSymbolToCompilandName, externalSymbols, neededCompilands);
 				}
 				else
 				{
@@ -2140,7 +2149,7 @@ LiveModule::ErrorType::Enum LiveModule::Update(FileAttributeCache* fileCache, Di
 				LC_LOG_DEV("%s %s", objPath.c_str(), compiland ? "was recompiled" : "is new");
 				neededCompilands.emplace(objPath);
 
-				Helper::UpdateExternalSymbolsAndNeededFiles(objPath, compiland, compilandUniqueId, m_pchSymbolToCompilandName, externalSymbols, neededCompilands);
+				UpdateExternalSymbolsAndNeededFiles(objPath, compiland, compilandUniqueId, m_pchSymbolToCompilandName, externalSymbols, neededCompilands);
 			}
 		}
 	}
@@ -2211,7 +2220,7 @@ LiveModule::ErrorType::Enum LiveModule::Update(FileAttributeCache* fileCache, Di
 						// there is no entry yet for this COFF in the cache.
 						// this means that this .obj was not recompiled (otherwise it would have an entry already),
 						// but has been pulled in for the first time due to unresolved symbols.
-						auto task = scheduler::CreateTask(taskRoot, [this, objPath, &modifiedOrNewObjFiles]()
+						auto task = scheduler::CreateTask(taskRoot, [this, objPath, &modifiedOrNewObjFiles, &VfsToLocal]()
 						{
 							const symbols::Compiland* compiland = symbols::FindCompiland(m_compilandDB, objPath);
 							const std::wstring& wideObjPath = string::ToWideString(objPath);
@@ -2219,7 +2228,8 @@ LiveModule::ErrorType::Enum LiveModule::Update(FileAttributeCache* fileCache, Di
 
 							LC_LOG_DEV("Need %s for the first time, updating COFF cache", objPath.c_str());
 
-							coff::ObjFile* objFile = coff::OpenObj(wideObjPath.c_str());
+							// EPIC MOD - Added VfsToLocal
+							coff::ObjFile* objFile = coff::OpenObj(VfsToLocal(wideObjPath).c_str());
 							if (objFile && objFile->memoryFile)
 							{
 								// note that even though we might be dealing with a single-part .obj of an amalgamated .obj
@@ -2282,7 +2292,8 @@ LiveModule::ErrorType::Enum LiveModule::Update(FileAttributeCache* fileCache, Di
 
 
 	// update the COFF cache for all compiled files
-	UpdateCoffCache(m_compiledCompilands, m_coffCache, CacheUpdate::ALL, modifiedOrNewObjFiles);
+	// EPIC MOD - Added VfsToLocal
+	UpdateCoffCache(VfsToLocal, m_compiledCompilands, m_coffCache, CacheUpdate::ALL, modifiedOrNewObjFiles);
 
 	// BEGIN EPIC MOD
 	GLiveCodingServer->GetStatusChangeDelegate().ExecuteIfBound(L"Stripping COFFs...");
@@ -2315,7 +2326,8 @@ LiveModule::ErrorType::Enum LiveModule::Update(FileAttributeCache* fileCache, Di
 
 			const uint32_t compilandUniqueId = symbols::GetCompilandId(compiland, wideObjPath.c_str(), modifiedOrNewObjFiles);
 
-			coff::ObjFile* objFile = coff::OpenObj(wideObjPath.c_str());
+			// EPIC MOD - Added VfsToLocal
+			coff::ObjFile* objFile = coff::OpenObj(VfsToLocal(wideObjPath).c_str());
 			if (objFile && objFile->memoryFile)
 			{
 				coff::RawCoff* rawCoff = coff::ReadRaw(objFile, compilandUniqueId);
@@ -2347,7 +2359,8 @@ LiveModule::ErrorType::Enum LiveModule::Update(FileAttributeCache* fileCache, Di
 			// we need it after linking has finished  
 			{
 				const std::wstring bakPath = wideObjPath + L".bak";
-				Filesystem::Move(wideObjPath.c_str(), bakPath.c_str());
+				// EPIC MOD - Added VfsToLocal
+				Filesystem::Move(VfsToLocal(wideObjPath).c_str(), VfsToLocal(bakPath).c_str());
 			}
 
 			// remove linker directives which we don't want or need.
@@ -2919,7 +2932,8 @@ LiveModule::ErrorType::Enum LiveModule::Update(FileAttributeCache* fileCache, Di
 				}
 			}
 
-			coff::WriteRaw(wideObjPath.c_str(), rawCoff, removalStrategy);
+			// EPIC MOD - Added VfsToLocal
+			coff::WriteRaw(VfsToLocal(wideObjPath).c_str(), rawCoff, removalStrategy);
 			coff::DestroyRaw(rawCoff);
 		}
 	}
@@ -2955,7 +2969,9 @@ LiveModule::ErrorType::Enum LiveModule::Update(FileAttributeCache* fileCache, Di
 		const symbols::Compiland* compiland = compilandIt->second;
 		const uint32_t compilandUniqueId = symbols::GetCompilandId(compiland, wideObjPath.c_str(), modifiedOrNewObjFiles);
 
-		coff::ObjFile* coffFile = coff::OpenObj(wideObjPath.c_str());
+		// EPIC MOD - Added VfsToLocal
+		coff::ObjFile* coffFile = coff::OpenObj(VfsToLocal(wideObjPath).c_str());
+		
 		if (coffFile && coffFile->memoryFile)
 		{
 			coff::ExternalSymbolDB* externalSymbolDb = coff::GatherExternalSymbolDatabase(coffFile, compilandUniqueId);
@@ -3014,12 +3030,15 @@ LiveModule::ErrorType::Enum LiveModule::Update(FileAttributeCache* fileCache, Di
 		isExeOrPdbFileStillThere = false;
 		pdbPath = string::Replace(originalPdbPath, L".pdb", patchInstanceStr + std::wstring(L".pdb"));
 		exePath = string::Replace(originalPdbPath, L".pdb", patchInstanceStr + std::wstring(L".exe"));
-		const Filesystem::PathAttributes& pdbAttributes = Filesystem::GetAttributes(pdbPath.c_str());
-		const Filesystem::PathAttributes& exeAttributes = Filesystem::GetAttributes(exePath.c_str());
+		
+		// EPIC MOD - Added VfsToLocal
+		const Filesystem::PathAttributes& pdbAttributes = Filesystem::GetAttributes(VfsToLocal(pdbPath).c_str());
+		const Filesystem::PathAttributes& exeAttributes = Filesystem::GetAttributes(VfsToLocal(exePath).c_str());
 
 		if (Filesystem::DoesExist(pdbAttributes))
 		{
-			if (!Filesystem::DeleteIfExists(pdbPath.c_str()))
+			// EPIC MOD - Added VfsToLocal
+			if (!Filesystem::DeleteIfExists(VfsToLocal(pdbPath).c_str()))
 			{
 				// PDB file could not be deleted
 				isExeOrPdbFileStillThere = true;
@@ -3028,7 +3047,8 @@ LiveModule::ErrorType::Enum LiveModule::Update(FileAttributeCache* fileCache, Di
 		
 		if (Filesystem::DoesExist(exeAttributes))
 		{
-			if (!Filesystem::DeleteIfExists(exePath.c_str()))
+			// EPIC MOD - Added VfsToLocal
+			if (!Filesystem::DeleteIfExists(VfsToLocal(exePath).c_str()))
 			{
 				// EXE file could not be deleted
 				isExeOrPdbFileStillThere = true;
@@ -3138,8 +3158,8 @@ LiveModule::ErrorType::Enum LiveModule::Update(FileAttributeCache* fileCache, Di
 
 	telemetry::Scope linkScope("Linking");
 
-	const std::wstring linkerPath = GetLinkerPath(m_linkerDB);
-	const std::wstring linkerWorkingDirectory = (m_linkerDB->workingDirectory.GetLength() != 0u)
+	std::wstring linkerPath = GetLinkerPath(m_linkerDB);
+	std::wstring linkerWorkingDirectory = (m_linkerDB->workingDirectory.GetLength() != 0u)
 		? string::ToWideString(m_linkerDB->workingDirectory)	// we have a valid working directory
 		: Filesystem::GetDirectory(linkerPath.c_str()).GetString();						// no valid working directory, take the linker directory instead
 
@@ -3149,15 +3169,33 @@ LiveModule::ErrorType::Enum LiveModule::Update(FileAttributeCache* fileCache, Di
 	const Filesystem::Path responseFilePath = Filesystem::GenerateTempFilename();
 	Filesystem::CreateFileWithData(responseFilePath.GetString(), linkerOptions.c_str(), linkerOptions.size() * sizeof(wchar_t));
 
-	std::wstring linkerCommandLine = Filesystem::GetFilename(linkerPath.c_str()).GetString();
+	// BEGIN EPIC MOD
+	const Process::Environment linkerEnvironment = compiler::GetEnvironmentFromCache(linkerPath.c_str());
+	const void* linkerEnvironmentData = linkerEnvironment.data;
+
+	std::wstring linkerCommandLine;
+
+	// If there are vfs entries, we need to run the linker through ubacli with -vfs options to make sure linker result is the same as when running through UBT
+	if (!vfsEntries.empty())
+	{
+		for (auto& entry : vfsEntries)
+		{
+			linkerCommandLine += L" -vfs=\"" + entry.virtualPath + L";" + entry.localPath + L"\"";
+		}
+		linkerCommandLine += L" local ";
+		linkerCommandLine += linkerPath.c_str();
+		linkerPath = *FPaths::ConvertRelativePathToFull(FPaths::EngineDir() / L"Binaries/Win64/UnrealBuildAccelerator/x64/UbaCli.exe");
+		linkerWorkingDirectory = VfsToLocal(linkerWorkingDirectory);
+	}
+	else
+	{
+		linkerCommandLine += linkerPath.c_str();
+	}
+
 	linkerCommandLine += L" @\"";
 	linkerCommandLine += responseFilePath.GetString();
 	linkerCommandLine += L"\"";
 
-	const Process::Environment linkerEnvironment = compiler::GetEnvironmentFromCache(linkerPath.c_str());
-	const void* linkerEnvironmentData = linkerEnvironment.data;
-
-	// BEGIN EPIC MOD
 	GLiveCodingServer->GetStatusChangeDelegate().ExecuteIfBound(L"Linking patch...");
 	// END EPIC MOD
 
@@ -3173,12 +3211,14 @@ LiveModule::ErrorType::Enum LiveModule::Update(FileAttributeCache* fileCache, Di
 		const std::wstring originalPath(string::ToWideString(objPath));
 		const std::wstring bakPath = originalPath + L".bak";
 
-		const Filesystem::PathAttributes& attributes = Filesystem::GetAttributes(bakPath.c_str());
+		// BEGIN EPIC MOD - Added VfsToLocal
+		const Filesystem::PathAttributes& attributes = Filesystem::GetAttributes(VfsToLocal(bakPath).c_str());
 		if (Filesystem::DoesExist(attributes))
 		{
-			Filesystem::Delete(originalPath.c_str());
-			Filesystem::Move(bakPath.c_str(), originalPath.c_str());
+			Filesystem::Delete(VfsToLocal(originalPath).c_str());
+			Filesystem::Move(VfsToLocal(bakPath).c_str(), VfsToLocal(originalPath).c_str());
 		}
+		// END EPIC MOD
 	}
 
 	const std::wstring linkerProcessStdout = Process::GetStdOutData(linkerProcessContext);
@@ -3237,7 +3277,8 @@ LiveModule::ErrorType::Enum LiveModule::Update(FileAttributeCache* fileCache, Di
 	// END EPIC MOD
 
 	// try to load patch image
-	executable::Image* image = executable::OpenImage(exePath.c_str(), Filesystem::OpenMode::READ_WRITE);
+	// EPIC MOD - Added VfsToLocal
+	executable::Image* image = executable::OpenImage(VfsToLocal(exePath).c_str(), Filesystem::OpenMode::READ_WRITE);
 	if (!image)
 	{
 		LC_ERROR_USER("Cannot load patch executable %S", exePath.c_str());
@@ -3280,7 +3321,8 @@ LiveModule::ErrorType::Enum LiveModule::Update(FileAttributeCache* fileCache, Di
 			const PerProcessData& data = processData[i];
 
 			commands::LoadPatch cmd = {};
-			wcscpy_s(cmd.path, exePath.c_str());
+			// EPIC MOD - Added VfsToLocal
+			wcscpy_s(cmd.path, VfsToLocal(exePath).c_str());
 
 #if LC_64_BIT
 			// before doing anything further, we need to ensure that the patch can be loaded into the address space at a suitable location.
@@ -3307,14 +3349,16 @@ LiveModule::ErrorType::Enum LiveModule::Update(FileAttributeCache* fileCache, Di
 			{
 				// this image needs to be copied. create a new name based on the process ID, which must be unique
 				rebasedExePath = string::Replace(rebasedExePath, L".exe", std::wstring(L"_PID_") + std::to_wstring(+data.liveProcess->GetProcessId()) + L".exe");
-				Filesystem::Copy(exePath.c_str(), rebasedExePath.c_str());
-				wcscpy_s(cmd.path, rebasedExePath.c_str());
+				// EPIC MOD - Added VfsToLocal
+				Filesystem::Copy(VfsToLocal(exePath).c_str(), VfsToLocal(rebasedExePath).c_str());
+				wcscpy_s(cmd.path, VfsToLocal(rebasedExePath).c_str());
 			}
 			
 			if (imageNeedsToBeRebased)
 			{
 				// rebase the patch image to its preferred base address
-				executable::Image* rebasedImage = executable::OpenImage(rebasedExePath.c_str(), Filesystem::OpenMode::READ_WRITE);
+				// EPIC MOD - Added VfsToLocal
+				executable::Image* rebasedImage = executable::OpenImage(VfsToLocal(rebasedExePath).c_str(), Filesystem::OpenMode::READ_WRITE);
 				LC_LOG_DEV("Rebasing patch executable to image base 0x%" PRIX64 " (PID: %d)", preferredImageBase, data.liveProcess->GetProcessId());
 				executable::RebaseImage(rebasedImage, preferredImageBase);
 				executable::CloseImage(rebasedImage);
@@ -3416,7 +3460,8 @@ LiveModule::ErrorType::Enum LiveModule::Update(FileAttributeCache* fileCache, Di
 
 	telemetry::Scope loadPatchPDBScope("Loading PDB database");
 
-	symbols::Provider* patchSymbolProvider = symbols::OpenEXE(exePath.c_str(), symbols::OpenOptions::ACCUMULATE_SIZE);
+	// EPIC MOD - Added VfsToLocal
+	symbols::Provider* patchSymbolProvider = symbols::OpenEXE(VfsToLocal(exePath).c_str(), symbols::OpenOptions::ACCUMULATE_SIZE);
 	symbols::DiaCompilandDB* patch_diaCompilandDb = symbols::GatherDiaCompilands(patchSymbolProvider);
 	IDiaSymbol* patch_linkerSymbol = symbols::FindLinkerSymbol(patch_diaCompilandDb);
 
@@ -3432,9 +3477,10 @@ LiveModule::ErrorType::Enum LiveModule::Update(FileAttributeCache* fileCache, Di
 	scheduler::RunTask(taskPatch_symbolDB);
 
 
-	auto taskPatch_contributionDB = scheduler::CreateTask(taskRootPatchLoading, [exePath]()
+	// EPIC MOD - Added VfsToLocal
+	auto taskPatch_contributionDB = scheduler::CreateTask(taskRootPatchLoading, [realExePath = VfsToLocal(exePath)]()
 	{
-		symbols::Provider* localProvider = symbols::OpenEXE(exePath.c_str(), symbols::OpenOptions::NONE);
+		symbols::Provider* localProvider = symbols::OpenEXE(realExePath.c_str(), symbols::OpenOptions::NONE);
 		symbols::DiaCompilandDB* localDiaCompilandDb = symbols::GatherDiaCompilands(localProvider);
 
 		auto db = symbols::GatherContributions(localProvider);
@@ -3450,9 +3496,10 @@ LiveModule::ErrorType::Enum LiveModule::Update(FileAttributeCache* fileCache, Di
 	// note that we only gather symbols from .obj contained in the new patch executable.
 	// therefore we need to extract its compiland database as well, and cannot use the one from
 	// the original executable.
-	auto taskPatch_compilandDB = scheduler::CreateTask(taskRootPatchLoading, [this, exePath]()
+	// EPIC MOD - Added VfsToLocal
+	auto taskPatch_compilandDB = scheduler::CreateTask(taskRootPatchLoading, [this, realExePath = VfsToLocal(exePath)]()
 	{
-		symbols::Provider* localProvider = symbols::OpenEXE(exePath.c_str(), symbols::OpenOptions::NONE);
+		symbols::Provider* localProvider = symbols::OpenEXE(realExePath.c_str(), symbols::OpenOptions::NONE);
 		symbols::DiaCompilandDB* localDiaCompilandDb = symbols::GatherDiaCompilands(localProvider);
 
 		uint32_t options = 0u;
@@ -3567,7 +3614,8 @@ LiveModule::ErrorType::Enum LiveModule::Update(FileAttributeCache* fileCache, Di
 		// update the COFF cache for new patch compilands.
 		// there may be files for which we don't have a database yet, even though we updated the database for all compiled files.
 		// this can happen when a new .obj that is part of a library is linked in for the first time.
-		types::vector<symbols::ObjPath> updatedCoffs = UpdateCoffCache(patch_compilandDB->compilands, m_coffCache, CacheUpdate::NON_EXISTANT, modifiedOrNewObjFiles);
+		// EPIC MOD - Added VfsToLocal
+		types::vector<symbols::ObjPath> updatedCoffs = UpdateCoffCache(VfsToLocal, patch_compilandDB->compilands, m_coffCache, CacheUpdate::NON_EXISTANT, modifiedOrNewObjFiles);
 
 
 		// similarly, reconstruct symbols and dynamic initializers for new .obj that have been pulled in for the first time.
@@ -3626,7 +3674,8 @@ LiveModule::ErrorType::Enum LiveModule::Update(FileAttributeCache* fileCache, Di
 	// END EPIC MOD
 
 	// reconstruct symbols for all compilands that are part of the new patch executable
-	executable::Image* patchImage = executable::OpenImage(exePath.c_str(), Filesystem::OpenMode::READ);
+	// EPIC MOD - Added VfsToLocal
+	executable::Image* patchImage = executable::OpenImage(VfsToLocal(exePath).c_str(), Filesystem::OpenMode::READ);
 	executable::ImageSectionDB* patchImageSections = executable::GatherImageSectionDB(patchImage);
 
 	// gather the dynamic initializers and remaining symbols by walking the module
diff --git a/Engine/Source/Developer/Windows/LiveCodingServer/Private/External/LC_LiveModule.h b/Engine/Source/Developer/Windows/LiveCodingServer/Private/External/LC_LiveModule.h
--- a/Engine/Source/Developer/Windows/LiveCodingServer/Private/External/LC_LiveModule.h
+++ b/Engine/Source/Developer/Windows/LiveCodingServer/Private/External/LC_LiveModule.h
@@ -24,6 +24,12 @@ class ModulePatch;
 class LiveProcess;
 class VirtualMemoryRange;
 
+struct VfsEntry
+{
+	std::wstring virtualPath;
+	std::wstring localPath;
+};
+
 class LiveModule
 {
 public:
@@ -82,7 +88,7 @@ public:
 	// in EXTERNAL_BUILD_SYSTEM mode, this does not compile files but builds a patch containing modified .objs, loading the patch into the host application.
 	// optionally, an array of modified or new .objs can be given in this mode, which builds a patch containing these files, not checking for any other modifications.
 	// BEGIN EPIC MOD
-	ErrorType::Enum Update(FileAttributeCache* fileCache, DirectoryCache* directoryCache, UpdateType::Enum updateType, const types::vector<symbols::ModifiedObjFile>& modifiedOrNewObjFiles, const types::vector<std::wstring>& additionalLibraries);
+	ErrorType::Enum Update(FileAttributeCache* fileCache, DirectoryCache* directoryCache, UpdateType::Enum updateType, const types::vector<symbols::ModifiedObjFile>& modifiedOrNewObjFiles, const types::vector<std::wstring>& additionalLibraries, const types::vector<VfsEntry>& vfsEntries);
 	// END EPIC MOD
 	bool InstallCompiledPatches(LiveProcess* liveProcess, void* originalModuleBase);
 
diff --git a/Engine/Source/Developer/Windows/LiveCodingServer/Private/External/LC_ServerCommandThread.cpp b/Engine/Source/Developer/Windows/LiveCodingServer/Private/External/LC_ServerCommandThread.cpp
--- a/Engine/Source/Developer/Windows/LiveCodingServer/Private/External/LC_ServerCommandThread.cpp
+++ b/Engine/Source/Developer/Windows/LiveCodingServer/Private/External/LC_ServerCommandThread.cpp
@@ -777,7 +777,7 @@ bool ServerCommandThread::EnableRequiredModules(const TArray<FString>& RequiredM
 // END EPIC MOD
 
 // BEGIN EPIC MOD
-void ServerCommandThread::CompileChanges(bool didAllProcessesMakeProgress, commands::PostCompileResult& postCompileResult)
+void ServerCommandThread::CompileChanges(bool didAllProcessesMakeProgress, commands::PostCompileResult& postCompileResult, types::vector<VfsEntry>& vfsEntries)
 // END EPIC MOD
 {
 	// recompile files, if any
@@ -805,6 +805,9 @@ void ServerCommandThread::CompileChanges(bool didAllProcessesMakeProgress, comma
 
 		GLiveCodingServer->GetStatusChangeDelegate().ExecuteIfBound(L"Compiling changes for live coding...");
 
+		// EPIC MOD - Added VfsToLocal
+		FVfsEntries VfsToLocalPaths;
+
 		// Keep retrying the compile until we've added all the required modules
 		FModuleToModuleFiles ModuleToModuleFiles;
 		ELiveCodingCompileReason CompileReason = ELiveCodingCompileReason::Initial;
@@ -833,7 +836,8 @@ void ServerCommandThread::CompileChanges(bool didAllProcessesMakeProgress, comma
 
 			// Execute the compile
 			TArray<FString> RequiredModules;
-			ELiveCodingCompileResult CompileResult = CompileDelegate.Execute(Targets, ValidModules, LazyLoadModules, RequiredModules, ModuleToModuleFiles, CompileReason);
+			// EPIC MOD - Added VfsToLocal
+			ELiveCodingCompileResult CompileResult = CompileDelegate.Execute(Targets, ValidModules, LazyLoadModules, RequiredModules, ModuleToModuleFiles, VfsToLocalPaths, CompileReason);
 			if (CompileResult == ELiveCodingCompileResult::Success)
 			{
 				break;
@@ -865,6 +869,21 @@ void ServerCommandThread::CompileChanges(bool didAllProcessesMakeProgress, comma
 		// Reset the unity file cache
 		symbols::ResetCachedUnityManifests();
 
+		// BEGIN EPIC MOD - Added VfsToLocal
+		for (TPair<FString, FString>& kv : VfsToLocalPaths)
+		{
+			vfsEntries.emplace_back(VfsEntry{*kv.Key, *kv.Value});
+		}
+
+		auto VfsToLocal = [&](const std::wstring& path)
+			{
+				for (auto& entry : vfsEntries)
+					if (entry.virtualPath.compare(0, entry.virtualPath.size(), path, 0, entry.virtualPath.size()) == 0)
+						return entry.localPath + (path.data() + entry.virtualPath.size());
+				return path;
+			};
+		// END EPIC MOD - Added VfsToLocal
+
 		// Build up a list of all the modified object files in each module
 		types::unordered_map<std::wstring, const LiveModule*> EnabledModulesByName;
 		for (const LiveModule* liveModule : m_liveModules)
@@ -875,7 +894,8 @@ void ServerCommandThread::CompileChanges(bool didAllProcessesMakeProgress, comma
 		for(const TPair<FString, FModuleFiles>& Pair : ModuleToModuleFiles)
 		{
 			const LiveModule* liveModule = nullptr;
-			std::wstring ModuleFileName = Filesystem::NormalizePath(*Pair.Key).GetString();
+			// EPIC MOD - Added VfsToLocal
+			std::wstring ModuleFileName = VfsToLocal(Filesystem::NormalizePath(*Pair.Key).GetString());
 			types::unordered_map<std::wstring, const LiveModule*>::iterator ix = EnabledModulesByName.find(ModuleFileName);
 			if (ix == EnabledModulesByName.end())
 			{
@@ -908,7 +928,8 @@ void ServerCommandThread::CompileChanges(bool didAllProcessesMakeProgress, comma
 			{
 				std::wstring NormalizedObjectFile = Filesystem::NormalizePath(*ObjectFile).GetString();
 
-				if (!liveModule->IsModifiedSource(NormalizedObjectFile.c_str()))
+				// EPIC MOD - Added VfsToLocal
+				if (!liveModule->IsModifiedSource(VfsToLocal(NormalizedObjectFile).c_str()))
 				{
 					continue;
 				}
@@ -935,7 +956,8 @@ void ServerCommandThread::CompileChanges(bool didAllProcessesMakeProgress, comma
 
 					// Move the new file into place
 					m_restoreFiles.push_back(std::make_pair(OriginalObjectFile, NormalizedObjectFile));
-					Filesystem::Move(NormalizedObjectFile.c_str(), OriginalObjectFile.c_str());
+					// EPIC MOD - Added VfsToLocal
+					Filesystem::Move(VfsToLocal(NormalizedObjectFile).c_str(), VfsToLocal(OriginalObjectFile).c_str());
 					NormalizedObjectFile = OriginalObjectFile;
 				}
 
@@ -1004,7 +1026,7 @@ void ServerCommandThread::CompileChanges(bool didAllProcessesMakeProgress, comma
 
 		// build a patch with the given list of .objs for this module
 		const types::vector<symbols::ModifiedObjFile>& objFiles = objFilesIt->second;
-		LiveModule::ErrorType::Enum moduleUpdateError = liveModule->Update(&fileCache, m_directoryCache, updateType, objFiles, libFiles);
+		LiveModule::ErrorType::Enum moduleUpdateError = liveModule->Update(&fileCache, m_directoryCache, updateType, objFiles, libFiles, vfsEntries);
 
 		// only accept new error conditions for this module if there haven't been any updates until now.
 		// this ensures that error conditions are kept and can be shown when updating several modules at once.
@@ -1365,19 +1387,31 @@ Thread::ReturnValue ServerCommandThread::CompileThread(void)
 			// END EPIC MOD
 
 			// BEGIN EPIC MOD 
+			types::vector<VfsEntry> vfsToLocalPaths;
 			commands::PostCompileResult postCompileResult = commands::PostCompileResult::Success;
-			CompileChanges(didAllProcessesMakeProgress, postCompileResult);
+			CompileChanges(didAllProcessesMakeProgress, postCompileResult, vfsToLocalPaths);
 			// END EPIC MOD
 
 			// BEGIN EPIC MOD - Add the ability for pre and post compile notifications
 			CallPostcompileHooks(didAllProcessesMakeProgress, postCompileResult);
 			// END EPIC MOD
 
+
+			// Begin EPIC MOD - Added VfsToLocal
+			auto VfsToLocal = [&](const std::wstring& path)
+				{
+					for (auto& entry : vfsToLocalPaths)
+						if (entry.virtualPath.compare(0, entry.virtualPath.size(), path, 0, entry.virtualPath.size()) == 0)
+							return entry.localPath + (path.data() + entry.virtualPath.size());
+					return path;
+				};
+			// EndEPIC MOD - Added VfsToLocal
+
 			// BEGIN EPIC MOD - Non-destructive compile
 			for (std::vector<std::pair<std::wstring, std::wstring>>::reverse_iterator it = m_restoreFiles.rbegin(); it != m_restoreFiles.rend(); it++)
 			{
-				Filesystem::DeleteIfExists(it->second.c_str());
-				Filesystem::Move(it->first.c_str(), it->second.c_str());
+				Filesystem::DeleteIfExists(VfsToLocal(it->second).c_str());
+				Filesystem::Move(VfsToLocal(it->first).c_str(), VfsToLocal(it->second).c_str());
 			}
 			m_restoreFiles.clear();
 			// END EPIC MOD
diff --git a/Engine/Source/Developer/Windows/LiveCodingServer/Private/External/LC_ServerCommandThread.h b/Engine/Source/Developer/Windows/LiveCodingServer/Private/External/LC_ServerCommandThread.h
--- a/Engine/Source/Developer/Windows/LiveCodingServer/Private/External/LC_ServerCommandThread.h
+++ b/Engine/Source/Developer/Windows/LiveCodingServer/Private/External/LC_ServerCommandThread.h
@@ -78,7 +78,7 @@ private:
 	LiveProcess* FindProcessById(Process::Id processId);
 
 	// BEGIN EPIC MOD
-	void CompileChanges(bool didAllProcessesMakeProgress, commands::PostCompileResult& postCompileResult);
+	void CompileChanges(bool didAllProcessesMakeProgress, commands::PostCompileResult& postCompileResult, types::vector<VfsEntry>& vfsToLocalPaths);
 	// END EPIC MOD
 
 	// actions
diff --git a/Engine/Source/Developer/Windows/LiveCodingServer/Public/ILiveCodingServer.h b/Engine/Source/Developer/Windows/LiveCodingServer/Public/ILiveCodingServer.h
--- a/Engine/Source/Developer/Windows/LiveCodingServer/Public/ILiveCodingServer.h
+++ b/Engine/Source/Developer/Windows/LiveCodingServer/Public/ILiveCodingServer.h
@@ -53,6 +53,8 @@ struct FModuleFiles
 };
 typedef TMap<FString, FModuleFiles> FModuleToModuleFiles;
 
+typedef TArray<TPair<FString, FString>> FVfsEntries;
+
 class ILiveCodingServer : public IModularFeature
 {
 public:
@@ -75,7 +77,7 @@ public:
 	DECLARE_DELEGATE_TwoParams(FLogOutputDelegate, ELiveCodingLogVerbosity, const wchar_t*);
 	virtual FLogOutputDelegate& GetLogOutputDelegate() = 0;
 
-	DECLARE_DELEGATE_RetVal_SixParams(ELiveCodingCompileResult, FCompileDelegate, const TArray<FString>&, const TArray<FString>&, const TSet<FString>&, TArray<FString>&, FModuleToModuleFiles&, ELiveCodingCompileReason);
+	DECLARE_DELEGATE_RetVal_SevenParams(ELiveCodingCompileResult, FCompileDelegate, const TArray<FString>&, const TArray<FString>&, const TSet<FString>&, TArray<FString>&, FModuleToModuleFiles&, FVfsEntries&, ELiveCodingCompileReason);
 	virtual FCompileDelegate& GetCompileDelegate() = 0;
 
 	DECLARE_DELEGATE(FCompileStartedDelegate);
diff --git a/Engine/Source/Programs/LiveCodingConsole/Private/LiveCodingConsole.cpp b/Engine/Source/Programs/LiveCodingConsole/Private/LiveCodingConsole.cpp
--- a/Engine/Source/Programs/LiveCodingConsole/Private/LiveCodingConsole.cpp
+++ b/Engine/Source/Programs/LiveCodingConsole/Private/LiveCodingConsole.cpp
@@ -273,7 +273,7 @@ private:
 	}
 
 	ELiveCodingCompileResult CompilePatch(const TArray<FString>& Targets, const TArray<FString>& ValidModules, const TSet<FString>& LazyLoadModules,
-		TArray<FString>& RequiredModules, FModuleToModuleFiles& ModuleToModuleFiles, ELiveCodingCompileReason CompileReason)
+		TArray<FString>& RequiredModules, FModuleToModuleFiles& ModuleToModuleFiles, FVfsEntries& VfsEntries, ELiveCodingCompileReason CompileReason)
 	{
 		// Get the UBT path
 		FString Executable;
@@ -422,6 +422,8 @@ private:
 			}
 		}
 
+		VfsEntries = Manifest.VfsEntries;
+
 		return ELiveCodingCompileResult::Success;
 	}
 
diff --git a/Engine/Source/Programs/LiveCodingConsole/Private/LiveCodingManifest.cpp b/Engine/Source/Programs/LiveCodingConsole/Private/LiveCodingManifest.cpp
--- a/Engine/Source/Programs/LiveCodingConsole/Private/LiveCodingManifest.cpp
+++ b/Engine/Source/Programs/LiveCodingConsole/Private/LiveCodingManifest.cpp
@@ -111,5 +111,39 @@ bool FLiveCodingManifest::Parse(FJsonObject& Object, FString& OutFailReason)
 		}
 	}
 
+	const TArray<TSharedPtr<FJsonValue>>* VfsArray;
+	if (!Object.TryGetArrayField(TEXT("Vfs"), VfsArray))
+	{
+		OutFailReason = TEXT("missing Vfs field");
+		return false;
+	}
+
+	for (const TSharedPtr<FJsonValue>& VfsValue : *VfsArray)
+	{
+		if (VfsValue->Type != EJson::Object)
+		{
+			OutFailReason = TEXT("invalid vfs object");
+			return false;
+		}
+
+		const FJsonObject& VfsObject = *VfsValue->AsObject();
+
+		FString VirtualPath;
+		if (!VfsObject.TryGetStringField(TEXT("Virtual"), VirtualPath))
+		{
+			OutFailReason = TEXT("missing Vfs 'Virtual' field");
+			return false;
+		}
+
+		FString LocalPath;
+		if (!VfsObject.TryGetStringField(TEXT("Local"), LocalPath))
+		{
+			OutFailReason = TEXT("missing Vfs 'Local' field");
+			return false;
+		}
+
+		VfsEntries.Emplace(VirtualPath, LocalPath);
+	}
+
 	return true;
 }
diff --git a/Engine/Source/Programs/LiveCodingConsole/Private/LiveCodingManifest.h b/Engine/Source/Programs/LiveCodingConsole/Private/LiveCodingManifest.h
--- a/Engine/Source/Programs/LiveCodingConsole/Private/LiveCodingManifest.h
+++ b/Engine/Source/Programs/LiveCodingConsole/Private/LiveCodingManifest.h
@@ -12,6 +12,7 @@ struct FLiveCodingManifest
 	TMap<FString, FString> LinkerEnvironment;
 	TMap<FString, TArray<FString>> BinaryToObjectFiles;
 	TMap<FString, TArray<FString>> Libraries;
+	TArray<TPair<FString, FString>> VfsEntries;
 
 	bool Read(const TCHAR* FileName, FString& OutFailReason);
 	bool Parse(FJsonObject& Object, FString& OutFailReason);
