[LiveCodingConsole]
* Fixed so vfs path can handle link.exe paths that has space in them

diff --git a/Engine/Source/Developer/Windows/LiveCodingServer/Private/External/LC_LiveModule.cpp b/Engine/Source/Developer/Windows/LiveCodingServer/Private/External/LC_LiveModule.cpp
--- a/Engine/Source/Developer/Windows/LiveCodingServer/Private/External/LC_LiveModule.cpp
+++ b/Engine/Source/Developer/Windows/LiveCodingServer/Private/External/LC_LiveModule.cpp
@@ -3182,8 +3182,9 @@ LiveModule::ErrorType::Enum LiveModule::Update(FileAttributeCache* fileCache, Di
 		{
 			linkerCommandLine += L" -vfs=\"" + entry.virtualPath + L";" + entry.localPath + L"\"";
 		}
-		linkerCommandLine += L" local ";
+		linkerCommandLine += L" local \"";
 		linkerCommandLine += linkerPath.c_str();
+		linkerCommandLine += '\"';
 		linkerPath = *FPaths::ConvertRelativePathToFull(FPaths::EngineDir() / L"Binaries/Win64/UnrealBuildAccelerator/x64/UbaCli.exe");
 		linkerWorkingDirectory = VfsToLocal(linkerWorkingDirectory);
 	}
