[LiveCodingConsole]
* Fixed non-vfs path when path to linker has spaces in it

diff --git a/Engine/Source/Developer/Windows/LiveCodingServer/Private/External/LC_LiveModule.cpp b/Engine/Source/Developer/Windows/LiveCodingServer/Private/External/LC_LiveModule.cpp
--- a/Engine/Source/Developer/Windows/LiveCodingServer/Private/External/LC_LiveModule.cpp
+++ b/Engine/Source/Developer/Windows/LiveCodingServer/Private/External/LC_LiveModule.cpp
@@ -3190,7 +3190,7 @@ LiveModule::ErrorType::Enum LiveModule::Update(FileAttributeCache* fileCache, Di
 	}
 	else
 	{
-		linkerCommandLine += linkerPath.c_str();
+		linkerCommandLine = Filesystem::GetFilename(linkerPath.c_str()).GetString();
 	}
 
 	linkerCommandLine += L" @\"";
