MovieRenderQueue - Fix low SampleRate when exporting wav on Mac
Also move the default value for Windows platforms in INI instead

#tests Tested to export wav on Windows & Mac and check the SampleRate to be the value defined in BaseEngine.ini

diff --git a/Engine/Config/BaseEngine.ini b/Engine/Config/BaseEngine.ini
--- a/Engine/Config/BaseEngine.ini
+++ b/Engine/Config/BaseEngine.ini
@@ -3342,6 +3342,17 @@ bEnable2DPhysics=false
 bEnableDistanceFields=true
 bEnableRayTracing=true
 bTarget32Bit=false
+CacheSizeKB=65536
+MaxSampleRate=48000
+HighSampleRate=32000
+MedSampleRate=24000
+LowSampleRate=12000
+MinSampleRate=8000
+CompressionQualityModifier=1
+AudioSampleRate=48000
+AudioCallbackBufferFrameSize=1024
+AudioNumBuffersToEnqueue=1
+AudioNumSourceWorkers=4
 
 [/Script/LinuxTargetPlatform.LinuxTargetSettings]
 +TargetedRHIs=SF_VULKAN_SM5
@@ -3358,6 +3369,10 @@ EditorDefaultArchitecture=MacTargetArchitectureHost
 DefaultArchitecture=MacTargetArchitectureHost
 bBuildAllSupportedOnBuildMachine=true
 bEnableRayTracing=false
+AudioSampleRate=48000
+AudioCallbackBufferFrameSize=1024
+AudioNumBuffersToEnqueue=1
+AudioNumSourceWorkers=4
 
 [/Script/MacTargetPlatform.XcodeProjectSettings]
 bUseModernXcode=true
diff --git a/Engine/Source/Developer/Windows/WindowsTargetPlatformSettings/Private/WindowsTargetPlatformClasses.cpp b/Engine/Source/Developer/Windows/WindowsTargetPlatformSettings/Private/WindowsTargetPlatformClasses.cpp
--- a/Engine/Source/Developer/Windows/WindowsTargetPlatformSettings/Private/WindowsTargetPlatformClasses.cpp
+++ b/Engine/Source/Developer/Windows/WindowsTargetPlatformSettings/Private/WindowsTargetPlatformClasses.cpp
@@ -7,19 +7,8 @@
 
 UWindowsTargetSettings::UWindowsTargetSettings( const FObjectInitializer& ObjectInitializer )
 	: Super(ObjectInitializer)
-	, CacheSizeKB(65536)
-	, MaxSampleRate(48000)
-	, HighSampleRate(32000)
-	, MedSampleRate(24000)
-	, LowSampleRate(12000)
-	, MinSampleRate(8000)
-	, CompressionQualityModifier(1)
 {
-	// Default windows settings
-	AudioSampleRate = 48000;
-	AudioCallbackBufferFrameSize = 1024;
-	AudioNumBuffersToEnqueue = 1;
-	AudioNumSourceWorkers = 4;
+	// Default values are defined in BaseEngine.ini
 }
 
 static bool FilterShaderPlatform_D3D12(const FString& InShaderPlatform)
