LiveCodingServer: Fix C++20 compile issues

#rnx

diff --git a/Engine/Source/Developer/Windows/LiveCoding/Private/External/LC_Logging.h b/Engine/Source/Developer/Windows/LiveCoding/Private/External/LC_Logging.h
--- a/Engine/Source/Developer/Windows/LiveCoding/Private/External/LC_Logging.h
+++ b/Engine/Source/Developer/Windows/LiveCoding/Private/External/LC_Logging.h
@@ -75,19 +75,21 @@ namespace Logging
 #define LC_LOG_INDENT_DEV						const Logging::Indent LC_PP_UNIQUE_NAME(autoIndentDev)(Logging::Channel::DEV)
 #define LC_LOG_INDENT_TELEMETRY					const Logging::Indent LC_PP_UNIQUE_NAME(autoIndentTelemetry)(Logging::Channel::TELEMETRY)
 
-#define LC_LOG_USER(_format, ...)				Logging::Log(Logging::Channel::USER, Logging::Type::LOG_INFO, "%s" _format "\n", Logging::GetIndentation(Logging::Channel::USER), __VA_ARGS__)
-#define LC_LOG_DEV(_format, ...)				Logging::Log(Logging::Channel::DEV, Logging::Type::LOG_INFO, "%s" _format "\n", Logging::GetIndentation(Logging::Channel::DEV), __VA_ARGS__)
-#define LC_LOG_TELEMETRY(_format, ...)			Logging::Log(Logging::Channel::TELEMETRY, Logging::Type::LOG_INFO, "%s" _format "\n", Logging::GetIndentation(Logging::Channel::TELEMETRY), __VA_ARGS__)
+// BEGIN EPIC MOD - __VA_ARGS__ -> ##__VA_ARGS__
+#define LC_LOG_USER(_format, ...)				Logging::Log(Logging::Channel::USER, Logging::Type::LOG_INFO, "%s" _format "\n", Logging::GetIndentation(Logging::Channel::USER), ##__VA_ARGS__)
+#define LC_LOG_DEV(_format, ...)				Logging::Log(Logging::Channel::DEV, Logging::Type::LOG_INFO, "%s" _format "\n", Logging::GetIndentation(Logging::Channel::DEV), ##__VA_ARGS__)
+#define LC_LOG_TELEMETRY(_format, ...)			Logging::Log(Logging::Channel::TELEMETRY, Logging::Type::LOG_INFO, "%s" _format "\n", Logging::GetIndentation(Logging::Channel::TELEMETRY), ##__VA_ARGS__)
 
-#define LC_WARNING_USER(_format, ...)			Logging::Log(Logging::Channel::USER, Logging::Type::LOG_WARNING, "%s" _format "\n", Logging::GetIndentation(Logging::Channel::USER), __VA_ARGS__)
-#define LC_WARNING_DEV(_format, ...)			Logging::Log(Logging::Channel::DEV, Logging::Type::LOG_WARNING, "WARNING: %s" _format "\n", Logging::GetIndentation(Logging::Channel::DEV), __VA_ARGS__)
-#define LC_WARNING_TELEMETRY(_format, ...)		Logging::Log(Logging::Channel::TELEMETRY, Logging::Type::LOG_WARNING, "WARNING: %s" _format "\n", Logging::GetIndentation(Logging::Channel::TELEMETRY), __VA_ARGS__)
+#define LC_WARNING_USER(_format, ...)			Logging::Log(Logging::Channel::USER, Logging::Type::LOG_WARNING, "%s" _format "\n", Logging::GetIndentation(Logging::Channel::USER), ##__VA_ARGS__)
+#define LC_WARNING_DEV(_format, ...)			Logging::Log(Logging::Channel::DEV, Logging::Type::LOG_WARNING, "WARNING: %s" _format "\n", Logging::GetIndentation(Logging::Channel::DEV), ##__VA_ARGS__)
+#define LC_WARNING_TELEMETRY(_format, ...)		Logging::Log(Logging::Channel::TELEMETRY, Logging::Type::LOG_WARNING, "WARNING: %s" _format "\n", Logging::GetIndentation(Logging::Channel::TELEMETRY), ##__VA_ARGS__)
 
-#define LC_ERROR_USER(_format, ...)				Logging::Log(Logging::Channel::USER, Logging::Type::LOG_ERROR, "%s" _format "\n", Logging::GetIndentation(Logging::Channel::USER), __VA_ARGS__)
-#define LC_ERROR_DEV(_format, ...)				Logging::Log(Logging::Channel::DEV, Logging::Type::LOG_ERROR, "ERROR: %s" _format "\n", Logging::GetIndentation(Logging::Channel::DEV), __VA_ARGS__)
-#define LC_ERROR_TELEMETRY(_format, ...)		Logging::Log(Logging::Channel::TELEMETRY, Logging::Type::LOG_ERROR, "ERROR: %s" _format "\n", Logging::GetIndentation(Logging::Channel::TELEMETRY), __VA_ARGS__)
+#define LC_ERROR_USER(_format, ...)				Logging::Log(Logging::Channel::USER, Logging::Type::LOG_ERROR, "%s" _format "\n", Logging::GetIndentation(Logging::Channel::USER), ##__VA_ARGS__)
+#define LC_ERROR_DEV(_format, ...)				Logging::Log(Logging::Channel::DEV, Logging::Type::LOG_ERROR, "ERROR: %s" _format "\n", Logging::GetIndentation(Logging::Channel::DEV), ##__VA_ARGS__)
+#define LC_ERROR_TELEMETRY(_format, ...)		Logging::Log(Logging::Channel::TELEMETRY, Logging::Type::LOG_ERROR, "ERROR: %s" _format "\n", Logging::GetIndentation(Logging::Channel::TELEMETRY), ##__VA_ARGS__)
 
-#define LC_SUCCESS_USER(_format, ...)			Logging::Log(Logging::Channel::USER, Logging::Type::LOG_SUCCESS, "%s" _format "\n", Logging::GetIndentation(Logging::Channel::USER), __VA_ARGS__)
+#define LC_SUCCESS_USER(_format, ...)			Logging::Log(Logging::Channel::USER, Logging::Type::LOG_SUCCESS, "%s" _format "\n", Logging::GetIndentation(Logging::Channel::USER), ##__VA_ARGS__)
+// END EPIC MOD 
 
 
 #endif // LC_VERSION
\ No newline at end of file
diff --git a/Engine/Source/Developer/Windows/LiveCodingServer/LiveCodingServer.Build.cs b/Engine/Source/Developer/Windows/LiveCodingServer/LiveCodingServer.Build.cs
--- a/Engine/Source/Developer/Windows/LiveCodingServer/LiveCodingServer.Build.cs
+++ b/Engine/Source/Developer/Windows/LiveCodingServer/LiveCodingServer.Build.cs
@@ -7,8 +7,6 @@ public class LiveCodingServer : ModuleRules
 {
 	public LiveCodingServer(ReadOnlyTargetRules Target) : base(Target)
 	{
-		CppStandard = CppStandardVersion.Cpp17;
-
 		PrivateDependencyModuleNames.Add("Core");
         PrivateDependencyModuleNames.Add("Json");
         PrivateDependencyModuleNames.Add("LiveCoding");
diff --git a/Engine/Source/Developer/Windows/LiveCodingServer/Private/External/LC_LiveModule.cpp b/Engine/Source/Developer/Windows/LiveCodingServer/Private/External/LC_LiveModule.cpp
--- a/Engine/Source/Developer/Windows/LiveCodingServer/Private/External/LC_LiveModule.cpp
+++ b/Engine/Source/Developer/Windows/LiveCodingServer/Private/External/LC_LiveModule.cpp
@@ -1157,6 +1157,9 @@ void LiveModule::Load(symbols::Provider* provider, symbols::DiaCompilandDB* diaC
 			symbols::ObjPath objPath;
 
 			GatherResult(void) = default;
+			// BEGIN EPIC MOD
+			GatherResult(coff::CoffDB* _database, symbols::ObjPath _objPath) : database(_database), objPath(_objPath) {}
+			// END EPIC MOD
 			GatherResult(const GatherResult& other) = default;
 			GatherResult(GatherResult&& other) = default;
 
diff --git a/Engine/Source/Developer/Windows/LiveCodingServer/Private/External/LC_Symbols.h b/Engine/Source/Developer/Windows/LiveCodingServer/Private/External/LC_Symbols.h
--- a/Engine/Source/Developer/Windows/LiveCodingServer/Private/External/LC_Symbols.h
+++ b/Engine/Source/Developer/Windows/LiveCodingServer/Private/External/LC_Symbols.h
@@ -32,6 +32,10 @@ namespace symbols
 		ImmutableString name;
 		uint32_t rva;
 
+		// BEGIN EPIC MOD
+		Symbol(ImmutableString _name, uint32_t _rva) : name(_name), rva(_rva) {}
+		// END EPIC MOD
+
 		LC_DISABLE_ASSIGNMENT(Symbol);
 		LC_DISABLE_COPY(Symbol);
 	};
@@ -115,6 +119,13 @@ namespace symbols
 		bool isPartOfLibrary;
 		bool wasRecompiled;
 
+		// BEGIN EPIC MOD
+		Compiland(ImmutableString _originalObjPath, ImmutableString _srcPath, ImmutableString _pdbPath, ImmutableString _compilerPath, ImmutableString _commandLine, ImmutableString _workingDirectory, ImmutableString _amalgamationPath, CompilandSourceFiles* _sourceFiles, uint32_t _uniqueId, Type::Enum _type, bool _isPartOfLibrary, bool _wasRecompiled) :
+			originalObjPath(_originalObjPath), srcPath(_srcPath), pdbPath(_pdbPath), compilerPath(_compilerPath), commandLine(_commandLine), workingDirectory(_workingDirectory), amalgamationPath(_amalgamationPath), sourceFiles(_sourceFiles), uniqueId(_uniqueId), type(_type), isPartOfLibrary(_isPartOfLibrary), wasRecompiled(_wasRecompiled)
+		{
+		}
+		// END EPIC MOD
+
 		LC_DISABLE_ASSIGNMENT(Compiland);
 		LC_DISABLE_COPY(Compiland);
 	};
