Support for Unused Interpolators removal for DXC shader compiler. Enabled by default in platform settings just for PC by CVarRemoveUnusedInterpolators.
Seems like there is bug in DXC compiler which prevents SV_Coverage semantic from being detected as Input Parameter in the shader reflection data. Other system interpolators are fine... Always adding SV_Coverage as Used interpolator to mitigate this issue.

diff --git a/Engine/Source/Developer/Windows/ShaderFormatD3D/Private/D3DShaderCompiler.cpp b/Engine/Source/Developer/Windows/ShaderFormatD3D/Private/D3DShaderCompiler.cpp
--- a/Engine/Source/Developer/Windows/ShaderFormatD3D/Private/D3DShaderCompiler.cpp
+++ b/Engine/Source/Developer/Windows/ShaderFormatD3D/Private/D3DShaderCompiler.cpp
@@ -726,8 +726,10 @@ static bool CompileAndProcessD3DShaderFXCExt(
 	const FString& PreprocessedShaderSource,
 	const FString& EntryPointName,
 	const FShaderParameterParser& ShaderParameterParser,
-	const TCHAR* ShaderProfile, bool bSecondPassAferUnusedInputRemoval,
-	TArray<FString>& FilteredErrors, FShaderCompilerOutput& Output)
+	const TCHAR* ShaderProfile,
+	const ED3DShaderModel ShaderModel, 
+	const bool bSecondPassAferUnusedInputRemoval,
+	FShaderCompilerOutput& Output)
 {
 	TRACE_CPUPROFILER_EVENT_SCOPE(CompileAndProcessD3DShaderFXCExt);
 
@@ -941,15 +943,44 @@ static bool CompileAndProcessD3DShaderFXCExt(
 	}
 	else
 	{
-		FilteredErrors.Add(TEXT("Couldn't find D3D shader compiler DLL"));
+		FShaderCompilerError NewError;
+		NewError.StrippedErrorMessage = TEXT("Couldn't find D3D shader compiler DLL");
+		Output.Errors.Add(NewError);
 		Result = E_FAIL;
 	}
 
 	// Filter any errors.
+	TArray<FString> FilteredErrors;
 	void* ErrorBuffer = Errors ? Errors->GetBufferPointer() : NULL;
 	if (ErrorBuffer)
 	{
 		D3D11FilterShaderCompileWarnings(ANSI_TO_TCHAR(ErrorBuffer), FilteredErrors);
+
+		// Process errors
+		for (int32 ErrorIndex = 0; ErrorIndex < FilteredErrors.Num(); ErrorIndex++)
+		{
+			const FString& CurrentError = FilteredErrors[ErrorIndex];
+			FShaderCompilerError NewError;
+
+			// Extract filename and line number from FXC output with format:
+			// "d:\Project\Binaries\BasePassPixelShader(30,7): error X3000: invalid target or usage string"
+			int32 FirstParenIndex = CurrentError.Find(TEXT("("));
+			int32 LastParenIndex = CurrentError.Find(TEXT("):"));
+			if (FirstParenIndex != INDEX_NONE &&
+				LastParenIndex != INDEX_NONE &&
+				LastParenIndex > FirstParenIndex)
+			{
+				// Extract and store error message with source filename
+				NewError.ErrorVirtualFilePath = CurrentError.Left(FirstParenIndex);
+				NewError.ErrorLineString = CurrentError.Mid(FirstParenIndex + 1, LastParenIndex - FirstParenIndex - FCString::Strlen(TEXT("(")));
+				NewError.StrippedErrorMessage = CurrentError.Right(CurrentError.Len() - LastParenIndex - FCString::Strlen(TEXT("):")));
+			}
+			else
+			{
+				NewError.StrippedErrorMessage = CurrentError;
+			}
+			Output.Errors.Add(NewError);
+		}
 	}
 
 	// Fail the compilation if certain extended features are being used, since those are not supported on all D3D11 cards.
@@ -977,7 +1008,9 @@ static bool CompileAndProcessD3DShaderFXCExt(
 					// dcl_globalFlags will contain enableDoublePrecisionFloatOps when the shader uses doubles, even though the docs on dcl_globalFlags don't say anything about this
 					if (DisassemblyStringW.Contains(TEXT("enableDoublePrecisionFloatOps")))
 					{
-						FilteredErrors.Add(TEXT("Shader uses double precision floats, which are not supported on all D3D11 hardware!"));
+						FShaderCompilerError NewError;
+						NewError.StrippedErrorMessage = TEXT("Shader uses double precision floats, which are not supported on all D3D11 hardware!");
+						Output.Errors.Add(NewError);
 						return false;
 					}
 				}
@@ -991,7 +1024,9 @@ static bool CompileAndProcessD3DShaderFXCExt(
 					// https://docs.microsoft.com/en-us/windows/win32/direct3ddxgi/format-support-for-direct3d-11-0-feature-level-hardware
 					if (DisassemblyStringW.Contains(TEXT("Typed UAV Load Additional Formats")))
 					{
-						FilteredErrors.Add(TEXT("Shader uses UAV loads from additional typed formats, which are not supported on all D3D11 hardware! Set r.D3D.CheckedForTypedUAVs=0 if you want to allow typed UAV loads for your project, or individual shaders can opt-in by specifying CFLAG_AllowTypedUAVLoads."));
+						FShaderCompilerError NewError;
+						NewError.StrippedErrorMessage = TEXT("Shader uses UAV loads from additional typed formats, which are not supported on all D3D11 hardware! Set r.D3D.CheckedForTypedUAVs=0 if you want to allow typed UAV loads for your project, or individual shaders can opt-in by specifying CFLAG_AllowTypedUAVLoads.");
+						Output.Errors.Add(NewError);
 						return false;
 					}
 				}
@@ -1026,128 +1061,50 @@ static bool CompileAndProcessD3DShaderFXCExt(
 			D3D11_SHADER_DESC ShaderDesc;
 			Reflector->GetDesc(&ShaderDesc);
 
-			if (Input.Target.Frequency == SF_Pixel)
+			// Remove unused interpolators from pixel shader 
+			// (propagated to corresponding VS from pipeline by later setting Output.bSupportsQueryingUsedAttributes and Output.UsedAttributes)
 			{
-				bool bFoundUnused = false;
-				for (uint32 Index = 0; Index < ShaderDesc.InputParameters; ++Index)
-				{
-					// VC++ horrible hack: Runtime ESP checks get confused and fail for some reason calling Reflector->GetInputParameterDesc() (because it comes from another DLL?)
-					// so "guard it" using the middle of an array; it's been confirmed NO corruption is really happening.
-					D3D11_SIGNATURE_PARAMETER_DESC ParamDescs[3];
-					D3D11_SIGNATURE_PARAMETER_DESC& ParamDesc = ParamDescs[1];
-					Reflector->GetInputParameterDesc(Index, &ParamDesc);
-					if (ParamDesc.SystemValueType == D3D_NAME_UNDEFINED)
+				ShaderCompileLambdaType ShaderCompileLambda = [CompileFlags](
+					const FShaderCompilerInput& Input,
+					const FString& PreprocessedShaderSource,
+					const FString& EntryPointName,
+					const FShaderParameterParser& ShaderParameterParser,
+					const TCHAR* ShaderProfile,
+					const ED3DShaderModel ShaderModel,
+					const bool bProcessingSecondTime,
+					FShaderCompilerOutput& Output)
 					{
-						if (ParamDesc.ReadWriteMask != 0)
-						{
-							FString SemanticName = ANSI_TO_TCHAR(ParamDesc.SemanticName);
-
-							CompileData.ShaderInputs.AddUnique(SemanticName);
-
-							// Add the number (for the case of TEXCOORD)
-							FString SemanticIndexName = FString::Printf(TEXT("%s%d"), *SemanticName, ParamDesc.SemanticIndex);
-							CompileData.ShaderInputs.AddUnique(SemanticIndexName);
-
-							// Add _centroid
-							CompileData.ShaderInputs.AddUnique(SemanticName + TEXT("_centroid"));
-							CompileData.ShaderInputs.AddUnique(SemanticIndexName + TEXT("_centroid"));
-						}
-						else
-						{
-							bFoundUnused = true;
-						}
-					}
-					else
-					{
-						//if (ParamDesc.ReadWriteMask != 0)
-						{
-							// Keep system values
-							CompileData.ShaderInputs.AddUnique(FString(ANSI_TO_TCHAR(ParamDesc.SemanticName)));
-						}
-					}
-				}
-
-				if (Input.Environment.CompilerFlags.Contains(CFLAG_ForceRemoveUnusedInterpolators) && Input.bCompilingForShaderPipeline && bFoundUnused && !bSecondPassAferUnusedInputRemoval)
+						return CompileAndProcessD3DShaderFXCExt(
+							CompileFlags,
+							Input,
+							PreprocessedShaderSource,
+							EntryPointName,
+							ShaderParameterParser,
+							ShaderProfile,
+							ShaderModel,
+							bProcessingSecondTime,
+							Output);
+					};
+
+				bool CompileResult = false;
+				const bool RemovedUnusedInterpolatorsApplied =
+					RemoveUnusedInterpolators<ShaderCompilerType::FXC, ShaderCompileLambdaType,
+					ID3D11ShaderReflection, D3D11_SHADER_DESC, D3D11_SIGNATURE_PARAMETER_DESC>(
+						Input,
+						PreprocessedShaderSource,
+						EntryPointName,
+						ShaderParameterParser,
+						ShaderProfile,
+						ShaderModel,
+						bSecondPassAferUnusedInputRemoval,
+						CompileData,
+						Reflector,
+						ShaderCompileLambda,
+						Output,
+						CompileResult);
+				if (RemovedUnusedInterpolatorsApplied)
 				{
-					// Rewrite the source removing the unused inputs so the bindings will match.
-					// We may need to do this more than once if unused inputs change after the removal. Ie. for complex shaders, what can happen is:
-					// pass1 detects that input A is not used, but input B and C are. Input A is removed, and we recompile (pass2). After the recompilation, we see that Input B is now also unused in pass2 
-					// (it became simpler and the compiler could see through that).
-					// Since unused inputs are passed to the next stage, that will cause us to generate a vertex shader that does not output B, but our pixel shader will still be expecting B on input,
-					// as it was rewritten based on the pass1 results.
-
-					FShaderCompilerOutput OriginalOutput = Output;
-					const int kMaxReasonableAttempts = 64;
-					for (int32 Attempt = 0; Attempt < kMaxReasonableAttempts; ++Attempt)
-					{
-						TArray<FString> RemoveErrors;
-						FString ModifiedShaderSource = PreprocessedShaderSource;
-						FString ModifiedEntryPointName = Input.EntryPointName;
-						if (RemoveUnusedInputs(ModifiedShaderSource, CompileData.ShaderInputs, ModifiedEntryPointName, RemoveErrors))
-						{
-							Output = OriginalOutput;
-							if (!CompileAndProcessD3DShaderFXCExt(CompileFlags, Input, ModifiedShaderSource, ModifiedEntryPointName, ShaderParameterParser, ShaderProfile, true, FilteredErrors, Output))
-							{
-								// if we failed to compile the shader, propagate the error up
-								return false;
-							}
-
-							// check if the ShaderInputs changed - if not, we're done here
-							if (Output.UsedAttributes.Num() == CompileData.ShaderInputs.Num())
-							{
-								Output.ModifiedShaderSource = MoveTemp(ModifiedShaderSource);
-								Output.ModifiedEntryPointName = MoveTemp(ModifiedEntryPointName);
-
-								return true;
-							}
-
-							// second pass cannot use more attributes than previously
-							if (Output.UsedAttributes.Num() > CompileData.ShaderInputs.Num())
-							{
-								UE_LOG(LogD3DShaderCompiler, Warning, TEXT("Second pass had more used attributes (%d) than first pass (%d)"), Output.UsedAttributes.Num(), CompileData.ShaderInputs.Num());
-								FShaderCompilerError NewError;
-								NewError.StrippedErrorMessage = FString::Printf(TEXT("Second pass had more used attributes (%d) than first pass (%d)"), Output.UsedAttributes.Num(), CompileData.ShaderInputs.Num());
-								Output = OriginalOutput;
-								Output.Errors.Add(NewError);
-								break;
-							}
-
-							// if we're about to run out of attempts, report
-							if (Attempt >= kMaxReasonableAttempts - 1)
-							{
-								UE_LOG(LogD3DShaderCompiler, Warning, TEXT("Unable to determine unused inputs after %d attempts (last number of used attributes: %d, previous step:%d)!"),
-									Attempt + 1,
-									Output.UsedAttributes.Num(),
-									CompileData.ShaderInputs.Num()
-									);
-								FShaderCompilerError NewError;
-								NewError.StrippedErrorMessage = FString::Printf(TEXT("Unable to determine unused inputs after %d attempts (last number of used attributes: %d, previous step:%d)!"),
-									Attempt + 1,
-									Output.UsedAttributes.Num(),
-									CompileData.ShaderInputs.Num()
-									);
-								Output = OriginalOutput;
-								Output.Errors.Add(NewError);
-								break;
-							}
-
-							CompileData.ShaderInputs = Output.UsedAttributes;
-							// go around to remove newly identified unused inputs
-						}
-						else
-						{
-							UE_LOG(LogD3DShaderCompiler, Warning, TEXT("Failed to remove unused inputs from shader: %s"), *Input.GenerateShaderName());
-							for (const FString& ErrorMessage : RemoveErrors)
-							{
-								// Add error to shader output but also make sure the error shows up on build farm by emitting a log entry
-								UE_LOG(LogD3DShaderCompiler, Warning, TEXT("%s"), *ErrorMessage);
-								FShaderCompilerError NewError;
-								NewError.StrippedErrorMessage = ErrorMessage;
-								Output.Errors.Add(NewError);
-							}
-							break;
-						}
-					}
+					return CompileResult;
 				}
 			}
 
@@ -1167,13 +1124,22 @@ static bool CompileAndProcessD3DShaderFXCExt(
 		}
 		else
 		{
-			FilteredErrors.Add(TEXT("Couldn't find shader reflection function in D3D Compiler DLL"));
+			FShaderCompilerError NewError;
+			NewError.StrippedErrorMessage = TEXT("Couldn't find shader reflection function in D3D Compiler DLL");
+			Output.Errors.Add(NewError);
 			Result = E_FAIL;
 			Output.bSucceeded = false;
 		}
 		
-		if (!ValidateResourceCounts(CompileData, FilteredErrors))
+		TArray<FString> FilteredErrorsFromResourceCounts;
+		if (!ValidateResourceCounts(CompileData, FilteredErrorsFromResourceCounts))
 		{
+			for (FString& error : FilteredErrorsFromResourceCounts)
+			{
+				FShaderCompilerError NewError;
+				NewError.StrippedErrorMessage = error;
+				Output.Errors.Add(NewError);
+			}
 			Result = E_FAIL;
 			Output.bSucceeded = false;
 		}
@@ -1181,7 +1147,9 @@ static bool CompileAndProcessD3DShaderFXCExt(
 		// Check for resource limits for feature level 11.0
 		if (CompileData.NumUAVs > GD3DMaximumNumUAVs)
 		{
-			FilteredErrors.Add(FString::Printf(TEXT("Number of UAVs exceeded limit: %d slots used, but limit is %d due to maximum feature level 11.0"), CompileData.NumUAVs, GD3DMaximumNumUAVs));
+			FShaderCompilerError NewError;
+			NewError.StrippedErrorMessage = FString::Printf(TEXT("Number of UAVs exceeded limit: %d slots used, but limit is %d due to maximum feature level 11.0"), CompileData.NumUAVs, GD3DMaximumNumUAVs);
+			Output.Errors.Add(NewError);
 			Result = E_FAIL;
 			Output.bSucceeded = false;
 		}
@@ -1259,7 +1227,8 @@ bool CompileAndProcessD3DShaderFXC(
 	const FString& InEntryPointName,
 	const FShaderParameterParser& ShaderParameterParser,
 	const TCHAR* ShaderProfile,
-	bool bSecondPassAferUnusedInputRemoval,
+	const ED3DShaderModel ShaderModel,
+	const bool bSecondPassAferUnusedInputRemoval,
 	FShaderCompilerOutput& Output)
 {
 	// @TODO - implement different material path to allow us to remove backwards compatibility flag on sm5 shaders
@@ -1290,34 +1259,7 @@ bool CompileAndProcessD3DShaderFXC(
 			CompileFlags |= TranslateCompilerFlagD3D11((ECompilerFlags)Flag);
 		});
 
-	TArray<FString> FilteredErrors;
-	const bool bSuccess = CompileAndProcessD3DShaderFXCExt(CompileFlags, Input, InPreprocessedSource, InEntryPointName, ShaderParameterParser, ShaderProfile, false, FilteredErrors, Output);
-
-	// Process errors
-	for (int32 ErrorIndex = 0; ErrorIndex < FilteredErrors.Num(); ErrorIndex++)
-	{
-		const FString& CurrentError = FilteredErrors[ErrorIndex];
-		FShaderCompilerError NewError;
-
-		// Extract filename and line number from FXC output with format:
-		// "d:\Project\Binaries\BasePassPixelShader(30,7): error X3000: invalid target or usage string"
-		int32 FirstParenIndex = CurrentError.Find(TEXT("("));
-		int32 LastParenIndex = CurrentError.Find(TEXT("):"));
-		if (FirstParenIndex != INDEX_NONE &&
-			LastParenIndex != INDEX_NONE &&
-			LastParenIndex > FirstParenIndex)
-		{
-			// Extract and store error message with source filename
-			NewError.ErrorVirtualFilePath = CurrentError.Left(FirstParenIndex);
-			NewError.ErrorLineString = CurrentError.Mid(FirstParenIndex + 1, LastParenIndex - FirstParenIndex - FCString::Strlen(TEXT("(")));
-			NewError.StrippedErrorMessage = CurrentError.Right(CurrentError.Len() - LastParenIndex - FCString::Strlen(TEXT("):")));
-		}
-		else
-		{
-			NewError.StrippedErrorMessage = CurrentError;
-		}
-		Output.Errors.Add(NewError);
-	}
+	const bool bSuccess = CompileAndProcessD3DShaderFXCExt(CompileFlags, Input, InPreprocessedSource, InEntryPointName, ShaderParameterParser, ShaderProfile, ShaderModel, false, Output);
 
 	return bSuccess;
 }
@@ -1465,7 +1407,7 @@ void CompileD3DShader(const FShaderCompilerInput& Input, const FShaderPreprocess
 
 	const bool bSuccess = DoesShaderModelRequireDXC(ShaderModel)
 		? CompileAndProcessD3DShaderDXC(Input, PreprocessedSource, EntryPointName, ShaderParameterParser, ShaderProfile, ShaderModel, false, Output)
-		: CompileAndProcessD3DShaderFXC(Input, PreprocessedSource, EntryPointName, ShaderParameterParser, ShaderProfile, false, Output);
+		: CompileAndProcessD3DShaderFXC(Input, PreprocessedSource, EntryPointName, ShaderParameterParser, ShaderProfile, ShaderModel, false, Output);
 
 	if (!bSuccess && !Output.Errors.Num())
 	{
diff --git a/Engine/Source/Developer/Windows/ShaderFormatD3D/Private/D3DShaderCompiler.inl b/Engine/Source/Developer/Windows/ShaderFormatD3D/Private/D3DShaderCompiler.inl
--- a/Engine/Source/Developer/Windows/ShaderFormatD3D/Private/D3DShaderCompiler.inl
+++ b/Engine/Source/Developer/Windows/ShaderFormatD3D/Private/D3DShaderCompiler.inl
@@ -4,6 +4,7 @@
 
 #include "D3D12RHI.h"
 #include "RHIShaderBindingLayout.h"
+#include <functional>
 
 struct FD3DShaderCompileData
 {
@@ -581,3 +582,180 @@ inline void GenerateFinalOutput(
 		}
 	}
 }
+
+enum class ShaderCompilerType
+{
+	FXC,
+	DXC,
+};
+
+using ShaderCompileLambdaType = std::function<bool(
+	const FShaderCompilerInput& Input,
+	const FString& PreprocessedShaderSource,
+	const FString& EntryPointName,
+	const FShaderParameterParser& ShaderParameterParser,
+	const TCHAR* ShaderProfile,
+	const ED3DShaderModel ShaderModel,
+	const bool bProcessingSecondTime,
+	FShaderCompilerOutput& Output)>;
+
+template<ShaderCompilerType LambdaPlatform, typename TShaderCompileLambdaType,
+	typename ID3D1xShaderReflection, typename D3D1x_SHADER_DESC, typename D3D1x_SIGNATURE_PARAMETER_DESC>
+bool RemoveUnusedInterpolators(
+	const FShaderCompilerInput& Input,
+	const FString& PreprocessedShaderSource,
+	const FString& EntryPointName,
+	const FShaderParameterParser& ShaderParameterParser,
+	const TCHAR* ShaderProfile,
+	const ED3DShaderModel ShaderModel,
+	const bool bProcessingSecondTime,
+	FD3DShaderCompileData& CompileData,
+	TRefCountPtr<ID3D1xShaderReflection> Reflector,
+	TShaderCompileLambdaType ShaderCompileLambda,
+	FShaderCompilerOutput& Output,
+	bool& CompileResult)
+{
+	if (Input.Target.Frequency == SF_Pixel)
+	{
+		if (Reflector)
+		{
+			// Read the constant table description.
+			D3D1x_SHADER_DESC ShaderDesc;
+			Reflector->GetDesc(&ShaderDesc);
+
+			bool bFoundUnused = false;
+			for (uint32 Index = 0; Index < ShaderDesc.InputParameters; ++Index)
+			{
+				// VC++ horrible hack: Runtime ESP checks get confused and fail for some reason calling Reflector->GetInputParameterDesc() (because it comes from another DLL?)
+				// so "guard it" using the middle of an array; it's been confirmed NO corruption is really happening.
+				D3D1x_SIGNATURE_PARAMETER_DESC ParamDescs[3];
+				D3D1x_SIGNATURE_PARAMETER_DESC& ParamDesc = ParamDescs[1];
+				Reflector->GetInputParameterDesc(Index, &ParamDesc);
+				if (ParamDesc.SystemValueType == D3D_NAME_UNDEFINED)
+				{
+					if (ParamDesc.ReadWriteMask != 0)
+					{
+						FString SemanticName = ANSI_TO_TCHAR(ParamDesc.SemanticName);
+
+						CompileData.ShaderInputs.AddUnique(SemanticName);
+
+						// Add the number (for the case of TEXCOORD)
+						FString SemanticIndexName = FString::Printf(TEXT("%s%d"), *SemanticName, ParamDesc.SemanticIndex);
+						CompileData.ShaderInputs.AddUnique(SemanticIndexName);
+
+						// Add _centroid
+						CompileData.ShaderInputs.AddUnique(SemanticName + TEXT("_centroid"));
+						CompileData.ShaderInputs.AddUnique(SemanticIndexName + TEXT("_centroid"));
+					}
+					else
+					{
+						bFoundUnused = true;
+					}
+				}
+				else
+				{
+					// if (ParamDesc.ReadWriteMask != 0)
+					{
+						// Keep system values
+						CompileData.ShaderInputs.AddUnique(FString(ANSI_TO_TCHAR(ParamDesc.SemanticName)));
+					}
+				}
+			}
+
+			if constexpr (LambdaPlatform == ShaderCompilerType::DXC)
+			{
+				//Seems like there is bug in DXC compiler which prevents SV_Coverage semantic 
+				//from being detected as Input Parameter in the shader reflection data. 
+				//Other system interpolators are fine... 
+				//Can compile GizmoMaterial to repro the issue for example: -run=CookShadersCommandlet -targetPlatform=Windows -material=GizmoMaterial
+				//Always adding SV_Coverage as Used interpolator to mitigate this issue.
+				CompileData.ShaderInputs.AddUnique(TEXT("SV_Coverage"));
+			}
+
+			if (Input.Environment.CompilerFlags.Contains(CFLAG_ForceRemoveUnusedInterpolators) && Input.bCompilingForShaderPipeline && bFoundUnused && !bProcessingSecondTime)
+			{
+				// Rewrite the source removing the unused inputs so the bindings will match.
+				// We may need to do this more than once if unused inputs change after the removal. Ie. for complex shaders, what can happen is:
+				// pass1 detects that input A is not used, but input B and C are. Input A is removed, and we recompile (pass2). After the recompilation, we see that Input B is now also unused in pass2
+				// (it became simpler and the compiler could see through that).
+				// Since unused inputs are passed to the next stage, that will cause us to generate a vertex shader that does not output B, but our pixel shader will still be expecting B on input,
+				// as it was rewritten based on the pass1 results.
+
+				FShaderCompilerOutput OriginalOutput = Output;
+				const int kMaxReasonableAttempts = 64;
+				for (int32 Attempt = 0; Attempt < kMaxReasonableAttempts; ++Attempt)
+				{
+					TArray<FString> RemoveErrors;
+					FString ModifiedShaderSource = PreprocessedShaderSource;
+					FString ModifiedEntryPointName = Input.EntryPointName;
+					if (RemoveUnusedInputs(ModifiedShaderSource, CompileData.ShaderInputs, ModifiedEntryPointName, RemoveErrors))
+					{
+						Output = OriginalOutput;
+						
+						CompileResult = ShaderCompileLambda(Input, ModifiedShaderSource, ModifiedEntryPointName, ShaderParameterParser, ShaderProfile, ShaderModel, true, Output);
+						if (!CompileResult)
+						{
+							// if we failed to compile the shader, propagate the error up
+							return true;
+						}
+
+						// check if the ShaderInputs changed - if not, we're done here
+						if (Output.UsedAttributes.Num() == CompileData.ShaderInputs.Num())
+						{
+							Output.ModifiedShaderSource = MoveTemp(ModifiedShaderSource);
+							Output.ModifiedEntryPointName = MoveTemp(ModifiedEntryPointName);
+
+							return true;
+						}
+
+						// second pass cannot use more attributes than previously
+						if (Output.UsedAttributes.Num() > CompileData.ShaderInputs.Num())
+						{
+							UE_LOG(LogD3DShaderCompiler, Warning, TEXT("Second pass had more used attributes (%d) than first pass (%d)"), Output.UsedAttributes.Num(), CompileData.ShaderInputs.Num());
+							FShaderCompilerError NewError;
+							NewError.StrippedErrorMessage = FString::Printf(TEXT("Second pass had more used attributes (%d) than first pass (%d)"), Output.UsedAttributes.Num(), CompileData.ShaderInputs.Num());
+							Output = OriginalOutput;
+							Output.Errors.Add(NewError);
+							break;
+						}
+
+						// if we're about to run out of attempts, report
+						if (Attempt >= kMaxReasonableAttempts - 1)
+						{
+							UE_LOG(LogD3DShaderCompiler, Warning, TEXT("Unable to determine unused inputs after %d attempts (last number of used attributes: %d, previous step:%d)!"),
+								Attempt + 1,
+								Output.UsedAttributes.Num(),
+								CompileData.ShaderInputs.Num());
+							FShaderCompilerError NewError;
+							NewError.StrippedErrorMessage = FString::Printf(TEXT("Unable to determine unused inputs after %d attempts (last number of used attributes: %d, previous step:%d)!"),
+								Attempt + 1,
+								Output.UsedAttributes.Num(),
+								CompileData.ShaderInputs.Num());
+							Output = OriginalOutput;
+							Output.Errors.Add(NewError);
+							break;
+						}
+
+						CompileData.ShaderInputs = Output.UsedAttributes;
+						// go around to remove newly identified unused inputs
+					}
+					else
+					{
+						UE_LOG(LogD3DShaderCompiler, Warning, TEXT("Failed to remove unused inputs from shader: %s"), *Input.GenerateShaderName());
+						for (const FString& ErrorMessage : RemoveErrors)
+						{
+							// Add error to shader output but also make sure the error shows up on build farm by emitting a log entry
+							UE_LOG(LogD3DShaderCompiler, Warning, TEXT("%s"), *ErrorMessage);
+							FShaderCompilerError NewError;
+							NewError.StrippedErrorMessage = ErrorMessage;
+							Output.Errors.Add(NewError);
+						}
+						break;
+					}
+				}
+			}
+		}
+	}
+
+	return false;
+}
diff --git a/Engine/Source/Developer/Windows/ShaderFormatD3D/Private/D3DShaderCompilerDXC.cpp b/Engine/Source/Developer/Windows/ShaderFormatD3D/Private/D3DShaderCompilerDXC.cpp
--- a/Engine/Source/Developer/Windows/ShaderFormatD3D/Private/D3DShaderCompilerDXC.cpp
+++ b/Engine/Source/Developer/Windows/ShaderFormatD3D/Private/D3DShaderCompilerDXC.cpp
@@ -858,8 +858,8 @@ bool CompileAndProcessD3DShaderDXC(
 	const FString& EntryPointName,
 	const FShaderParameterParser& ShaderParameterParser,
 	const TCHAR* ShaderProfile,
-	ED3DShaderModel ShaderModel,
-	bool bProcessingSecondTime,
+	const ED3DShaderModel ShaderModel,
+	const bool bProcessingSecondTime,
 	FShaderCompilerOutput& Output)
 {
 	TRACE_CPUPROFILER_EVENT_SCOPE(CompileAndProcessD3DShaderDXC);
@@ -1022,6 +1022,55 @@ bool CompileAndProcessD3DShaderDXC(
 			}
 		}
 
+		// Remove unused interpolators from pixel shader 
+		// (propagated to corresponding VS from pipeline by later setting Output.bSupportsQueryingUsedAttributes and Output.UsedAttributes)
+		{
+			TRefCountPtr<ID3D12ShaderReflection> Reflector;
+			Utils->CreateReflection(&ReflBuffer, IID_PPV_ARGS(Reflector.GetInitReference()));
+
+			ShaderCompileLambdaType ShaderCompileLambda = [](
+				const FShaderCompilerInput& Input,
+				const FString& PreprocessedShaderSource,
+				const FString& EntryPointName,
+				const FShaderParameterParser& ShaderParameterParser,
+				const TCHAR* ShaderProfile,
+				const ED3DShaderModel ShaderModel,
+				const bool bProcessingSecondTime,
+				FShaderCompilerOutput& Output)
+				{
+					return CompileAndProcessD3DShaderDXC(
+						Input,
+						PreprocessedShaderSource,
+						EntryPointName,
+						ShaderParameterParser,
+						ShaderProfile,
+						ShaderModel,
+						bProcessingSecondTime,
+						Output);
+				};
+
+			bool CompileResult = false;
+			const bool RemovedUnusedInterpolatorsApplied =
+				RemoveUnusedInterpolators<ShaderCompilerType::DXC, ShaderCompileLambdaType,
+				ID3D12ShaderReflection, D3D12_SHADER_DESC, D3D12_SIGNATURE_PARAMETER_DESC>(
+					Input,
+					PreprocessedShaderSource,
+					EntryPointName,
+					ShaderParameterParser,
+					ShaderProfile,
+					ShaderModel,
+					bProcessingSecondTime,
+					CompileData,
+					Reflector,
+					ShaderCompileLambda,
+					Output,
+					CompileResult);
+			if (RemovedUnusedInterpolatorsApplied)
+			{
+				return CompileResult;
+			}
+		}
+
 		if (bIsWorkGraphShader || bIsRayTracingShader)
 		{
 			TRefCountPtr<ID3D12LibraryReflection> LibraryReflection;
