[LiveCoding]
* Fixed so global mutex is released earlier during preexit so people can start compiling before editor process is entirely gone

diff --git a/Engine/Source/Developer/Windows/LiveCoding/LiveCoding.Build.cs b/Engine/Source/Developer/Windows/LiveCoding/LiveCoding.Build.cs
--- a/Engine/Source/Developer/Windows/LiveCoding/LiveCoding.Build.cs
+++ b/Engine/Source/Developer/Windows/LiveCoding/LiveCoding.Build.cs
@@ -11,6 +11,7 @@ public class LiveCoding : ModuleRules
 		PrivateDependencyModuleNames.Add("CoreUObject");
 		PrivateDependencyModuleNames.Add("Settings");
 		PrivateDependencyModuleNames.Add("BuildSettings");
+		PrivateDependencyModuleNames.Add("OodleDataCompression");
 
 		if (Target.bCompileAgainstEngine)
 		{
diff --git a/Engine/Source/Developer/Windows/LiveCoding/Private/External/LC_MemoryMappedFile.cpp b/Engine/Source/Developer/Windows/LiveCoding/Private/External/LC_MemoryMappedFile.cpp
--- a/Engine/Source/Developer/Windows/LiveCoding/Private/External/LC_MemoryMappedFile.cpp
+++ b/Engine/Source/Developer/Windows/LiveCoding/Private/External/LC_MemoryMappedFile.cpp
@@ -8,6 +8,8 @@
 #include "LC_MemoryMappedFile.h"
 // BEGIN EPIC MOD
 #include "LC_Logging.h"
+#include "LC_Vfs.h"
+#include <oodle2.h>
 // END EPIC MOD
 
 namespace detail
@@ -86,8 +88,11 @@ namespace Filesystem
 }
 
 
-Filesystem::MemoryMappedFile* Filesystem::OpenMemoryMappedFile(const wchar_t* path, OpenMode::Enum openMode)
+Filesystem::MemoryMappedFile* Filesystem::OpenMemoryMappedFile(const wchar_t* path, OpenMode::Enum openMode, bool checkIfCompressed)
 {
+	Path temp;
+	path = Devirtualize(path, temp);
+
 	HANDLE file = ::CreateFileW(path, detail::GetDesiredAccess(openMode), detail::GetShareMode(openMode), NULL, OPEN_EXISTING, FILE_ATTRIBUTE_NORMAL, NULL);
 	if (file == INVALID_HANDLE_VALUE)
 	{
@@ -113,6 +118,49 @@ Filesystem::MemoryMappedFile* Filesystem::OpenMemoryMappedFile(const wchar_t* pa
 		return nullptr;
 	}
 
+	if (checkIfCompressed)
+	{
+		auto data = (const BYTE*)baseAddress;
+		static constexpr char UbaMagic[] = { 'U', 'B', 'A', 0x01 };
+		if (memcmp(data, UbaMagic, sizeof(UbaMagic)) == 0)
+		{
+			data += 24;
+			uint64_t fileSize = *(uint64_t*)data;
+			data += sizeof(uint64_t);
+
+			::ULARGE_INTEGER integer = {};
+			integer.QuadPart = fileSize;
+			HANDLE memoryMapping = ::CreateFileMappingW(INVALID_HANDLE_VALUE, NULL, PAGE_READWRITE, integer.HighPart, integer.LowPart, NULL);
+			if (memoryMapping == NULL)
+			{
+				LC_ERROR_USER("Cannot create named shared memory. Error: 0x%X", ::GetLastError());
+				return nullptr;
+			}
+			void* memoryView = ::MapViewOfFile(memoryMapping, FILE_MAP_ALL_ACCESS, 0, 0, fileSize);
+
+			BYTE* destMem = (BYTE*)memoryView;
+			uint64_t left = fileSize;
+			while (left)
+			{
+				uint32_t compressedBlockSize = *(uint32_t*)data;
+				data += 4;
+				uint32_t decompressedBlockSize = *(uint32_t*)data;
+				data += 4;
+				OO_SINTa decompLen = OodleLZ_Decompress(data, (OO_SINTa)compressedBlockSize, destMem, (OO_SINTa)decompressedBlockSize, OodleLZ_FuzzSafe_Yes, OodleLZ_CheckCRC_No, OodleLZ_Verbosity_None, NULL, 0, NULL, NULL, NULL, 0);
+				data += compressedBlockSize;
+				destMem += decompressedBlockSize;
+				left -= decompressedBlockSize;
+			}
+
+			::UnmapViewOfFile(baseAddress);
+			::CloseHandle(fileMapping);
+
+			baseAddress = memoryView;
+			fileMapping = memoryMapping;
+		}
+	}
+
+
 	return new MemoryMappedFile { file, fileMapping, baseAddress };
 }
 
diff --git a/Engine/Source/Developer/Windows/LiveCoding/Private/External/LC_MemoryMappedFile.h b/Engine/Source/Developer/Windows/LiveCoding/Private/External/LC_MemoryMappedFile.h
--- a/Engine/Source/Developer/Windows/LiveCoding/Private/External/LC_MemoryMappedFile.h
+++ b/Engine/Source/Developer/Windows/LiveCoding/Private/External/LC_MemoryMappedFile.h
@@ -12,7 +12,7 @@ namespace Filesystem
 	// opaque type
 	struct MemoryMappedFile;
 
-	MemoryMappedFile* OpenMemoryMappedFile(const wchar_t* path, OpenMode::Enum openMode);
+	MemoryMappedFile* OpenMemoryMappedFile(const wchar_t* path, OpenMode::Enum openMode, bool checkIfCompressed = false);
 	void CloseMemoryMappedFile(MemoryMappedFile*& file);
 
 	void* GetMemoryMappedFileData(MemoryMappedFile* file);
diff --git a/Engine/Source/Developer/Windows/LiveCoding/Private/External/LC_Vfs.cpp b/Engine/Source/Developer/Windows/LiveCoding/Private/External/LC_Vfs.cpp
new file mode 100644
--- /dev/null
+++ b/Engine/Source/Developer/Windows/LiveCoding/Private/External/LC_Vfs.cpp
@@ -0,0 +1,49 @@
+// Copyright 2011-2020 Molecular Matters GmbH, all rights reserved.
+
+// BEGIN EPIC MOD
+
+#if LC_VERSION == 1
+
+#include "LC_Vfs.h"
+
+namespace Filesystem
+{
+
+types::vector<VfsEntry> g_vfsEntries;
+
+void SetVfs(const types::vector<VfsEntry>& vfsEntries)
+{
+	g_vfsEntries = vfsEntries;
+}
+
+const types::vector<VfsEntry>& GetVfsEntries()
+{
+	return g_vfsEntries;
+}
+
+std::wstring Devirtualize(const std::wstring& path)
+{
+	for (auto& entry : g_vfsEntries)
+		if (entry.virtualPath.compare(0, entry.virtualPath.size(), path, 0, entry.virtualPath.size()) == 0)
+			return entry.localPath + (path.data() + entry.virtualPath.size());
+	return path;
+}
+
+const wchar_t* Devirtualize(const wchar_t* path, Path& temp)
+{
+	for (auto& entry : g_vfsEntries)
+	{
+		if (entry.virtualPath.compare(0, entry.virtualPath.size(), path, 0, entry.virtualPath.size()) != 0)
+			continue;
+		temp = entry.localPath.c_str();
+		temp += (path + entry.virtualPath.size());
+		return temp.GetString();
+	}
+	return path;
+}
+
+}
+
+#endif
+
+// END EPIC MOD
diff --git a/Engine/Source/Developer/Windows/LiveCoding/Private/External/LC_Vfs.h b/Engine/Source/Developer/Windows/LiveCoding/Private/External/LC_Vfs.h
new file mode 100644
--- /dev/null
+++ b/Engine/Source/Developer/Windows/LiveCoding/Private/External/LC_Vfs.h
@@ -0,0 +1,31 @@
+// Copyright 2011-2020 Molecular Matters GmbH, all rights reserved.
+
+#pragma once
+
+// BEGIN EPIC MOD
+
+#if LC_VERSION == 1
+
+#include "LC_FilesystemTypes.h"
+#include "LC_Types.h"
+
+namespace Filesystem
+{
+	struct VfsEntry
+	{
+		std::wstring virtualPath;
+		std::wstring localPath;
+	};
+	void SetVfs(const types::vector<VfsEntry>& vfsEntries);
+
+	const types::vector<VfsEntry>& GetVfsEntries();
+
+	std::wstring Devirtualize(const std::wstring& path);
+
+	const wchar_t* Devirtualize(const wchar_t* path, Path& temp);
+}
+
+
+#endif
+
+// END EPIC MOD
diff --git a/Engine/Source/Developer/Windows/LiveCoding/Private/LiveCodingModule.cpp b/Engine/Source/Developer/Windows/LiveCoding/Private/LiveCodingModule.cpp
--- a/Engine/Source/Developer/Windows/LiveCoding/Private/LiveCodingModule.cpp
+++ b/Engine/Source/Developer/Windows/LiveCoding/Private/LiveCodingModule.cpp
@@ -510,23 +510,7 @@
 
 void FLiveCodingModule::ShutdownModule()
 {
-	LppShutdown();
-
-	FCoreDelegates::OnEndFrame.Remove(EndFrameDelegateHandle);
-
-	IConsoleManager& ConsoleManager = IConsoleManager::Get();
-	ConsoleManager.UnregisterConsoleObject(SourceProjectVariable);
-	ConsoleManager.UnregisterConsoleObject(ConsolePathVariable);
-	ConsoleManager.UnregisterConsoleObject(CompileCommand);
-	ConsoleManager.UnregisterConsoleObject(CompileSyncCommand);
-	ConsoleManager.UnregisterConsoleObject(EnableCommand);
-
-	// Unregister from the dll notifications
-	if (CallbackCookie)
-	{
-		FNtDllFunction UnregisterFunc("LdrUnregisterDllNotification");
-		UnregisterFunc(CallbackCookie);
-	}
+	StopLiveCoding();
 }
 
 void FLiveCodingModule::EnableByDefault(bool bEnable)
@@ -1144,14 +1128,58 @@
 			return false;
 		}
 
+		// Create a mutex that allows UBT to detect that we shouldn't hot-reload into this executable. The handle to it will be released automatically when the process exits.
+		FString ExecutablePath = FPaths::ConvertRelativePathToFull(FPlatformProcess::ExecutablePath());
+
+		FString MutexName = TEXT("Global\\LiveCoding_");
+		for (int Idx = 0; Idx < ExecutablePath.Len(); Idx++)
+		{
+			TCHAR Character = ExecutablePath[Idx];
+			if (Character == '/' || Character == '\\' || Character == ':')
+			{
+				MutexName += '+';
+			}
+			else
+			{
+				MutexName += Character;
+			}
+		}
+
+		GlobalMutex = CreateMutex(NULL, FALSE, *MutexName);
+		if (!GlobalMutex)
+		{
+			State = EState::NotRunning;
+			return false;
+		}
+
 		UE_LOG(LogLiveCoding, Display, TEXT("Starting LiveCoding"));
 
 		// Enable external build system
 		LppUseExternalBuildSystem();
 
+		const TCHAR* VfsPaths = BuildSettings::GetVfsPathsWide();
+		if (*VfsPaths)
+		{
+			FString Temp;
+			if (SourceProject.Len() > 0)
+			{
+				const TCHAR* FirstSemi = FCString::Strchr(VfsPaths, ';');
+				FString VfsRoot = FPaths::GetPath(FString(MakeStringView(VfsPaths, FirstSemi - VfsPaths)));
+
+				Temp = VfsRoot / FPaths::GetBaseFilename(SourceProject) + ';' + FPaths::GetPath(SourceProject) + ';';
+#if PLATFORM_WINDOWS
+				Temp.ReplaceInline(TEXT("/"), TEXT("\\"));
+#endif
+				Temp += VfsPaths;
+				VfsPaths = *Temp;
+			}
+			LppApplySettingString("Vfs", VfsPaths);
+		}
+
 		// Enable the server
 		FString ProjectPath = FPaths::ConvertRelativePathToFull(FPaths::ProjectDir()).ToLower();
 		FString ProcessGroup = FString::Printf(TEXT("UE_%s_0x%08x"), FApp::GetProjectName(), GetTypeHash(ProjectPath));
+
 		LppRegisterProcessGroup(TCHAR_TO_ANSI(*ProcessGroup));
 
 		SetBuildArguments();
@@ -1168,25 +1196,6 @@
 		}
 #endif
 
-		// Create a mutex that allows UBT to detect that we shouldn't hot-reload into this executable. The handle to it will be released automatically when the process exits.
-		FString ExecutablePath = FPaths::ConvertRelativePathToFull(FPlatformProcess::ExecutablePath());
-
-		FString MutexName = TEXT("Global\\LiveCoding_");
-		for (int Idx = 0; Idx < ExecutablePath.Len(); Idx++)
-		{
-			TCHAR Character = ExecutablePath[Idx];
-			if (Character == '/' || Character == '\\' || Character == ':')
-			{
-				MutexName += '+';
-			}
-			else
-			{
-				MutexName += Character;
-			}
-		}
-
-		ensure(CreateMutex(NULL, FALSE, *MutexName));
-
 		// Configure all the current modules. For non-commandlets, schedule it to be done in the first Tick() so we can batch everything together.
 		if (IsRunningCommandlet())
 		{
@@ -1207,10 +1216,44 @@
 
 		// Mark it as started
 		State = EState::RunningAndEnabled;
+
+		FCoreDelegates::OnPreExit.AddLambda( [this]()
+		{
+			StopLiveCoding();
+		});
 	}
 	return true;
 }
 
+void FLiveCodingModule::StopLiveCoding()
+{
+	if (!GlobalMutex)
+	{
+		return;
+	}
+
+	LppShutdown();
+
+	ReleaseMutex(GlobalMutex);
+	GlobalMutex = nullptr;
+
+	FCoreDelegates::OnEndFrame.Remove(EndFrameDelegateHandle);
+
+	IConsoleManager& ConsoleManager = IConsoleManager::Get();
+	ConsoleManager.UnregisterConsoleObject(SourceProjectVariable);
+	ConsoleManager.UnregisterConsoleObject(ConsolePathVariable);
+	ConsoleManager.UnregisterConsoleObject(CompileCommand);
+	ConsoleManager.UnregisterConsoleObject(CompileSyncCommand);
+	ConsoleManager.UnregisterConsoleObject(EnableCommand);
+
+	// Unregister from the dll notifications
+	if (CallbackCookie)
+	{
+		FNtDllFunction UnregisterFunc("LdrUnregisterDllNotification");
+		UnregisterFunc(CallbackCookie);
+	}
+}
+
 void FLiveCodingModule::SetBuildArguments()
 {
 	FScopeLock lock(&SetBuildArgumentsCs);
diff --git a/Engine/Source/Developer/Windows/LiveCoding/Private/LiveCodingModule.h b/Engine/Source/Developer/Windows/LiveCoding/Private/LiveCodingModule.h
--- a/Engine/Source/Developer/Windows/LiveCoding/Private/LiveCodingModule.h
+++ b/Engine/Source/Developer/Windows/LiveCoding/Private/LiveCodingModule.h
@@ -63,6 +63,7 @@ private:
 	void StartLiveCodingAsync(ELiveCodingStartupMode StartupMode);
 	bool SetupConsolePath();
 	bool StartLiveCoding(ELiveCodingStartupMode StartupMode);
+	void StopLiveCoding();
 	void OnModulesChanged(FName ModuleName, EModuleChangeReason Reason);
 	void UpdateModules(bool bAllowStarting);
 	bool ShouldPreloadModule(const TSet<FName>& PreloadedFileNames, const FString& FullFilePath) const;
@@ -143,5 +144,7 @@ private:
 #else
 	TUniquePtr<FNullReload> Reload;
 #endif
+
+	void* GlobalMutex = NULL;
 };
 
diff --git a/Engine/Source/Developer/Windows/LiveCodingServer/Private/External/LC_AppSettings.cpp b/Engine/Source/Developer/Windows/LiveCodingServer/Private/External/LC_AppSettings.cpp
--- a/Engine/Source/Developer/Windows/LiveCodingServer/Private/External/LC_AppSettings.cpp
+++ b/Engine/Source/Developer/Windows/LiveCodingServer/Private/External/LC_AppSettings.cpp
@@ -789,6 +789,35 @@ void appSettings::ApplySettingString(const char* const settingName, const wchar_
 			g_fastBuildDllName
 		};
 
+		// BEGIN EPIC MOD
+		if (strcmp(settingName, "Vfs") == 0)
+		{
+			types::vector<Filesystem::VfsEntry> vfsEntries;
+
+			size_t start = 0, end;
+			std::wstring input(value);
+
+			Filesystem::VfsEntry entry;
+			while ((end = input.find(';', start)) != std::wstring::npos)
+			{
+				if (entry.virtualPath.empty())
+					entry.virtualPath = input.substr(start, end - start);
+				else
+				{
+					entry.localPath = input.substr(start, end - start);
+					vfsEntries.push_back(std::move(entry));
+				}
+				start = end + 1;
+			}
+			entry.localPath = input.substr(start); // Last token
+			if (!entry.localPath.empty())
+				vfsEntries.push_back(std::move(entry));
+
+			Filesystem::SetVfs(vfsEntries);
+			return;
+		}
+		// END EPIC MOD
+
 		const SettingString* setting = ApplySetting(settings, COUNT, settingName, value);
 		if (setting)
 		{
diff --git a/Engine/Source/Developer/Windows/LiveCodingServer/Private/External/LC_ChangeNotification.cpp b/Engine/Source/Developer/Windows/LiveCodingServer/Private/External/LC_ChangeNotification.cpp
--- a/Engine/Source/Developer/Windows/LiveCodingServer/Private/External/LC_ChangeNotification.cpp
+++ b/Engine/Source/Developer/Windows/LiveCodingServer/Private/External/LC_ChangeNotification.cpp
@@ -8,6 +8,7 @@
 #include "LC_ChangeNotification.h"
 // BEGIN EPIC MOD
 #include "LC_Logging.h"
+#include "LC_Vfs.h"
 #include "Windows/WindowsHWrapper.h"
 // END EPIC MOD
 
@@ -26,6 +27,8 @@ ChangeNotification::~ChangeNotification(void)
 void ChangeNotification::Create(const wchar_t* path)
 {
 	// BEGIN EPIC MOD
+	Filesystem::Path temp;
+	path = Filesystem::Devirtualize(path, temp);
 	m_handle = ::FindFirstChangeNotificationW(path, Windows::TRUE, FILE_NOTIFY_CHANGE_LAST_WRITE);
 	// END EPIC MOD
 	if (m_handle == INVALID_HANDLE_VALUE)
diff --git a/Engine/Source/Developer/Windows/LiveCodingServer/Private/External/LC_Coff.cpp b/Engine/Source/Developer/Windows/LiveCodingServer/Private/External/LC_Coff.cpp
--- a/Engine/Source/Developer/Windows/LiveCodingServer/Private/External/LC_Coff.cpp
+++ b/Engine/Source/Developer/Windows/LiveCodingServer/Private/External/LC_Coff.cpp
@@ -19,6 +19,7 @@
 // BEGIN EPIC MOD
 #include "LC_Allocators.h"
 #include "LC_Assert.h"
+#include "LC_Vfs.h"
 #include <algorithm>
 // END EPIC MOD
 
@@ -38,9 +39,14 @@ namespace coff
 {
 	ObjFile* OpenObj(const wchar_t* filename)
 	{
+		// BEGIN EPIC MOD
+		Filesystem::Path temp;
+		filename = Filesystem::Devirtualize(filename, temp);
+		// END EPIC MOD
+
 		ObjFile* objFile = LC_NEW(&g_objFileAllocator, ObjFile);
 		objFile->filename = string::ToUtf8String(filename);
-		objFile->memoryFile = Filesystem::OpenMemoryMappedFile(filename, Filesystem::OpenMode::READ);
+		objFile->memoryFile = Filesystem::OpenMemoryMappedFile(filename, Filesystem::OpenMode::READ, true);
 
 		return objFile;
 	}
diff --git a/Engine/Source/Developer/Windows/LiveCodingServer/Private/External/LC_Filesystem.cpp b/Engine/Source/Developer/Windows/LiveCodingServer/Private/External/LC_Filesystem.cpp
--- a/Engine/Source/Developer/Windows/LiveCodingServer/Private/External/LC_Filesystem.cpp
+++ b/Engine/Source/Developer/Windows/LiveCodingServer/Private/External/LC_Filesystem.cpp
@@ -55,12 +55,19 @@ namespace detail
 			if (optional.second)
 			{
 				// value was inserted, update it with the correct data
+				// 
+				// BEGIN EPIC MOD
+				Filesystem::Path temp;
+				path = Filesystem::Devirtualize(path, temp);
+				// END EPIC MOD
+
 				HANDLE file = ::CreateFileW(path, FILE_READ_ATTRIBUTES | STANDARD_RIGHTS_READ, FILE_SHARE_READ, 0, OPEN_EXISTING, FILE_FLAG_BACKUP_SEMANTICS, 0);
 				if (file != INVALID_HANDLE_VALUE)
 				{
 					wchar_t buffer[Filesystem::Path::CAPACITY] = {};
 					::GetFinalPathNameByHandleW(file, buffer, Filesystem::Path::CAPACITY, 0u);
 					::CloseHandle(file);
+					// TODO: Revirtualize?
 
 					// the path returned by GetFinalPathNameByHandle starts with "\\?\", cut that off
 					data.assign(buffer + 4u);
@@ -165,6 +172,11 @@ Filesystem::DriveType::Enum Filesystem::GetDriveType(const wchar_t* path)
 
 Filesystem::PathAttributes Filesystem::GetAttributes(const wchar_t* path)
 {
+	// BEGIN EPIC MOD
+	Path temp;
+	path = Devirtualize(path, temp);
+	// END EPIC MOD
+
 	::WIN32_FILE_ATTRIBUTE_DATA attributes = {};
 	attributes.dwFileAttributes = INVALID_FILE_ATTRIBUTES;
 
@@ -228,6 +240,10 @@ bool Filesystem::IsRelativePath(const wchar_t* path)
 void Filesystem::Copy(const wchar_t* srcPath, const wchar_t* destPath)
 {
 	// BEGIN EPIC MOD
+	Path srcTemp;
+	srcPath = Devirtualize(srcPath, srcTemp);
+	Path destTemp;
+	destPath = Devirtualize(destPath, destTemp);
 	const BOOL success = ::CopyFileW(srcPath, destPath, Windows::FALSE);
 	if (success == Windows::FALSE)
 	// END EPIC MOD
@@ -239,6 +255,13 @@ void Filesystem::Copy(const wchar_t* srcPath, const wchar_t* destPath)
 
 void Filesystem::Move(const wchar_t* srcPath, const wchar_t* destPath)
 {
+	// BEGIN EPIC MOD
+	Path srcTemp;
+	srcPath = Devirtualize(srcPath, srcTemp);
+	Path destTemp;
+	destPath = Devirtualize(destPath, destTemp);
+	// END EPIC MOD
+
 	const BOOL success = ::MoveFileExW(srcPath, destPath, MOVEFILE_COPY_ALLOWED | MOVEFILE_REPLACE_EXISTING | MOVEFILE_WRITE_THROUGH);
 	// BEGIN EPIC MOD
 	if (success == Windows::FALSE)
@@ -251,6 +274,11 @@ void Filesystem::Move(const wchar_t* srcPath, const wchar_t* destPath)
 
 void Filesystem::Delete(const wchar_t* path)
 {
+	// BEGIN EPIC MOD
+	Path temp;
+	path = Devirtualize(path, temp);
+	// END EPIC MOD
+
 	const BOOL success = ::DeleteFileW(path);
 	// BEGIN EPIC MOD
 	if (success == Windows::FALSE)
@@ -263,6 +291,11 @@ void Filesystem::Delete(const wchar_t* path)
 
 bool Filesystem::DeleteIfExists(const wchar_t* path)
 {
+	// BEGIN EPIC MOD
+	Path temp;
+	path = Devirtualize(path, temp);
+	// END EPIC MOD
+
 	const BOOL success = ::DeleteFileW(path);
 	// BEGIN EPIC MOD
 	return (success != Windows::FALSE);
@@ -405,6 +438,11 @@ Filesystem::Path Filesystem::NormalizePathWithoutResolvingLinks(const wchar_t* p
 
 bool Filesystem::CreateFileWithData(const wchar_t* path, const void* data, size_t size)
 {
+	// BEGIN EPIC MOD
+	Path temp;
+	path = Devirtualize(path, temp);
+	// END EPIC MOD
+
 	HANDLE file = ::CreateFileW(path, GENERIC_READ | GENERIC_WRITE, 0, NULL, CREATE_ALWAYS, FILE_ATTRIBUTE_NORMAL, NULL);
 	if (file == INVALID_HANDLE_VALUE)
 	{
diff --git a/Engine/Source/Developer/Windows/LiveCodingServer/Private/External/LC_Filesystem.h b/Engine/Source/Developer/Windows/LiveCodingServer/Private/External/LC_Filesystem.h
--- a/Engine/Source/Developer/Windows/LiveCodingServer/Private/External/LC_Filesystem.h
+++ b/Engine/Source/Developer/Windows/LiveCodingServer/Private/External/LC_Filesystem.h
@@ -7,7 +7,7 @@
 // BEGIN EPIC MOD
 #include "CoreTypes.h"
 // END EPIC MOD
-#include "LC_FilesystemTypes.h"
+#include "LC_Vfs.h"
 // BEGIN EPIC MOD
 #include "LC_Types.h"
 // END EPIC MOD
diff --git a/Engine/Source/Developer/Windows/LiveCodingServer/Private/External/LC_LiveModule.cpp b/Engine/Source/Developer/Windows/LiveCodingServer/Private/External/LC_LiveModule.cpp
--- a/Engine/Source/Developer/Windows/LiveCodingServer/Private/External/LC_LiveModule.cpp
+++ b/Engine/Source/Developer/Windows/LiveCodingServer/Private/External/LC_LiveModule.cpp
@@ -539,9 +539,8 @@ namespace
 	};
 
 
-	// EPIC MOD - Added VfsToLocal
-	template <typename T, typename VfsToLocalFunc>
-	static types::vector<symbols::ObjPath> UpdateCoffCache(const VfsToLocalFunc& vfsToLocal, const T& compilands, CoffCache<coff::CoffDB>* coffCache, CacheUpdate::Enum updateType, const types::vector<symbols::ModifiedObjFile>& modifiedOrNewObjFiles)
+	template <typename T>
+	static types::vector<symbols::ObjPath> UpdateCoffCache(const T& compilands, CoffCache<coff::CoffDB>* coffCache, CacheUpdate::Enum updateType, const types::vector<symbols::ModifiedObjFile>& modifiedOrNewObjFiles)
 	{
 		LC_LOG_INDENT_DEV;
 
@@ -568,13 +567,11 @@ namespace
 			{
 				updatedCoffs.push_back(objPath);
 
-				// EPIC MOD - Added VfsToLocal
-				auto task = scheduler::CreateTask(taskRoot, [objPath, wideObjPath, coffCache, compilandUniqueId, &vfsToLocal]()
+				auto task = scheduler::CreateTask(taskRoot, [objPath, wideObjPath, coffCache, compilandUniqueId]()
 				{
 					LC_LOG_DEV("Updating COFF cache for file %s", objPath.c_str());
 
-					// EPIC MOD - Added VfsToLocal
-					coff::ObjFile* objFile = coff::OpenObj(vfsToLocal(wideObjPath).c_str());
+					coff::ObjFile* objFile = coff::OpenObj(wideObjPath.c_str());
 					if (objFile && objFile->memoryFile)
 					{
 						coff::CoffDB* database = coff::GatherDatabase(objFile, compilandUniqueId);
@@ -1355,9 +1352,8 @@ void LiveModule::UpdateDirectoryCache(DirectoryCache* cache)
 	}
 }
 
-
 // BEGIN EPIC MOD
-LiveModule::ErrorType::Enum LiveModule::Update(FileAttributeCache* fileCache, DirectoryCache* directoryCache, UpdateType::Enum updateType, const types::vector<symbols::ModifiedObjFile>& modifiedOrNewObjFiles, const types::vector<std::wstring>& additionalLibraries, const types::vector<VfsEntry>& vfsEntries)
+LiveModule::ErrorType::Enum LiveModule::Update(FileAttributeCache* fileCache, DirectoryCache* directoryCache, UpdateType::Enum updateType, const types::vector<symbols::ModifiedObjFile>& modifiedOrNewObjFiles, const types::vector<std::wstring>& additionalLibraries)
 // END EPIC MOD
 {
 	telemetry::Scope updateScope("Update live module");
@@ -1366,14 +1362,6 @@ LiveModule::ErrorType::Enum LiveModule::Update(FileAttributeCache* fileCache, Di
 
 	// BEGIN EPIC MOD
 	GLiveCodingServer->GetStatusChangeDelegate().ExecuteIfBound(L"Checking modified files...");
-
-	auto VfsToLocal = [&](const std::wstring& path)
-		{
-			for (auto& entry : vfsEntries)
-				if (entry.virtualPath.compare(0, entry.virtualPath.size(), path, 0, entry.virtualPath.size()) == 0)
-					return entry.localPath + (path.data() + entry.virtualPath.size());
-			return path;
-		};
 	// END EPIC MOD
 
 	bool usesClang = false;
@@ -2039,14 +2027,13 @@ LiveModule::ErrorType::Enum LiveModule::Update(FileAttributeCache* fileCache, Di
 		LC_LOG_DEV("Finding set of .obj files");
 		LC_LOG_INDENT_DEV;
 
-		// EPIC MOD - Changed to lambda to be able to access VfsToLocal
 		auto UpdateExternalSymbolsAndNeededFiles = [&]
 			(
 				const symbols::ObjPath& objPath, const symbols::Compiland* compiland, uint32_t compilandUniqueId, const types::StringMap<ImmutableString>& pchSymbolToCompilandName,
 				types::StringMap<CompilandInfo>& externalSymbols, types::StringSet& neededCompilands
 			)
 			{
-				coff::ObjFile* coffFile = coff::OpenObj(VfsToLocal(string::ToWideString(objPath)).c_str());
+				coff::ObjFile* coffFile = coff::OpenObj(string::ToWideString(objPath).c_str());
 				if (coffFile && coffFile->memoryFile)
 				{
 					coff::ExternalSymbolDB* externalSymbolDb = coff::GatherExternalSymbolDatabase(coffFile, compilandUniqueId);
@@ -2223,7 +2210,7 @@ LiveModule::ErrorType::Enum LiveModule::Update(FileAttributeCache* fileCache, Di
 						// there is no entry yet for this COFF in the cache.
 						// this means that this .obj was not recompiled (otherwise it would have an entry already),
 						// but has been pulled in for the first time due to unresolved symbols.
-						auto task = scheduler::CreateTask(taskRoot, [this, objPath, &modifiedOrNewObjFiles, &VfsToLocal]()
+						auto task = scheduler::CreateTask(taskRoot, [this, objPath, &modifiedOrNewObjFiles]()
 						{
 							const symbols::Compiland* compiland = symbols::FindCompiland(m_compilandDB, objPath);
 							const std::wstring& wideObjPath = string::ToWideString(objPath);
@@ -2231,8 +2218,7 @@ LiveModule::ErrorType::Enum LiveModule::Update(FileAttributeCache* fileCache, Di
 
 							LC_LOG_DEV("Need %s for the first time, updating COFF cache", objPath.c_str());
 
-							// EPIC MOD - Added VfsToLocal
-							coff::ObjFile* objFile = coff::OpenObj(VfsToLocal(wideObjPath).c_str());
+							coff::ObjFile* objFile = coff::OpenObj(wideObjPath.c_str());
 							if (objFile && objFile->memoryFile)
 							{
 								// note that even though we might be dealing with a single-part .obj of an amalgamated .obj
@@ -2295,8 +2281,7 @@ LiveModule::ErrorType::Enum LiveModule::Update(FileAttributeCache* fileCache, Di
 
 
 	// update the COFF cache for all compiled files
-	// EPIC MOD - Added VfsToLocal
-	UpdateCoffCache(VfsToLocal, m_compiledCompilands, m_coffCache, CacheUpdate::ALL, modifiedOrNewObjFiles);
+	UpdateCoffCache(m_compiledCompilands, m_coffCache, CacheUpdate::ALL, modifiedOrNewObjFiles);
 
 	// BEGIN EPIC MOD
 	GLiveCodingServer->GetStatusChangeDelegate().ExecuteIfBound(L"Stripping COFFs...");
@@ -2329,8 +2314,7 @@ LiveModule::ErrorType::Enum LiveModule::Update(FileAttributeCache* fileCache, Di
 
 			const uint32_t compilandUniqueId = symbols::GetCompilandId(compiland, wideObjPath.c_str(), modifiedOrNewObjFiles);
 
-			// EPIC MOD - Added VfsToLocal
-			coff::ObjFile* objFile = coff::OpenObj(VfsToLocal(wideObjPath).c_str());
+			coff::ObjFile* objFile = coff::OpenObj(wideObjPath.c_str());
 			if (objFile && objFile->memoryFile)
 			{
 				coff::RawCoff* rawCoff = coff::ReadRaw(objFile, compilandUniqueId);
@@ -2362,8 +2346,7 @@ LiveModule::ErrorType::Enum LiveModule::Update(FileAttributeCache* fileCache, Di
 			// we need it after linking has finished  
 			{
 				const std::wstring bakPath = wideObjPath + L".bak";
-				// EPIC MOD - Added VfsToLocal
-				Filesystem::Move(VfsToLocal(wideObjPath).c_str(), VfsToLocal(bakPath).c_str());
+				Filesystem::Move(wideObjPath.c_str(), bakPath.c_str());
 			}
 
 			// remove linker directives which we don't want or need.
@@ -2935,8 +2918,7 @@ LiveModule::ErrorType::Enum LiveModule::Update(FileAttributeCache* fileCache, Di
 				}
 			}
 
-			// EPIC MOD - Added VfsToLocal
-			coff::WriteRaw(VfsToLocal(wideObjPath).c_str(), rawCoff, removalStrategy);
+			coff::WriteRaw(wideObjPath.c_str(), rawCoff, removalStrategy);
 			coff::DestroyRaw(rawCoff);
 		}
 	}
@@ -2972,8 +2954,7 @@ LiveModule::ErrorType::Enum LiveModule::Update(FileAttributeCache* fileCache, Di
 		const symbols::Compiland* compiland = compilandIt->second;
 		const uint32_t compilandUniqueId = symbols::GetCompilandId(compiland, wideObjPath.c_str(), modifiedOrNewObjFiles);
 
-		// EPIC MOD - Added VfsToLocal
-		coff::ObjFile* coffFile = coff::OpenObj(VfsToLocal(wideObjPath).c_str());
+		coff::ObjFile* coffFile = coff::OpenObj(wideObjPath.c_str());
 		
 		if (coffFile && coffFile->memoryFile)
 		{
@@ -3034,14 +3015,12 @@ LiveModule::ErrorType::Enum LiveModule::Update(FileAttributeCache* fileCache, Di
 		pdbPath = string::Replace(originalPdbPath, L".pdb", patchInstanceStr + std::wstring(L".pdb"));
 		exePath = string::Replace(originalPdbPath, L".pdb", patchInstanceStr + std::wstring(L".exe"));
 		
-		// EPIC MOD - Added VfsToLocal
-		const Filesystem::PathAttributes& pdbAttributes = Filesystem::GetAttributes(VfsToLocal(pdbPath).c_str());
-		const Filesystem::PathAttributes& exeAttributes = Filesystem::GetAttributes(VfsToLocal(exePath).c_str());
+		const Filesystem::PathAttributes& pdbAttributes = Filesystem::GetAttributes(pdbPath.c_str());
+		const Filesystem::PathAttributes& exeAttributes = Filesystem::GetAttributes(exePath.c_str());
 
 		if (Filesystem::DoesExist(pdbAttributes))
 		{
-			// EPIC MOD - Added VfsToLocal
-			if (!Filesystem::DeleteIfExists(VfsToLocal(pdbPath).c_str()))
+			if (!Filesystem::DeleteIfExists(pdbPath.c_str()))
 			{
 				// PDB file could not be deleted
 				isExeOrPdbFileStillThere = true;
@@ -3050,8 +3029,7 @@ LiveModule::ErrorType::Enum LiveModule::Update(FileAttributeCache* fileCache, Di
 		
 		if (Filesystem::DoesExist(exeAttributes))
 		{
-			// EPIC MOD - Added VfsToLocal
-			if (!Filesystem::DeleteIfExists(VfsToLocal(exePath).c_str()))
+			if (!Filesystem::DeleteIfExists(exePath.c_str()))
 			{
 				// EXE file could not be deleted
 				isExeOrPdbFileStillThere = true;
@@ -3179,9 +3157,9 @@ LiveModule::ErrorType::Enum LiveModule::Update(FileAttributeCache* fileCache, Di
 	std::wstring linkerCommandLine;
 
 	// If there are vfs entries, we need to run the linker through ubacli with -vfs options to make sure linker result is the same as when running through UBT
-	if (!vfsEntries.empty())
+	if (!Filesystem::GetVfsEntries().empty())
 	{
-		for (auto& entry : vfsEntries)
+		for (auto& entry : Filesystem::GetVfsEntries())
 		{
 			linkerCommandLine += L" -vfs=\"" + entry.virtualPath + L";" + entry.localPath + L"\"";
 		}
@@ -3189,7 +3167,7 @@ LiveModule::ErrorType::Enum LiveModule::Update(FileAttributeCache* fileCache, Di
 		linkerCommandLine += linkerPath.c_str();
 		linkerCommandLine += '\"';
 		linkerPath = *FPaths::ConvertRelativePathToFull(FPaths::EngineDir() / L"Binaries/Win64/UnrealBuildAccelerator/x64/UbaCli.exe");
-		linkerWorkingDirectory = VfsToLocal(linkerWorkingDirectory);
+		linkerWorkingDirectory = Filesystem::Devirtualize(linkerWorkingDirectory);
 	}
 	else
 	{
@@ -3215,12 +3193,11 @@ LiveModule::ErrorType::Enum LiveModule::Update(FileAttributeCache* fileCache, Di
 		const std::wstring originalPath(string::ToWideString(objPath));
 		const std::wstring bakPath = originalPath + L".bak";
 
-		// BEGIN EPIC MOD - Added VfsToLocal
-		const Filesystem::PathAttributes& attributes = Filesystem::GetAttributes(VfsToLocal(bakPath).c_str());
+		const Filesystem::PathAttributes& attributes = Filesystem::GetAttributes(bakPath.c_str());
 		if (Filesystem::DoesExist(attributes))
 		{
-			Filesystem::Delete(VfsToLocal(originalPath).c_str());
-			Filesystem::Move(VfsToLocal(bakPath).c_str(), VfsToLocal(originalPath).c_str());
+			Filesystem::Delete(originalPath.c_str());
+			Filesystem::Move(bakPath.c_str(), originalPath.c_str());
 		}
 		// END EPIC MOD
 	}
@@ -3281,8 +3258,7 @@ LiveModule::ErrorType::Enum LiveModule::Update(FileAttributeCache* fileCache, Di
 	// END EPIC MOD
 
 	// try to load patch image
-	// EPIC MOD - Added VfsToLocal
-	executable::Image* image = executable::OpenImage(VfsToLocal(exePath).c_str(), Filesystem::OpenMode::READ_WRITE);
+	executable::Image* image = executable::OpenImage(exePath.c_str(), Filesystem::OpenMode::READ_WRITE);
 	if (!image)
 	{
 		LC_ERROR_USER("Cannot load patch executable %S", exePath.c_str());
@@ -3325,8 +3301,7 @@ LiveModule::ErrorType::Enum LiveModule::Update(FileAttributeCache* fileCache, Di
 			const PerProcessData& data = processData[i];
 
 			commands::LoadPatch cmd = {};
-			// EPIC MOD - Added VfsToLocal
-			wcscpy_s(cmd.path, VfsToLocal(exePath).c_str());
+			wcscpy_s(cmd.path, Filesystem::Devirtualize(exePath).c_str());
 
 #if LC_64_BIT
 			// before doing anything further, we need to ensure that the patch can be loaded into the address space at a suitable location.
@@ -3353,16 +3328,14 @@ LiveModule::ErrorType::Enum LiveModule::Update(FileAttributeCache* fileCache, Di
 			{
 				// this image needs to be copied. create a new name based on the process ID, which must be unique
 				rebasedExePath = string::Replace(rebasedExePath, L".exe", std::wstring(L"_PID_") + std::to_wstring(+data.liveProcess->GetProcessId()) + L".exe");
-				// EPIC MOD - Added VfsToLocal
-				Filesystem::Copy(VfsToLocal(exePath).c_str(), VfsToLocal(rebasedExePath).c_str());
-				wcscpy_s(cmd.path, VfsToLocal(rebasedExePath).c_str());
+				Filesystem::Copy(exePath.c_str(), rebasedExePath.c_str());
+				wcscpy_s(cmd.path, Filesystem::Devirtualize(rebasedExePath).c_str());
 			}
 			
 			if (imageNeedsToBeRebased)
 			{
 				// rebase the patch image to its preferred base address
-				// EPIC MOD - Added VfsToLocal
-				executable::Image* rebasedImage = executable::OpenImage(VfsToLocal(rebasedExePath).c_str(), Filesystem::OpenMode::READ_WRITE);
+				executable::Image* rebasedImage = executable::OpenImage(rebasedExePath.c_str(), Filesystem::OpenMode::READ_WRITE);
 				LC_LOG_DEV("Rebasing patch executable to image base 0x%" PRIX64 " (PID: %d)", preferredImageBase, data.liveProcess->GetProcessId());
 				executable::RebaseImage(rebasedImage, preferredImageBase);
 				executable::CloseImage(rebasedImage);
@@ -3464,8 +3437,7 @@ LiveModule::ErrorType::Enum LiveModule::Update(FileAttributeCache* fileCache, Di
 
 	telemetry::Scope loadPatchPDBScope("Loading PDB database");
 
-	// EPIC MOD - Added VfsToLocal
-	symbols::Provider* patchSymbolProvider = symbols::OpenEXE(VfsToLocal(exePath).c_str(), symbols::OpenOptions::ACCUMULATE_SIZE);
+	symbols::Provider* patchSymbolProvider = symbols::OpenEXE(Filesystem::Devirtualize(exePath).c_str(), symbols::OpenOptions::ACCUMULATE_SIZE);
 	symbols::DiaCompilandDB* patch_diaCompilandDb = symbols::GatherDiaCompilands(patchSymbolProvider);
 	IDiaSymbol* patch_linkerSymbol = symbols::FindLinkerSymbol(patch_diaCompilandDb);
 
@@ -3481,10 +3453,9 @@ LiveModule::ErrorType::Enum LiveModule::Update(FileAttributeCache* fileCache, Di
 	scheduler::RunTask(taskPatch_symbolDB);
 
 
-	// EPIC MOD - Added VfsToLocal
-	auto taskPatch_contributionDB = scheduler::CreateTask(taskRootPatchLoading, [realExePath = VfsToLocal(exePath)]()
+	auto taskPatch_contributionDB = scheduler::CreateTask(taskRootPatchLoading, [exePath]()
 	{
-		symbols::Provider* localProvider = symbols::OpenEXE(realExePath.c_str(), symbols::OpenOptions::NONE);
+		symbols::Provider* localProvider = symbols::OpenEXE(Filesystem::Devirtualize(exePath).c_str(), symbols::OpenOptions::NONE);
 		symbols::DiaCompilandDB* localDiaCompilandDb = symbols::GatherDiaCompilands(localProvider);
 
 		auto db = symbols::GatherContributions(localProvider);
@@ -3500,10 +3471,9 @@ LiveModule::ErrorType::Enum LiveModule::Update(FileAttributeCache* fileCache, Di
 	// note that we only gather symbols from .obj contained in the new patch executable.
 	// therefore we need to extract its compiland database as well, and cannot use the one from
 	// the original executable.
-	// EPIC MOD - Added VfsToLocal
-	auto taskPatch_compilandDB = scheduler::CreateTask(taskRootPatchLoading, [this, realExePath = VfsToLocal(exePath)]()
+	auto taskPatch_compilandDB = scheduler::CreateTask(taskRootPatchLoading, [this, exePath]()
 	{
-		symbols::Provider* localProvider = symbols::OpenEXE(realExePath.c_str(), symbols::OpenOptions::NONE);
+		symbols::Provider* localProvider = symbols::OpenEXE(Filesystem::Devirtualize(exePath).c_str(), symbols::OpenOptions::NONE);
 		symbols::DiaCompilandDB* localDiaCompilandDb = symbols::GatherDiaCompilands(localProvider);
 
 		uint32_t options = 0u;
@@ -3618,8 +3588,7 @@ LiveModule::ErrorType::Enum LiveModule::Update(FileAttributeCache* fileCache, Di
 		// update the COFF cache for new patch compilands.
 		// there may be files for which we don't have a database yet, even though we updated the database for all compiled files.
 		// this can happen when a new .obj that is part of a library is linked in for the first time.
-		// EPIC MOD - Added VfsToLocal
-		types::vector<symbols::ObjPath> updatedCoffs = UpdateCoffCache(VfsToLocal, patch_compilandDB->compilands, m_coffCache, CacheUpdate::NON_EXISTANT, modifiedOrNewObjFiles);
+		types::vector<symbols::ObjPath> updatedCoffs = UpdateCoffCache(patch_compilandDB->compilands, m_coffCache, CacheUpdate::NON_EXISTANT, modifiedOrNewObjFiles);
 
 
 		// similarly, reconstruct symbols and dynamic initializers for new .obj that have been pulled in for the first time.
@@ -3678,8 +3647,7 @@ LiveModule::ErrorType::Enum LiveModule::Update(FileAttributeCache* fileCache, Di
 	// END EPIC MOD
 
 	// reconstruct symbols for all compilands that are part of the new patch executable
-	// EPIC MOD - Added VfsToLocal
-	executable::Image* patchImage = executable::OpenImage(VfsToLocal(exePath).c_str(), Filesystem::OpenMode::READ);
+	executable::Image* patchImage = executable::OpenImage(exePath.c_str(), Filesystem::OpenMode::READ);
 	executable::ImageSectionDB* patchImageSections = executable::GatherImageSectionDB(patchImage);
 
 	// gather the dynamic initializers and remaining symbols by walking the module
diff --git a/Engine/Source/Developer/Windows/LiveCodingServer/Private/External/LC_LiveModule.h b/Engine/Source/Developer/Windows/LiveCodingServer/Private/External/LC_LiveModule.h
--- a/Engine/Source/Developer/Windows/LiveCodingServer/Private/External/LC_LiveModule.h
+++ b/Engine/Source/Developer/Windows/LiveCodingServer/Private/External/LC_LiveModule.h
@@ -24,12 +24,6 @@ class ModulePatch;
 class LiveProcess;
 class VirtualMemoryRange;
 
-struct VfsEntry
-{
-	std::wstring virtualPath;
-	std::wstring localPath;
-};
-
 class LiveModule
 {
 public:
@@ -88,7 +82,7 @@ public:
 	// in EXTERNAL_BUILD_SYSTEM mode, this does not compile files but builds a patch containing modified .objs, loading the patch into the host application.
 	// optionally, an array of modified or new .objs can be given in this mode, which builds a patch containing these files, not checking for any other modifications.
 	// BEGIN EPIC MOD
-	ErrorType::Enum Update(FileAttributeCache* fileCache, DirectoryCache* directoryCache, UpdateType::Enum updateType, const types::vector<symbols::ModifiedObjFile>& modifiedOrNewObjFiles, const types::vector<std::wstring>& additionalLibraries, const types::vector<VfsEntry>& vfsEntries);
+	ErrorType::Enum Update(FileAttributeCache* fileCache, DirectoryCache* directoryCache, UpdateType::Enum updateType, const types::vector<symbols::ModifiedObjFile>& modifiedOrNewObjFiles, const types::vector<std::wstring>& additionalLibraries);
 	// END EPIC MOD
 	bool InstallCompiledPatches(LiveProcess* liveProcess, void* originalModuleBase);
 
diff --git a/Engine/Source/Developer/Windows/LiveCodingServer/Private/External/LC_ServerCommandThread.cpp b/Engine/Source/Developer/Windows/LiveCodingServer/Private/External/LC_ServerCommandThread.cpp
--- a/Engine/Source/Developer/Windows/LiveCodingServer/Private/External/LC_ServerCommandThread.cpp
+++ b/Engine/Source/Developer/Windows/LiveCodingServer/Private/External/LC_ServerCommandThread.cpp
@@ -777,7 +777,7 @@ bool ServerCommandThread::EnableRequiredModules(const TArray<FString>& RequiredM
 // END EPIC MOD
 
 // BEGIN EPIC MOD
-void ServerCommandThread::CompileChanges(bool didAllProcessesMakeProgress, commands::PostCompileResult& postCompileResult, types::vector<VfsEntry>& vfsEntries)
+void ServerCommandThread::CompileChanges(bool didAllProcessesMakeProgress, commands::PostCompileResult& postCompileResult)
 // END EPIC MOD
 {
 	// recompile files, if any
@@ -805,9 +805,6 @@ void ServerCommandThread::CompileChanges(bool didAllProcessesMakeProgress, comma
 
 		GLiveCodingServer->GetStatusChangeDelegate().ExecuteIfBound(L"Compiling changes for live coding...");
 
-		// EPIC MOD - Added VfsToLocal
-		FVfsEntries VfsToLocalPaths;
-
 		// Keep retrying the compile until we've added all the required modules
 		FModuleToModuleFiles ModuleToModuleFiles;
 		ELiveCodingCompileReason CompileReason = ELiveCodingCompileReason::Initial;
@@ -836,8 +833,7 @@ void ServerCommandThread::CompileChanges(bool didAllProcessesMakeProgress, comma
 
 			// Execute the compile
 			TArray<FString> RequiredModules;
-			// EPIC MOD - Added VfsToLocal
-			ELiveCodingCompileResult CompileResult = CompileDelegate.Execute(Targets, ValidModules, LazyLoadModules, RequiredModules, ModuleToModuleFiles, VfsToLocalPaths, CompileReason);
+			ELiveCodingCompileResult CompileResult = CompileDelegate.Execute(Targets, ValidModules, LazyLoadModules, RequiredModules, ModuleToModuleFiles, CompileReason);
 			if (CompileResult == ELiveCodingCompileResult::Success)
 			{
 				break;
@@ -869,21 +865,6 @@ void ServerCommandThread::CompileChanges(bool didAllProcessesMakeProgress, comma
 		// Reset the unity file cache
 		symbols::ResetCachedUnityManifests();
 
-		// BEGIN EPIC MOD - Added VfsToLocal
-		for (TPair<FString, FString>& kv : VfsToLocalPaths)
-		{
-			vfsEntries.emplace_back(VfsEntry{*kv.Key, *kv.Value});
-		}
-
-		auto VfsToLocal = [&](const std::wstring& path)
-			{
-				for (auto& entry : vfsEntries)
-					if (entry.virtualPath.compare(0, entry.virtualPath.size(), path, 0, entry.virtualPath.size()) == 0)
-						return entry.localPath + (path.data() + entry.virtualPath.size());
-				return path;
-			};
-		// END EPIC MOD - Added VfsToLocal
-
 		// Build up a list of all the modified object files in each module
 		types::unordered_map<std::wstring, const LiveModule*> EnabledModulesByName;
 		for (const LiveModule* liveModule : m_liveModules)
@@ -894,8 +875,7 @@ void ServerCommandThread::CompileChanges(bool didAllProcessesMakeProgress, comma
 		for(const TPair<FString, FModuleFiles>& Pair : ModuleToModuleFiles)
 		{
 			const LiveModule* liveModule = nullptr;
-			// EPIC MOD - Added VfsToLocal
-			std::wstring ModuleFileName = VfsToLocal(Filesystem::NormalizePath(*Pair.Key).GetString());
+			std::wstring ModuleFileName = Filesystem::Devirtualize(Filesystem::NormalizePath(*Pair.Key).GetString());
 			types::unordered_map<std::wstring, const LiveModule*>::iterator ix = EnabledModulesByName.find(ModuleFileName);
 			if (ix == EnabledModulesByName.end())
 			{
@@ -928,8 +908,7 @@ void ServerCommandThread::CompileChanges(bool didAllProcessesMakeProgress, comma
 			{
 				std::wstring NormalizedObjectFile = Filesystem::NormalizePath(*ObjectFile).GetString();
 
-				// EPIC MOD - Added VfsToLocal
-				if (!liveModule->IsModifiedSource(VfsToLocal(NormalizedObjectFile).c_str()))
+				if (!liveModule->IsModifiedSource(Filesystem::Devirtualize(NormalizedObjectFile).c_str()))
 				{
 					continue;
 				}
@@ -956,8 +935,7 @@ void ServerCommandThread::CompileChanges(bool didAllProcessesMakeProgress, comma
 
 					// Move the new file into place
 					m_restoreFiles.push_back(std::make_pair(OriginalObjectFile, NormalizedObjectFile));
-					// EPIC MOD - Added VfsToLocal
-					Filesystem::Move(VfsToLocal(NormalizedObjectFile).c_str(), VfsToLocal(OriginalObjectFile).c_str());
+					Filesystem::Move(NormalizedObjectFile.c_str(), OriginalObjectFile.c_str());
 					NormalizedObjectFile = OriginalObjectFile;
 				}
 
@@ -1026,7 +1004,7 @@ void ServerCommandThread::CompileChanges(bool didAllProcessesMakeProgress, comma
 
 		// build a patch with the given list of .objs for this module
 		const types::vector<symbols::ModifiedObjFile>& objFiles = objFilesIt->second;
-		LiveModule::ErrorType::Enum moduleUpdateError = liveModule->Update(&fileCache, m_directoryCache, updateType, objFiles, libFiles, vfsEntries);
+		LiveModule::ErrorType::Enum moduleUpdateError = liveModule->Update(&fileCache, m_directoryCache, updateType, objFiles, libFiles);
 
 		// only accept new error conditions for this module if there haven't been any updates until now.
 		// this ensures that error conditions are kept and can be shown when updating several modules at once.
@@ -1387,9 +1365,8 @@ Thread::ReturnValue ServerCommandThread::CompileThread(void)
 			// END EPIC MOD
 
 			// BEGIN EPIC MOD 
-			types::vector<VfsEntry> vfsToLocalPaths;
 			commands::PostCompileResult postCompileResult = commands::PostCompileResult::Success;
-			CompileChanges(didAllProcessesMakeProgress, postCompileResult, vfsToLocalPaths);
+			CompileChanges(didAllProcessesMakeProgress, postCompileResult);
 			// END EPIC MOD
 
 			// BEGIN EPIC MOD - Add the ability for pre and post compile notifications
@@ -1397,21 +1374,11 @@ Thread::ReturnValue ServerCommandThread::CompileThread(void)
 			// END EPIC MOD
 
 
-			// Begin EPIC MOD - Added VfsToLocal
-			auto VfsToLocal = [&](const std::wstring& path)
-				{
-					for (auto& entry : vfsToLocalPaths)
-						if (entry.virtualPath.compare(0, entry.virtualPath.size(), path, 0, entry.virtualPath.size()) == 0)
-							return entry.localPath + (path.data() + entry.virtualPath.size());
-					return path;
-				};
-			// EndEPIC MOD - Added VfsToLocal
-
 			// BEGIN EPIC MOD - Non-destructive compile
 			for (std::vector<std::pair<std::wstring, std::wstring>>::reverse_iterator it = m_restoreFiles.rbegin(); it != m_restoreFiles.rend(); it++)
 			{
-				Filesystem::DeleteIfExists(VfsToLocal(it->second).c_str());
-				Filesystem::Move(VfsToLocal(it->first).c_str(), VfsToLocal(it->second).c_str());
+				Filesystem::DeleteIfExists(it->second.c_str());
+				Filesystem::Move(it->first.c_str(), it->second.c_str());
 			}
 			m_restoreFiles.clear();
 			// END EPIC MOD
diff --git a/Engine/Source/Developer/Windows/LiveCodingServer/Private/External/LC_ServerCommandThread.h b/Engine/Source/Developer/Windows/LiveCodingServer/Private/External/LC_ServerCommandThread.h
--- a/Engine/Source/Developer/Windows/LiveCodingServer/Private/External/LC_ServerCommandThread.h
+++ b/Engine/Source/Developer/Windows/LiveCodingServer/Private/External/LC_ServerCommandThread.h
@@ -78,7 +78,7 @@ private:
 	LiveProcess* FindProcessById(Process::Id processId);
 
 	// BEGIN EPIC MOD
-	void CompileChanges(bool didAllProcessesMakeProgress, commands::PostCompileResult& postCompileResult, types::vector<VfsEntry>& vfsToLocalPaths);
+	void CompileChanges(bool didAllProcessesMakeProgress, commands::PostCompileResult& postCompileResult);
 	// END EPIC MOD
 
 	// actions
diff --git a/Engine/Source/Developer/Windows/LiveCodingServer/Public/ILiveCodingServer.h b/Engine/Source/Developer/Windows/LiveCodingServer/Public/ILiveCodingServer.h
--- a/Engine/Source/Developer/Windows/LiveCodingServer/Public/ILiveCodingServer.h
+++ b/Engine/Source/Developer/Windows/LiveCodingServer/Public/ILiveCodingServer.h
@@ -53,8 +53,6 @@ struct FModuleFiles
 };
 typedef TMap<FString, FModuleFiles> FModuleToModuleFiles;
 
-typedef TArray<TPair<FString, FString>> FVfsEntries;
-
 class ILiveCodingServer : public IModularFeature
 {
 public:
@@ -77,7 +75,7 @@ public:
 	DECLARE_DELEGATE_TwoParams(FLogOutputDelegate, ELiveCodingLogVerbosity, const wchar_t*);
 	virtual FLogOutputDelegate& GetLogOutputDelegate() = 0;
 
-	DECLARE_DELEGATE_RetVal_SevenParams(ELiveCodingCompileResult, FCompileDelegate, const TArray<FString>&, const TArray<FString>&, const TSet<FString>&, TArray<FString>&, FModuleToModuleFiles&, FVfsEntries&, ELiveCodingCompileReason);
+	DECLARE_DELEGATE_RetVal_SixParams(ELiveCodingCompileResult, FCompileDelegate, const TArray<FString>&, const TArray<FString>&, const TSet<FString>&, TArray<FString>&, FModuleToModuleFiles&, ELiveCodingCompileReason);
 	virtual FCompileDelegate& GetCompileDelegate() = 0;
 
 	DECLARE_DELEGATE(FCompileStartedDelegate);
diff --git a/Engine/Source/Programs/LiveCodingConsole/Private/LiveCodingConsole.cpp b/Engine/Source/Programs/LiveCodingConsole/Private/LiveCodingConsole.cpp
--- a/Engine/Source/Programs/LiveCodingConsole/Private/LiveCodingConsole.cpp
+++ b/Engine/Source/Programs/LiveCodingConsole/Private/LiveCodingConsole.cpp
@@ -273,7 +273,7 @@ private:
 	}
 
 	ELiveCodingCompileResult CompilePatch(const TArray<FString>& Targets, const TArray<FString>& ValidModules, const TSet<FString>& LazyLoadModules,
-		TArray<FString>& RequiredModules, FModuleToModuleFiles& ModuleToModuleFiles, FVfsEntries& VfsEntries, ELiveCodingCompileReason CompileReason)
+		TArray<FString>& RequiredModules, FModuleToModuleFiles& ModuleToModuleFiles, ELiveCodingCompileReason CompileReason)
 	{
 		// Get the UBT path
 		FString Executable;
@@ -422,8 +422,6 @@ private:
 			}
 		}
 
-		VfsEntries = Manifest.VfsEntries;
-
 		return ELiveCodingCompileResult::Success;
 	}
 
diff --git a/Engine/Source/Runtime/BuildSettings/Private/BuildSettings.cpp b/Engine/Source/Runtime/BuildSettings/Private/BuildSettings.cpp
--- a/Engine/Source/Runtime/BuildSettings/Private/BuildSettings.cpp
+++ b/Engine/Source/Runtime/BuildSettings/Private/BuildSettings.cpp
@@ -116,4 +116,9 @@ namespace BuildSettings
 	{
 		return UE_VFS_PATHS;
 	}
+
+	const TCHAR* GetVfsPathsWide()
+	{
+		return TEXT(UE_VFS_PATHS);
+	}
 }
diff --git a/Engine/Source/Runtime/BuildSettings/Public/BuildSettings.h b/Engine/Source/Runtime/BuildSettings/Public/BuildSettings.h
--- a/Engine/Source/Runtime/BuildSettings/Public/BuildSettings.h
+++ b/Engine/Source/Runtime/BuildSettings/Public/BuildSettings.h
@@ -142,4 +142,9 @@ namespace BuildSettings
 	 * Semicolon separated string containing virtual;local;virtual;local; paths. Ends with semicolon
 	 */
 	BUILDSETTINGS_API const char* GetVfsPaths();
+
+	/**
+	 * Semicolon separated string containing virtual;local;virtual;local; paths. Ends with semicolon
+	 */
+	BUILDSETTINGS_API const TCHAR* GetVfsPathsWide();
 }
