[LiveCodingConsole]
* Fixed so lpp handles compressed pch obj files

diff --git a/Engine/Source/Developer/Windows/LiveCoding/Private/External/LC_MemoryMappedFile.cpp b/Engine/Source/Developer/Windows/LiveCoding/Private/External/LC_MemoryMappedFile.cpp
--- a/Engine/Source/Developer/Windows/LiveCoding/Private/External/LC_MemoryMappedFile.cpp
+++ b/Engine/Source/Developer/Windows/LiveCoding/Private/External/LC_MemoryMappedFile.cpp
@@ -82,6 +82,7 @@ namespace Filesystem
 		HANDLE file;
 		HANDLE fileMapping;
 		void* baseAddress;
+		uint64_t size;
 
 		LC_DISABLE_ASSIGNMENT(MemoryMappedFile);
 	};
@@ -118,6 +119,8 @@ Filesystem::MemoryMappedFile* Filesystem::OpenMemoryMappedFile(const wchar_t* pa
 		return nullptr;
 	}
 
+	uint64_t fileSize = 0;
+
 	if (checkIfCompressed)
 	{
 		auto data = (const BYTE*)baseAddress;
@@ -125,7 +128,7 @@ Filesystem::MemoryMappedFile* Filesystem::OpenMemoryMappedFile(const wchar_t* pa
 		if (memcmp(data, UbaMagic, sizeof(UbaMagic)) == 0)
 		{
 			data += 24;
-			uint64_t fileSize = *(uint64_t*)data;
+			fileSize = *(uint64_t*)data;
 			data += sizeof(uint64_t);
 
 			::ULARGE_INTEGER integer = {};
@@ -161,7 +164,7 @@ Filesystem::MemoryMappedFile* Filesystem::OpenMemoryMappedFile(const wchar_t* pa
 	}
 
 
-	return new MemoryMappedFile { file, fileMapping, baseAddress };
+	return new MemoryMappedFile { file, fileMapping, baseAddress, fileSize };
 }
 
 
@@ -187,6 +190,13 @@ const void* Filesystem::GetMemoryMappedFileData(const MemoryMappedFile* file)
 	return file->baseAddress;
 }
 
+uint64_t Filesystem::GetMemoryMappedSize(const MemoryMappedFile* file)
+{
+	if (file->size)
+		return file->size;
+	else
+		return GetMemoryMappedFileSizeOnDisk(file);
+}
 
 uint64_t Filesystem::GetMemoryMappedFileSizeOnDisk(const MemoryMappedFile* file)
 {
diff --git a/Engine/Source/Developer/Windows/LiveCoding/Private/External/LC_MemoryMappedFile.h b/Engine/Source/Developer/Windows/LiveCoding/Private/External/LC_MemoryMappedFile.h
--- a/Engine/Source/Developer/Windows/LiveCoding/Private/External/LC_MemoryMappedFile.h
+++ b/Engine/Source/Developer/Windows/LiveCoding/Private/External/LC_MemoryMappedFile.h
@@ -18,6 +18,7 @@ namespace Filesystem
 	void* GetMemoryMappedFileData(MemoryMappedFile* file);
 	const void* GetMemoryMappedFileData(const MemoryMappedFile* file);
 
+	uint64_t GetMemoryMappedSize(const MemoryMappedFile* file);
 	uint64_t GetMemoryMappedFileSizeOnDisk(const MemoryMappedFile* file);
 }
 
diff --git a/Engine/Source/Developer/Windows/LiveCodingServer/Private/External/LC_Coff.cpp b/Engine/Source/Developer/Windows/LiveCodingServer/Private/External/LC_Coff.cpp
--- a/Engine/Source/Developer/Windows/LiveCodingServer/Private/External/LC_Coff.cpp
+++ b/Engine/Source/Developer/Windows/LiveCodingServer/Private/External/LC_Coff.cpp
@@ -289,10 +289,7 @@ namespace coff
 		using CoffAuxSymbol = typename coffDetail::CoffType::TypeMap<Type>::AuxSymbolType;
 
 		RawCoffType* rawCoff = new RawCoffType;
-		{
-			const Filesystem::PathAttributes& attributes = Filesystem::GetAttributes(string::ToWideString(file->filename).c_str());
-			rawCoff->size = Filesystem::GetSize(attributes);
-		}
+		rawCoff->size = Filesystem::GetMemoryMappedSize(file->memoryFile);
 		rawCoff->type = RawCoffType::TYPE;
 
 		const void* objFileBase = Filesystem::GetMemoryMappedFileData(file->memoryFile);
