[LiveCoding]
* Fixed bug where shutdown didn't properly unregister from dll notification causing read-after-free

diff --git a/Engine/Source/Developer/Windows/LiveCoding/Private/LiveCodingModule.cpp b/Engine/Source/Developer/Windows/LiveCoding/Private/LiveCodingModule.cpp
--- a/Engine/Source/Developer/Windows/LiveCoding/Private/LiveCodingModule.cpp
+++ b/Engine/Source/Developer/Windows/LiveCoding/Private/LiveCodingModule.cpp
@@ -1245,30 +1245,28 @@ bool FLiveCodingModule::StartLiveCoding(ELiveCodingStartupMode StartupMode)
 
 void FLiveCodingModule::StopLiveCoding()
 {
-	if (!GlobalMutex)
+	if (CallbackCookie)
 	{
-		return;
-	}
-
-	LppShutdown();
-
-	ReleaseMutex(GlobalMutex);
-	GlobalMutex = nullptr;
-
-	FCoreDelegates::OnEndFrame.Remove(EndFrameDelegateHandle);
+		LppShutdown();
+		FCoreDelegates::OnEndFrame.Remove(EndFrameDelegateHandle);
 
-	IConsoleManager& ConsoleManager = IConsoleManager::Get();
-	ConsoleManager.UnregisterConsoleObject(SourceProjectVariable);
-	ConsoleManager.UnregisterConsoleObject(ConsolePathVariable);
-	ConsoleManager.UnregisterConsoleObject(CompileCommand);
-	ConsoleManager.UnregisterConsoleObject(CompileSyncCommand);
-	ConsoleManager.UnregisterConsoleObject(EnableCommand);
+		IConsoleManager& ConsoleManager = IConsoleManager::Get();
+		ConsoleManager.UnregisterConsoleObject(SourceProjectVariable);
+		ConsoleManager.UnregisterConsoleObject(ConsolePathVariable);
+		ConsoleManager.UnregisterConsoleObject(CompileCommand);
+		ConsoleManager.UnregisterConsoleObject(CompileSyncCommand);
+		ConsoleManager.UnregisterConsoleObject(EnableCommand);
 
-	// Unregister from the dll notifications
-	if (CallbackCookie)
-	{
+		// Unregister from the dll notifications
 		FNtDllFunction UnregisterFunc("LdrUnregisterDllNotification");
 		UnregisterFunc(CallbackCookie);
+		CallbackCookie = nullptr;
+	}
+
+	if (GlobalMutex)
+	{
+		ReleaseMutex(GlobalMutex);
+		GlobalMutex = nullptr;
 	}
 }
 
