[shaders] D3D symbols improvements
- change to exporting external shader PDBs for D3D formats
- output dxil file for aftermath as an additional symbol file rather than having it tied to shader debug info, and don't bother outputting two copies (aftermath loads the PDB anyways, having reflection data in the DXIL is of no use)
- clean up some old unused code cruft in the dxc compiler and simplify the dxil container stripping code (the only thing we need to strip explicitly is reflection data)
- fix FXC batch file generation to include symbol-related args
- add "symbols info" support for FXC-compiled shaders; note that this requires generating debug info then stripping as FXC doesn't generate the debug name blob without this, however in my tests this has no impact to compile performance (my theory is that FXC is always generating debug info under the hood, and the flags only control whether it is saved into the resulting shader)

diff --git a/Engine/Source/Developer/Windows/ShaderFormatD3D/Private/D3DShaderCompiler.cpp b/Engine/Source/Developer/Windows/ShaderFormatD3D/Private/D3DShaderCompiler.cpp
--- a/Engine/Source/Developer/Windows/ShaderFormatD3D/Private/D3DShaderCompiler.cpp
+++ b/Engine/Source/Developer/Windows/ShaderFormatD3D/Private/D3DShaderCompiler.cpp
@@ -234,6 +234,17 @@ static FString D3D11CreateShaderCompileCommandLine(
 		FXCCommandline += FString(TEXT(" /WX"));
 	}
 
+	if (CompileFlags & D3DCOMPILE_DEBUG_NAME_FOR_BINARY)
+	{
+		CompileFlags &= ~D3DCOMPILE_DEBUG_NAME_FOR_BINARY;
+		FXCCommandline += FString(TEXT(" /Zsb"));
+	}
+	else if (CompileFlags & D3DCOMPILE_DEBUG_NAME_FOR_SOURCE)
+	{
+		CompileFlags &= ~D3DCOMPILE_DEBUG_NAME_FOR_SOURCE;
+		FXCCommandline += FString(TEXT(" /Zss"));
+	}
+
 	switch (CompileFlags & SHADER_OPTIMIZATION_LEVEL_MASK)
 	{
 		case D3DCOMPILE_OPTIMIZATION_LEVEL2:
@@ -346,6 +357,18 @@ typedef HRESULT(WINAPI *pD3DStripShader)
 	 __in UINT       uStripFlags,
 	__out ID3DBlob** ppStrippedBlob);
 
+typedef HRESULT(WINAPI *pD3DGetBlobPart)
+	(__in_bcount(SrcDataSize) LPCVOID pSrcData,
+	 __in SIZE_T SrcDataSize,
+	 __in D3D_BLOB_PART Part,
+	 __in UINT Flags,
+	 __out ID3DBlob** ppPart);
+
+typedef HRESULT(WINAPI* pD3DGetDebugInfo)
+	(__in_bcount(SrcDataSize) LPCVOID pSrcData,
+	 __in SIZE_T SrcDataSize,
+	 __out ID3DBlob** ppPart);
+
 #define DEFINE_GUID_FOR_CURRENT_COMPILER(name, l, w1, w2, b1, b2, b3, b4, b5, b6, b7, b8) \
 	static const GUID name = { l, w1, w2, { b1, b2,  b3,  b4,  b5,  b6,  b7,  b8 } }
 
@@ -363,6 +386,8 @@ public:
 	static pD3DReflect GetReflect() { return Instance().Reflect; }
 	static pD3DDisassemble GetDisassemble() { return Instance().Disassemble; }
 	static pD3DStripShader GetStripShader() { return Instance().StripShader; }
+	static pD3DGetBlobPart GetGetBlobPart() { return Instance().GetBlobPart; }
+	static pD3DGetDebugInfo GetGetDebugInfo() { return Instance().GetDebugInfo; }
 
 private:
 	FxcCompilerFunctions()
@@ -377,6 +402,8 @@ private:
 		Reflect = (pD3DReflect)(void*)GetProcAddress(CompilerDLL, "D3DReflect");
 		Disassemble = (pD3DDisassemble)(void*)GetProcAddress(CompilerDLL, "D3DDisassemble");
 		StripShader = (pD3DStripShader)(void*)GetProcAddress(CompilerDLL, "D3DStripShader");
+		GetBlobPart = (pD3DGetBlobPart)(void*)GetProcAddress(CompilerDLL, "D3DGetBlobPart");
+		GetDebugInfo = (pD3DGetDebugInfo)(void*)GetProcAddress(CompilerDLL, "D3DGetDebugInfo");
 	}
 
 	static FxcCompilerFunctions& Instance()
@@ -390,6 +417,8 @@ private:
 	pD3DReflect Reflect = nullptr;
 	pD3DDisassemble Disassemble = nullptr;
 	pD3DStripShader StripShader = nullptr;
+	pD3DGetBlobPart GetBlobPart = nullptr;
+	pD3DGetDebugInfo GetDebugInfo = nullptr;
 };
 
 static int D3DExceptionFilter(bool bCatchException)
@@ -758,6 +787,8 @@ static bool CompileAndProcessD3DShaderFXCExt(
 	pD3DReflect D3DReflectFunc = FxcCompilerFunctions::GetReflect();
 	pD3DDisassemble D3DDisassembleFunc = FxcCompilerFunctions::GetDisassemble();
 	pD3DStripShader D3DStripShaderFunc = FxcCompilerFunctions::GetStripShader();
+	pD3DGetBlobPart D3DGetBlobPartFunc = FxcCompilerFunctions::GetGetBlobPart();
+	pD3DGetDebugInfo D3DGetDebugInfoFunc = FxcCompilerFunctions::GetGetDebugInfo();
 
 	TRefCountPtr<ID3DBlob> Errors;
 
@@ -1181,12 +1212,64 @@ static bool CompileAndProcessD3DShaderFXCExt(
 			}
 
 			TRefCountPtr<ID3DBlob> CompressedData;
-
-			if (Input.Environment.CompilerFlags.Contains(CFLAG_GenerateSymbols))
+			TRefCountPtr<ID3DBlob> DebugDataBlob;
+			TRefCountPtr<ID3DBlob> DebugNameBlob;
+			
+			bool bGenerateSymbolsInfo = Input.Environment.CompilerFlags.Contains(CFLAG_GenerateSymbolsInfo);
+			bool bGenerateSymbols = Input.Environment.CompilerFlags.Contains(CFLAG_GenerateSymbols);
+			if ((bGenerateSymbols || bGenerateSymbolsInfo) && D3DGetBlobPartFunc && D3DGetDebugInfoFunc)
 			{
-				CompressedData = Shader;
+				Result = D3DGetBlobPartFunc(Shader->GetBufferPointer(), Shader->GetBufferSize(), D3D_BLOB_DEBUG_NAME, 0u, DebugNameBlob.GetInitReference());
+
+				if (SUCCEEDED(Result))
+				{
+					// copypasta from https://devblogs.microsoft.com/pix/using-automatic-shader-pdb-resolution-in-pix/
+					struct ShaderDebugName
+					{
+						uint16_t Flags;       // Reserved, must be set to zero.
+						uint16_t NameLength;  // Length of the debug name, without null terminator.
+						// Followed by NameLength bytes of the UTF-8-encoded name.
+						// Followed by a null terminator.
+						// Followed by [0-3] zero bytes to align to a 4-byte boundary.
+					};
+
+					const ShaderDebugName* DebugNameData = reinterpret_cast<const ShaderDebugName*>(DebugNameBlob->GetBufferPointer());
+					const char* Name = reinterpret_cast<const char*>(DebugNameData + 1);
+
+					FD3DShaderDebugData DebugData;
+					FD3DShaderDebugData::FFile& PdbFile = DebugData.Files.AddDefaulted_GetRef();
+					PdbFile.Name = Name;
+
+					if (bGenerateSymbols)
+					{
+						Result = D3DGetBlobPartFunc(Shader->GetBufferPointer(), Shader->GetBufferSize(), D3D_BLOB_PDB, 0u, DebugDataBlob.GetInitReference());
+						if (SUCCEEDED(Result) && DebugDataBlob->GetBufferSize() > 0)
+						{
+							PdbFile.Contents = MakeArrayViewFromBlob(DebugDataBlob);
+						}
+						else
+						{
+							FShaderCompilerError NewError;
+							NewError.StrippedErrorMessage = TEXT("Symbol generation was requested, but no PDB blob exists in the compiler output.");
+							Output.Errors.Add(NewError);
+							Result = E_FAIL;
+							Output.bSucceeded = false;
+						}
+					}
+
+					FMemoryWriter Ar(Output.ShaderCode.GetSymbolWriteAccess());
+					Ar << DebugData;
+				}
+				else
+				{
+					FShaderCompilerError NewError;
+					NewError.StrippedErrorMessage = TEXT("Symbol or symbols info generation was requested, but no debug name blob exists in the compiler output.");
+					Output.Errors.Add(NewError);
+					Output.bSucceeded = false;
+				}
 			}
-			else if (D3DStripShaderFunc)
+
+			if (D3DStripShaderFunc)
 			{
 				// Strip shader reflection and debug info
 				D3D_SHADER_DATA ShaderData;
@@ -1259,9 +1342,21 @@ bool CompileAndProcessD3DShaderFXC(
 		// Unpack uniform matrices as row-major to match the CPU layout.
 		| D3DCOMPILE_PACK_MATRIX_ROW_MAJOR;
 
-	if (Input.Environment.CompilerFlags.Contains(CFLAG_GenerateSymbols))
+	bool bGenerateSymbols = Input.Environment.CompilerFlags.Contains(CFLAG_GenerateSymbols);
+	bool bGenerateSymbolsInfo = Input.Environment.CompilerFlags.Contains(CFLAG_GenerateSymbolsInfo);
+
+	if (bGenerateSymbols || bGenerateSymbolsInfo)
 	{
 		CompileFlags |= D3DCOMPILE_DEBUG;
+
+		if (Input.Environment.CompilerFlags.Contains(CFLAG_AllowUniqueSymbols))
+		{
+			CompileFlags |= D3DCOMPILE_DEBUG_NAME_FOR_SOURCE;
+		}
+		else
+		{
+			CompileFlags |= D3DCOMPILE_DEBUG_NAME_FOR_BINARY;
+		}
 	}
 
 	if (Input.Environment.CompilerFlags.Contains(CFLAG_Debug))
diff --git a/Engine/Source/Developer/Windows/ShaderFormatD3D/Private/D3DShaderCompilerDXC.cpp b/Engine/Source/Developer/Windows/ShaderFormatD3D/Private/D3DShaderCompilerDXC.cpp
--- a/Engine/Source/Developer/Windows/ShaderFormatD3D/Private/D3DShaderCompilerDXC.cpp
+++ b/Engine/Source/Developer/Windows/ShaderFormatD3D/Private/D3DShaderCompilerDXC.cpp
@@ -126,9 +126,6 @@ static const TCHAR* DxcErrorCodeToString(HRESULT Code)
 #undef SWITCHCASE_TO_STRING
 }
 
-// Utility variable so we can place a breakpoint while debugging
-static int32 GBreakpointDXC = 0;
-
 static void LogFailedHRESULT(const TCHAR* FailedExpressionStr, HRESULT Result)
 {
 	if (Result == E_OUTOFMEMORY)
@@ -281,18 +278,13 @@ public:
 				ExtraArguments.Add(TEXT("-Zsb"));
 			}
 
-			ExtraArguments.Add(TEXT("-Qembed_debug"));
+			// generate the debug information (PDB)
 			ExtraArguments.Add(TEXT("-Zi"));
 
-			ExtraArguments.Add(TEXT("-Fd"));
-			ExtraArguments.Add(TEXT(".\\"));
-
-			bKeepEmbeddedPDB = true;
+			// always strip the PDB from the DXIL output blob, we retrieve from the compile result and save it manually in the PlatformDebugData
+			ExtraArguments.Add("-Qstrip_debug");
 		}
 
-		// Reflection will be removed later, otherwise the disassembly won't contain variables
-		//ExtraArguments.Add(TEXT("-Qstrip_reflect"));
-
 		// disable undesired warnings
 		ExtraArguments.Add(TEXT("-Wno-parentheses-equality"));
 
@@ -312,11 +304,6 @@ public:
 		return DumpDebugInfoPath;
 	}
 
-	bool ShouldKeepEmbeddedPDB() const
-	{
-		return bKeepEmbeddedPDB;
-	}
-
 	bool ShouldDump() const
 	{
 		return bDump;
@@ -615,37 +602,8 @@ static void DisassembleAndSave(TRefCountPtr<IDxcCompiler3>& Compiler, IDxcBlob*
 	}
 }
 
-static void DumpFourCCParts(dxc::DxcDllSupport& DxcDllHelper, TRefCountPtr<IDxcBlob>& Blob)
+static HRESULT RemoveReflectionData(dxc::DxcDllSupport& DxcDllHelper, TRefCountPtr<IDxcBlob>& Dxil)
 {
-#if UE_BUILD_DEBUG && IS_PROGRAM
-	TRefCountPtr<IDxcContainerReflection> Refl;
-	VERIFYHRESULT(DxcDllHelper.CreateInstance2(GetDxcMalloc(), CLSID_DxcContainerReflection, Refl.GetInitReference()));
-
-	VERIFYHRESULT(Refl->Load(Blob));
-
-	uint32 Count = 0;
-	VERIFYHRESULT(Refl->GetPartCount(&Count));
-
-	FPlatformMisc::LowLevelOutputDebugStringf(TEXT("*** Blob Size: %d, %d Parts\n"), Blob->GetBufferSize(), Count);
-
-	for (uint32 Index = 0; Index < Count; ++Index)
-	{
-		char FourCC[5] = "\0\0\0\0";
-		VERIFYHRESULT(Refl->GetPartKind(Index, (uint32*)FourCC));
-		TRefCountPtr<IDxcBlob> Part;
-		Refl->GetPartContent(Index, Part.GetInitReference());
-		FPlatformMisc::LowLevelOutputDebugStringf(TEXT("* %d %s, Size %d\n"), Index, ANSI_TO_TCHAR(FourCC), (uint32)Part->GetBufferSize());
-	}
-#endif
-}
-
-static bool RemoveContainerParts(const TConstArrayView<uint32> PartCodes, dxc::DxcDllSupport& DxcDllHelper, TRefCountPtr<IDxcBlob>& Dxil)
-{
-	if (PartCodes.Num() == 0)
-	{
-		return false;
-	}
-
 	TRefCountPtr<IDxcOperationResult> Result;
 	TRefCountPtr<IDxcContainerBuilder> Builder;
 	TRefCountPtr<IDxcBlob> StrippedDxil;
@@ -653,29 +611,25 @@ static bool RemoveContainerParts(const TConstArrayView<uint32> PartCodes, dxc::D
 	VERIFYHRESULT(DxcDllHelper.CreateInstance2(GetDxcMalloc(), CLSID_DxcContainerBuilder, Builder.GetInitReference()));
 	VERIFYHRESULT(Builder->Load(Dxil));
 
-	bool bSuccess = true;
-
-	for (uint32 PartCode : PartCodes)
+	HRESULT Res = Builder->RemovePart(DXC_PART_REFLECTION_DATA);
+	if (FAILED(Res))
 	{
-		if (FAILED(Builder->RemovePart(PartCode)))
-		{
-			bSuccess = false;
-			break;
-		}
+		return Res;
 	}
 
-	if (bSuccess)
+	Res = Builder->SerializeContainer(Result.GetInitReference());
+	if (FAILED(Res))
 	{
-		VERIFYHRESULT(Builder->SerializeContainer(Result.GetInitReference()));
-		if (SUCCEEDED(Result->GetResult(StrippedDxil.GetInitReference())))
-		{
-			Dxil.SafeRelease();
-			Dxil = StrippedDxil;
-			return true;
-		}
+		return Res;
 	}
 
-	return bSuccess;
+	Res = Result->GetResult(StrippedDxil.GetInitReference());
+	if (SUCCEEDED(Res))
+	{
+		Dxil.SafeRelease();
+		Dxil = StrippedDxil;
+	}
+	return Res;
 }
 
 static HRESULT D3DCompileToDxil(const char* SourceText, const FDxcArguments& Arguments,
@@ -740,56 +694,19 @@ static HRESULT D3DCompileToDxil(const char* SourceText, const FDxcArguments& Arg
 		VERIFYHRESULT(CompileResult->GetOutput(DXC_OUT_REFLECTION, IID_PPV_ARGS(OutReflectionBlob.GetInitReference()), ReflectionNameBlob.GetInitReference()));
 		RetrieveDebugNameAndBlob(CompileResult, OutPdbName, OutPdbBlob.GetInitReference(), OutHash);
  
- 		TArray<uint32, TInlineAllocator<4>> PartsToRemove;
-		if (!Arguments.ShouldKeepEmbeddedPDB())
-		{
-			if (CompileResult->HasOutput(DXC_OUT_PDB))
-			{
-				PartsToRemove.Add(DXC_PART_PDB);
-			}
-			if (CompileResult->HasOutput(DXC_OUT_REFLECTION))
-			{
-				PartsToRemove.Add(DXC_PART_REFLECTION_DATA);
-			}
-		}
-
 		if (Arguments.ShouldDump())
 		{
 			// Dump disassembly before we strip reflection out
 			const FString& DisasmFilename = Arguments.GetDumpDisassemblyFilename();
 			check(DisasmFilename.Len() > 0);
 			DisassembleAndSave(Compiler, OutDxilBlob, DisasmFilename);
-
-			// Dump dxil (.d3dasm -> .dxil)
-			FString DxilFile = Arguments.GetDumpDisassemblyFilename().LeftChop(7) + TEXT("_refl.dxil");
-			SaveDxcBlobToFile(OutDxilBlob, DxilFile);
-
-			// Dump the PDB.
-			if (OutPdbBlob.IsValid() && !OutPdbName.IsEmpty())
-			{
-				const FString PdbFile = Arguments.GetDumpDebugInfoPath() / OutPdbName;
-				SaveDxcBlobToFile(OutPdbBlob, PdbFile);
-			}
-		}
-
-		DumpFourCCParts(DxcDllHelper, OutDxilBlob);
-		if (RemoveContainerParts(PartsToRemove, DxcDllHelper, OutDxilBlob))
-		{
-			DumpFourCCParts(DxcDllHelper, OutDxilBlob);
 		}
 
-		if (Arguments.ShouldDump())
+		HRESULT ReflectionStripResult = RemoveReflectionData(DxcDllHelper, OutDxilBlob);
+		if (FAILED(ReflectionStripResult))
 		{
-			// Dump dxil (.d3dasm -> .dxil)
-			FString DxilFile = Arguments.GetDumpDisassemblyFilename().LeftChop(7) + TEXT("_norefl.dxil");
-			SaveDxcBlobToFile(OutDxilBlob, DxilFile);
+			return ReflectionStripResult;
 		}
-
-		GBreakpointDXC++;
-	}
-	else
-	{
-		GBreakpointDXC++;
 	}
 
 	CompileResult->GetErrorBuffer(OutErrorBlob.GetInitReference());
@@ -920,16 +837,25 @@ bool CompileAndProcessD3DShaderDXC(
 
 	Output.AddStatistic(UE::ShaderCompilerCommon::kPlatformHashStatName, BytesToHex(ShaderHash.HashDigest, sizeof(ShaderHash.HashDigest)), FGenericShaderStat::EFlags::Hidden);
 	
-	// Populate the platform-specific debug data with the PDB name, if available.
-	bool bWriteDebugData = Input.Environment.CompilerFlags.Contains(CFLAG_GenerateSymbolsInfo);
-	if (bWriteDebugData && !PdbName.IsEmpty())
+	// Populate the platform-specific debug data with the PDB name and/or data, if available and requested.
+	bool bWriteSymbolsInfo = Input.Environment.CompilerFlags.Contains(CFLAG_GenerateSymbolsInfo);
+	bool bWriteSymbols = Input.Environment.CompilerFlags.Contains(CFLAG_GenerateSymbols);
+	if ((bWriteSymbols || bWriteSymbolsInfo) && !PdbName.IsEmpty())
 	{
-		FD3DSM6ShaderDebugData DebugData;
-		DebugData.Name = PdbName;
+		check(!bWriteSymbols || (PdbBlob.IsValid() && (PdbBlob->GetBufferSize() > 0)));
+		FD3DShaderDebugData DebugData;
+		FD3DShaderDebugData::FFile& PdbFile = DebugData.Files.AddDefaulted_GetRef();
+		PdbFile.Name = PdbName;
 
-		// We don't export the PDB contents here because it would result in duplicate data,
-		// as we use embedded PDBs. Once we are able to use external PDBs, the PDB contents
-		// can be exported too.
+		if (bWriteSymbols)
+		{
+			PdbFile.Contents = MakeArrayViewFromBlob(PdbBlob);
+
+			// also export the .dxil file alongside the .pdb if symbols are on
+			FD3DShaderDebugData::FFile& DxilFile = DebugData.Files.AddDefaulted_GetRef();
+			DxilFile.Name = FPaths::ChangeExtension(PdbName, TEXT(".dxil"));
+			DxilFile.Contents = MakeArrayViewFromBlob(ShaderBlob);
+		}
 		
 		FMemoryWriter Ar(Output.ShaderCode.GetSymbolWriteAccess());
 		Ar << DebugData;
diff --git a/Engine/Source/Developer/Windows/ShaderFormatD3D/Private/ShaderFormatD3D.cpp b/Engine/Source/Developer/Windows/ShaderFormatD3D/Private/ShaderFormatD3D.cpp
--- a/Engine/Source/Developer/Windows/ShaderFormatD3D/Private/ShaderFormatD3D.cpp
+++ b/Engine/Source/Developer/Windows/ShaderFormatD3D/Private/ShaderFormatD3D.cpp
@@ -17,7 +17,7 @@ static FName NAME_PCD3D_SM6(TEXT("PCD3D_SM6"));
 static FName NAME_PCD3D_SM5(TEXT("PCD3D_SM5"));
 static FName NAME_PCD3D_ES3_1(TEXT("PCD3D_ES31"));
 
-static const FGuid UE_SHADER_PCD3D_SHARED_VER = FGuid("B0250A32-89C0-4F5D-9E5A-6DE6E2041829");
+static const FGuid UE_SHADER_PCD3D_SHARED_VER = FGuid("9D405D9B-44C0-4747-8E88-1EDB7413F5C3");
 static const FGuid UE_SHADER_PCD3D_SM6_VER    = FGuid("FC317769-CFF7-4007-A045-C1D4CD315BD4");
 static const FGuid UE_SHADER_PCD3D_SM5_VER    = FGuid("5B377D13-C70F-40C5-80C5-C9B228783469");
 static const FGuid UE_SHADER_PCD3D_ES3_1_VER  = FGuid("952939D9-1156-4347-97E9-9FFEA1A9FE14");
@@ -27,8 +27,9 @@ class FShaderFormatD3D : public UE::ShaderCompilerCommon::FBaseShaderFormat
 	uint32 DxcVersionHash = 0;
 
 #if WITH_ENGINE
-	// TODO: Support for D3D_SM5 and D3D_ES31
 	mutable FShaderSymbolExport ShaderSymbolExportSM6{ NAME_PCD3D_SM6 };
+	mutable FShaderSymbolExport ShaderSymbolExportSM5{ NAME_PCD3D_SM5 };
+	mutable FShaderSymbolExport ShaderSymbolExportES31{ NAME_PCD3D_ES3_1 };
 #endif
 
 public:
@@ -172,15 +173,32 @@ public:
 	{
 		if (Format == NAME_PCD3D_SM6)
 		{
-			ShaderSymbolExportSM6.NotifyShaderCompiled<FD3DSM6ShaderDebugData>(SymbolData, DebugInfo);
+			ShaderSymbolExportSM6.NotifyShaderCompiled<FD3DShaderDebugData>(SymbolData, DebugInfo);
+		}
+		else if (Format == NAME_PCD3D_SM5)
+		{
+			ShaderSymbolExportSM5.NotifyShaderCompiled<FD3DShaderDebugData>(SymbolData, DebugInfo);
+		}
+		else if (Format == NAME_PCD3D_ES3_1)
+		{
+			ShaderSymbolExportES31.NotifyShaderCompiled<FD3DShaderDebugData>(SymbolData, DebugInfo);
 		}
 	}
+
 	virtual void NotifyShaderCompilersShutdown(FName Format) const override
 	{
 		if (Format == NAME_PCD3D_SM6)
 		{
 			ShaderSymbolExportSM6.NotifyShaderCompilersShutdown();
 		}
+		else if (Format == NAME_PCD3D_SM5)
+		{
+			ShaderSymbolExportSM5.NotifyShaderCompilersShutdown();
+		}
+		else if (Format == NAME_PCD3D_ES3_1)
+		{
+			ShaderSymbolExportES31.NotifyShaderCompilersShutdown();
+		}
 	}
 #endif
 
diff --git a/Engine/Source/Developer/Windows/ShaderFormatD3D/Private/ShaderFormatD3D.h b/Engine/Source/Developer/Windows/ShaderFormatD3D/Private/ShaderFormatD3D.h
--- a/Engine/Source/Developer/Windows/ShaderFormatD3D/Private/ShaderFormatD3D.h
+++ b/Engine/Source/Developer/Windows/ShaderFormatD3D/Private/ShaderFormatD3D.h
@@ -60,31 +60,47 @@ bool CompileAndProcessD3DShaderDXC(
 struct FD3DShaderCompileData;
 bool ValidateResourceCounts(const FD3DShaderCompileData& CompiledData, TArray<FString>& OutFilteredErrors);
 
-struct FD3DSM6ShaderDebugData
+template <typename TBlobType>
+TConstArrayView<uint8> MakeArrayViewFromBlob(const TRefCountPtr<TBlobType> Blob)
 {
-	FString Name;
-	FString DebugInfo;
-	TArray<uint8> Contents;
-
-	inline friend FArchive& operator<<(FArchive& Ar, FD3DSM6ShaderDebugData& DebugData)
-	{
-		Ar << DebugData.Name;
-		Ar << DebugData.Contents;
-		return Ar;
-	}
+	return MakeArrayView<const uint8>(reinterpret_cast<const uint8*>(Blob->GetBufferPointer()), Blob->GetBufferSize());
+}
 
-	inline TConstArrayView<uint8> GetContents() const
+struct FD3DShaderDebugData
+{
+	struct FFile
 	{
-		return TConstArrayView<uint8>(Contents);
-	}
-
-	inline FString GetFilename() const
+		FString Name;
+		TArray<uint8> Contents;
+
+		inline friend FArchive& operator<<(FArchive& Ar, FFile& DebugData)
+		{
+			Ar << DebugData.Name;
+			Ar << DebugData.Contents;
+			return Ar;
+		}
+
+		inline FString GetFilename() const
+		{
+			return Name;
+		}
+
+		inline TConstArrayView<uint8> GetContents() const
+		{
+			return TConstArrayView<uint8>(Contents);
+		}
+	};
+
+	TArray<FFile> Files;
+
+	inline friend FArchive& operator<<(FArchive& Ar, FD3DShaderDebugData& DebugData)
 	{
-		return Name;
+		Ar << DebugData.Files;
+		return Ar;
 	}
 
-	TConstArrayView<FD3DSM6ShaderDebugData> GetAllSymbolData() const
+	TConstArrayView<FFile> GetAllSymbolData() const
 	{
-		return TConstArrayView<FD3DSM6ShaderDebugData>(this, 1);
+		return TConstArrayView<FFile>(Files);
 	}
 };
\ No newline at end of file
