Fix missing shader's diagnostic data when compiling with shader worker.

diff --git a/Engine/Source/Developer/Windows/ShaderFormatD3D/Private/ShaderFormatD3D.cpp b/Engine/Source/Developer/Windows/ShaderFormatD3D/Private/ShaderFormatD3D.cpp
--- a/Engine/Source/Developer/Windows/ShaderFormatD3D/Private/ShaderFormatD3D.cpp
+++ b/Engine/Source/Developer/Windows/ShaderFormatD3D/Private/ShaderFormatD3D.cpp
@@ -208,9 +208,10 @@ public:
 		FShaderCompilerOutput& Output,
 		const FString& WorkingDirectory) const
 	{
-		CompileD3DShader(Input, PreprocessOutput, Output, WorkingDirectory, DetermineShaderModel(Input));
-
+		// Fill in ShaderDiagnosticDatas into the output result before calling CompileD3DShader as it will use and serialize its content
 		Output.ShaderDiagnosticDatas = PreprocessOutput.GetDiagnosticDatas();
+
+		CompileD3DShader(Input, PreprocessOutput, Output, WorkingDirectory, DetermineShaderModel(Input));
 	}
 
 	static void AddShaderTargetDefines(FShaderCompilerInput& Input, uint32 ShaderTargetMajor, uint32 ShaderTargetMinor)
