Notify user to restart the editor after changing callback buffer size or number of buffers.

diff --git a/Engine/Source/Developer/Windows/WindowsTargetPlatformSettings/Classes/WindowsTargetSettings.h b/Engine/Source/Developer/Windows/WindowsTargetPlatformSettings/Classes/WindowsTargetSettings.h
--- a/Engine/Source/Developer/Windows/WindowsTargetPlatformSettings/Classes/WindowsTargetSettings.h
+++ b/Engine/Source/Developer/Windows/WindowsTargetPlatformSettings/Classes/WindowsTargetSettings.h
@@ -73,11 +73,11 @@ public:
 	int32 AudioSampleRate;
 
 	/** The amount of audio to compute each callback block. Lower values decrease latency but may increase CPU cost. */
-	UPROPERTY(config, EditAnywhere, Category = "Audio", meta = (ClampMin = "64", ClampMax = "4096", DisplayName = "Callback Buffer Size"))
+	UPROPERTY(config, EditAnywhere, Category = "Audio", meta = (ClampMin = "64", ClampMax = "4096", DisplayName = "Callback Buffer Size", ConfigRestartRequired = true))
 	int32 AudioCallbackBufferFrameSize;
 
 	/** The number of buffers to keep enqueued. More buffers increases latency, but can compensate for variable compute availability in audio callbacks on some platforms. */
-	UPROPERTY(config, EditAnywhere, Category = "Audio", meta = (ClampMin = "1", UIMin = "1", DisplayName = "Number of Buffers To Enqueue"))
+	UPROPERTY(config, EditAnywhere, Category = "Audio", meta = (ClampMin = "1", UIMin = "1", DisplayName = "Number of Buffers To Enqueue", ConfigRestartRequired = true))
 	int32 AudioNumBuffersToEnqueue;
 
 	/** The max number of channels (voices) to limit for this platform. The max channels used will be the minimum of this value and the global audio quality settings. A value of 0 will not apply a platform channel count max. */
