Compile fix for when LiveCoding is disabled

(e.g. compiling the editor for arm64 / arm64ec with -allmodules)

#rnx

diff --git a/Engine/Source/Developer/Windows/LiveCoding/Private/LiveCodingModule.cpp b/Engine/Source/Developer/Windows/LiveCoding/Private/LiveCodingModule.cpp
--- a/Engine/Source/Developer/Windows/LiveCoding/Private/LiveCodingModule.cpp
+++ b/Engine/Source/Developer/Windows/LiveCoding/Private/LiveCodingModule.cpp
@@ -1525,11 +1525,11 @@ void FLiveCodingModule::BeginReload()
 			GLiveCodingModule->bHasReinstancingOccurred = false;
 			GLiveCodingModule->bHasPatchBeenLoaded = false;
 			GPostCompileResult = commands::PostCompileResult::Success;
-#if WITH_EDITOR
+#if WITH_EDITOR && WITH_LIVE_CODING
 			GLiveCodingModule->Reload.Reset(new FReload(EActiveReloadType::LiveCoding, TEXT("LIVECODING"), *GLog));
 			GLiveCodingModule->Reload->SetEnableReinstancing(GLiveCodingModule->IsReinstancingEnabled());
 			GLiveCodingModule->Reload->SetSendReloadCompleteNotification(false);
-#else
+#elif WITH_LIVE_CODING
 			GLiveCodingModule->Reload.Reset(new FNullReload(*GLiveCodingModule));
 #endif
 		}
