Disable LiveCoding on Windows Arm64 - it's not supported

#rnx

diff --git a/Engine/Source/Developer/Windows/LiveCoding/LiveCoding.Build.cs b/Engine/Source/Developer/Windows/LiveCoding/LiveCoding.Build.cs
--- a/Engine/Source/Developer/Windows/LiveCoding/LiveCoding.Build.cs
+++ b/Engine/Source/Developer/Windows/LiveCoding/LiveCoding.Build.cs
@@ -35,7 +35,11 @@ public class LiveCoding : ModuleRules
             PrivateDefinitions.Add("USE_DEBUG_LIVE_CODING_CONSOLE=0");
         }
 
-		if (UseLiveCoding2)
+		if (!Target.Architecture.bIsX64)
+		{
+			PrivateDefinitions.Add("LC_VERSION=0");
+		}
+		else if (UseLiveCoding2)
 		{
 			PrivateDefinitions.Add("LC_VERSION=2");
 			PrivateIncludePaths.Add("ThirdParty/LivePlusPlus");
