Cache Windows targeted shader formats.
This avoids frequent GConfig file accesses/flushes.

diff --git a/Engine/Source/Developer/Windows/WindowsTargetPlatformSettings/Public/GenericWindowsTargetPlatformSettings.h b/Engine/Source/Developer/Windows/WindowsTargetPlatformSettings/Public/GenericWindowsTargetPlatformSettings.h
--- a/Engine/Source/Developer/Windows/WindowsTargetPlatformSettings/Public/GenericWindowsTargetPlatformSettings.h
+++ b/Engine/Source/Developer/Windows/WindowsTargetPlatformSettings/Public/GenericWindowsTargetPlatformSettings.h
@@ -129,14 +129,22 @@ public:
 		}
 	}
 
-	void GetAllTargetedShaderFormatsInternal(TArrayView<TCHAR const*> RelevantSettings, TArray<FName>& OutFormats) const
+	TArray<FName> GetAllTargetedShaderFormatsInternal(TConstArrayView<const TCHAR*> RelevantSettings) const
 	{
+		// no shaders needed for dedicated server target
+		if (TProperties::IsServerOnly())
+		{
+			return TArray<FName>();
+		}
+
+		FConfigCacheIni* ConfigSystem = this->GetConfigSystem();
+
 		TArray<FString> TargetedShaderFormats;
 
 		for (const TCHAR* Name : RelevantSettings)
 		{
 			TArray<FString> NewTargetedShaderFormats;
-			GConfig->GetArray(TEXT("/Script/WindowsTargetPlatform.WindowsTargetSettings"), Name, NewTargetedShaderFormats, GEngineIni);
+			ConfigSystem->GetArray(TEXT("/Script/WindowsTargetPlatform.WindowsTargetSettings"), Name, NewTargetedShaderFormats, GEngineIni);
 
 			for (const FString& NewShaderFormat : NewTargetedShaderFormats)
 			{
@@ -148,26 +156,25 @@ public:
 		TArray<FName> PossibleShaderFormats;
 		GetAllPossibleShaderFormats(PossibleShaderFormats);
 
-		for (int32 ShaderFormatIdx = TargetedShaderFormats.Num() - 1; ShaderFormatIdx >= 0; ShaderFormatIdx--)
+		TArray<FName> OutFormats;
+		for (const FString& ShaderFormat : TargetedShaderFormats)
 		{
-			FString ShaderFormat = TargetedShaderFormats[ShaderFormatIdx];
-			if (PossibleShaderFormats.Contains(FName(*ShaderFormat)) == false)
+			const FName ShaderFormatName(*ShaderFormat, FNAME_Find);
+
+			if (!ShaderFormatName.IsNone() && PossibleShaderFormats.Contains(ShaderFormatName))
 			{
-				TargetedShaderFormats.RemoveAt(ShaderFormatIdx);
+				OutFormats.Emplace(ShaderFormatName);
 			}
 		}
 
-		for(const FString& ShaderFormat : TargetedShaderFormats)
-		{
-			OutFormats.AddUnique(FName(*ShaderFormat));
-		}
+		return OutFormats;
 	}
 
 	virtual void GetAllTargetedShaderFormats( TArray<FName>& OutFormats ) const override 
 	{
 		// Get the Target RHIs for this platform, we do not always want all those that are supported. (reload in case user changed in the editor)
 
-		TCHAR const* RelevantSettings[] = 
+		const TCHAR* const RelevantSettings[] = 
 		{
 			TEXT("TargetedRHIs"),
 			TEXT("D3D12TargetedShaderFormats"),
@@ -175,22 +182,25 @@ public:
 			TEXT("VulkanTargetedShaderFormats")
 		};
 
-		GetAllTargetedShaderFormatsInternal(RelevantSettings, OutFormats);
+		static const TArray<FName> CachedFormats = GetAllTargetedShaderFormatsInternal(RelevantSettings);
+		OutFormats = CachedFormats;
 	}
 
 	virtual void GetRayTracingShaderFormats( TArray<FName>& OutFormats ) const override 
 	{
 		if (UsesRayTracing())
 		{
-			TCHAR const* RelevantSettings[] = 
+			const TCHAR* const RelevantSettings[] = 
 			{
 				TEXT("VulkanTargetedShaderFormats")
 			};
 
-			GetAllTargetedShaderFormatsInternal(RelevantSettings, OutFormats);
+			static const TArray<FName> CachedFormats = GetAllTargetedShaderFormatsInternal(RelevantSettings);
+			OutFormats = CachedFormats;
 
 			// We always support ray tracing shaders when cooking for D3D12 SM6, however we may skip them for SM5 based on project settings.
-			OutFormats.AddUnique(FName(TEXT("PCD3D_SM6")));
+			static FName NAME_PCD3D_SM6(TEXT("PCD3D_SM6"));
+			OutFormats.Emplace(NAME_PCD3D_SM6);
 		}
 	}
 
