[LiveCoding]
* Changed so livecoding2 is using RegisterRawInputDevices to track ctrl-alt-F11 instead of keyboard hook. Tried RegisterHotKey but does not work with multiple processes.. this seems to work fine

diff --git a/Engine/Source/Developer/Windows/LiveCoding/Private/LiveCodingModule2.cpp b/Engine/Source/Developer/Windows/LiveCoding/Private/LiveCodingModule2.cpp
--- a/Engine/Source/Developer/Windows/LiveCoding/Private/LiveCodingModule2.cpp
+++ b/Engine/Source/Developer/Windows/LiveCoding/Private/LiveCodingModule2.cpp
@@ -98,24 +98,6 @@ uint32 StartBroker()
 	return 0;
 }
 
-LRESULT CALLBACK KeyboardHook(int nCode, WPARAM wParam, LPARAM lParam)
-{
-    if (nCode == HC_ACTION && wParam == WM_KEYDOWN)
-	{
-        KBDLLHOOKSTRUCT* p = (KBDLLHOOKSTRUCT*)lParam;
-        if (p->vkCode == VK_F11)
-		{
-            bool ctrlDown = (GetAsyncKeyState(VK_CONTROL) & 0x8000) != 0;
-            bool altDown  = (GetAsyncKeyState(VK_MENU)    & 0x8000) != 0;
-            if (ctrlDown && altDown)
-			{
-				PostThreadMessage(GetCurrentThreadId(), WM_USER, 0, 0);
-			}
-        }
-    }
-    return CallNextHookEx(NULL, nCode, wParam, lParam);
-}
-
 DWORD WINAPI HookThread(void* ModulePtr)
 {
 	FLiveCodingModule& Module = *(FLiveCodingModule*)ModulePtr;
@@ -127,20 +109,34 @@ DWORD WINAPI HookThread(void* ModulePtr)
 	}
 	ON_SCOPE_EXIT{ CloseHandle(HotKeyMutex); HotKeyMutex = NULL; };
 
-    HHOOK KeyboardHookHandle = SetWindowsHookEx(WH_KEYBOARD_LL, KeyboardHook, NULL, 0);
-	if (!KeyboardHookHandle)
+	WNDCLASSA wc = { };
+	wc.lpfnWndProc   = DefWindowProc;
+	wc.hInstance     = hInstance;
+	wc.lpszClassName = "LiveCodingHiddenWindow";
+	RegisterClassA(&wc);
+	ON_SCOPE_EXIT{ UnregisterClassA("LiveCodingHiddenWindow", hInstance); };
+
+	HWND hwnd = CreateWindowExA(0, "LiveCodingHiddenWindow", "RawInputListener", 0, 0, 0, 0, 0, HWND_MESSAGE, NULL, NULL, NULL);
+	if (!hwnd)
 	{
-		UE_LOG(LogLiveCoding, Error, TEXT("SetWindowsHookEx failed when setting up listener for hot key"));
+		UE_LOG(LogLiveCoding, Error, TEXT("CreateWindowEx failed when setting up listener for hot key"));
 		return 0;
 	}
-	ON_SCOPE_EXIT{ UnhookWindowsHookEx(KeyboardHookHandle); };
+	ON_SCOPE_EXIT{ DestroyWindow(hwnd); };
 
-    MSG msg;
-    while (GetMessageW(&msg, NULL, 0, 0))
+	RAWINPUTDEVICE rid { 0x01, 0x06, RIDEV_INPUTSINK, hwnd };
+	if (!RegisterRawInputDevices(&rid, 1, sizeof(rid)))
+	{
+		UE_LOG(LogLiveCoding, Error, TEXT("RegisterRawInputDevices failed when setting up listener for hot key"));
+		return 1;
+	}
+	
+	MSG msg;
+	while (GetMessageW(&msg, NULL, 0, 0))
 	{
-        TranslateMessage(&msg);
-        DispatchMessage(&msg);
-		if (msg.message != WM_USER)
+		TranslateMessage(&msg);
+		DispatchMessage(&msg);
+		if (msg.message != WM_INPUT || !((GetAsyncKeyState(VK_CONTROL) & 0x8000) && (GetAsyncKeyState(VK_MENU) & 0x8000) && (GetAsyncKeyState(VK_F11) & 0x8000)))
 		{
 			continue;
 		}
