Fixes for inline raytracing and non uniform resource indexing

diff --git a/Engine/Shaders/Private/RayTracing/TraceRayInlineCommon.ush b/Engine/Shaders/Private/RayTracing/TraceRayInlineCommon.ush
--- a/Engine/Shaders/Private/RayTracing/TraceRayInlineCommon.ush
+++ b/Engine/Shaders/Private/RayTracing/TraceRayInlineCommon.ush
@@ -3,7 +3,7 @@
 #pragma once
 
 // Update this GUID to invalidate inline RayTracing shaders and force a shader recompilation.
-#pragma message("UESHADERMETADATA_VERSION 06A6E5F1-49BF-4FF3-9D90-4B501E7E5FC4")
+#pragma message("UESHADERMETADATA_VERSION D26C9B14-CE7A-41CF-8BC6-ED25A12F5AD0")
 
 struct FTraceRayInlineTraversalStatistics
 {
diff --git a/Engine/Shaders/Private/RayTracing/TraceRayInlineD3D12.ush b/Engine/Shaders/Private/RayTracing/TraceRayInlineD3D12.ush
--- a/Engine/Shaders/Private/RayTracing/TraceRayInlineD3D12.ush
+++ b/Engine/Shaders/Private/RayTracing/TraceRayInlineD3D12.ush
@@ -61,8 +61,8 @@ FTriangleBaseAttributes LoadInlineRayTracingTriangleAttributes(
 	const uint IndexBufferHandle = Params.BindlessHitGroupSystemIndexBuffer;
 	const uint VertexBufferHandle = Params.BindlessHitGroupSystemVertexBuffer;
 
-	ByteAddressBuffer IndexBuffer = ResourceDescriptorHeap[IndexBufferHandle];
-	ByteAddressBuffer VertexBuffer = ResourceDescriptorHeap[VertexBufferHandle];	
+	ByteAddressBuffer IndexBuffer = ResourceDescriptorHeap[NonUniformResourceIndex(IndexBufferHandle)];
+	ByteAddressBuffer VertexBuffer = ResourceDescriptorHeap[NonUniformResourceIndex(VertexBufferHandle)];	
 
 	return LoadTriangleBaseAttributes(IndexBuffer, IndexBufferOffsetInBytes, IndexBufferStride, VertexBuffer, VertexBufferOffsetInBytes, VertexStride, PrimitiveIndex);
 }
diff --git a/Engine/Shaders/Public/Platform/D3D/D3DCommon.ush b/Engine/Shaders/Public/Platform/D3D/D3DCommon.ush
--- a/Engine/Shaders/Public/Platform/D3D/D3DCommon.ush
+++ b/Engine/Shaders/Public/Platform/D3D/D3DCommon.ush
@@ -144,7 +144,7 @@ void UEReportAssertWithPayload(uint AssertID, uint4 Payload)
 #endif
 
 #if PLATFORM_SUPPORTS_BINDLESS
-	#if (RAYCALLABLESHADER || RAYHITGROUPSHADER || RAYGENSHADER || RAYMISSSHADER)
+	#if (RAYCALLABLESHADER || RAYHITGROUPSHADER || RAYGENSHADER || RAYMISSSHADER || PLATFORM_SUPPORTS_INLINE_RAY_TRACING)
 		// Bindless index can be divergent in hit, miss and callable shaders
 		#define UE_BINDLESS_INDEX(Index) NonUniformResourceIndex(Index)
 	#else
diff --git a/Engine/Source/Developer/Windows/ShaderFormatD3D/Private/D3DShaderCompiler.cpp b/Engine/Source/Developer/Windows/ShaderFormatD3D/Private/D3DShaderCompiler.cpp
--- a/Engine/Source/Developer/Windows/ShaderFormatD3D/Private/D3DShaderCompiler.cpp
+++ b/Engine/Source/Developer/Windows/ShaderFormatD3D/Private/D3DShaderCompiler.cpp
@@ -1391,7 +1391,7 @@ struct FD3DShaderParameterParserPlatformConfiguration : public FShaderParameterP
 {
 	FD3DShaderParameterParserPlatformConfiguration(const FShaderCompilerInput& Input)
 		: FShaderParameterParser::FPlatformConfiguration(TEXTVIEW("cbuffer"), EShaderParameterParserConfigurationFlags::UseStableConstantBuffer|EShaderParameterParserConfigurationFlags::SupportsBindless)
-		, bIsRayTracingShader(Input.IsRayTracingShader())
+		, bIsRayTracingShader(Input.IsRayTracingShader() || Input.Environment.CompilerFlags.Contains(CFLAG_InlineRayTracing))
 		, HitGroupSystemIndexBufferName(FShaderParameterParser::kBindlessSRVPrefix + FString(TEXT("HitGroupSystemIndexBuffer")))
 		, HitGroupSystemVertexBufferName(FShaderParameterParser::kBindlessSRVPrefix + FString(TEXT("HitGroupSystemVertexBuffer")))
 	{
