[shaders] don't strip reflection data from DXIL container if CFLAG_ExtraShaderData is set

diff --git a/Engine/Source/Developer/Windows/ShaderFormatD3D/Private/D3DShaderCompilerDXC.cpp b/Engine/Source/Developer/Windows/ShaderFormatD3D/Private/D3DShaderCompilerDXC.cpp
--- a/Engine/Source/Developer/Windows/ShaderFormatD3D/Private/D3DShaderCompilerDXC.cpp
+++ b/Engine/Source/Developer/Windows/ShaderFormatD3D/Private/D3DShaderCompilerDXC.cpp
@@ -143,7 +143,7 @@ protected:
 	FString DumpDisasmFilename;
 	FString BatchBaseFilename;
 	FString DumpDebugInfoPath;
-	bool bKeepEmbeddedPDB = false;
+	bool bStripReflection = true;
 	bool bDump = false;
 
 	TArray<FString> ExtraArguments;
@@ -274,6 +274,9 @@ public:
 
 		// @lh-todo: This fixes a loop unrolling issue that showed up in DOFGatherKernel with cs_6_6 with the latest DXC revision
 		ExtraArguments.Add(TEXT("-disable-lifetime-markers"));
+
+		// temporary change until nsight supports reconstructing reflection data from PDB; don't strip reflection if CFLAG_ExtraShaderData is specified
+		bStripReflection = !Input.Environment.CompilerFlags.Contains(CFLAG_ExtraShaderData);
 	}
 
 	FString GetDumpDebugInfoPath() const
@@ -286,6 +289,11 @@ public:
 		return bDump;
 	}
 
+	bool ShouldStripReflectionData() const
+	{
+		return bStripReflection;
+	}
+
 	FString GetEntryPointName() const
 	{
 		return Exports.Len() > 0 ? FString(TEXT("")) : EntryPoint;
@@ -679,10 +687,13 @@ static HRESULT D3DCompileToDxil(const char* SourceText, const FDxcArguments& Arg
 			DisassembleAndSave(Compiler, OutDxilBlob, DisasmFilename);
 		}
 
-		HRESULT ReflectionStripResult = RemoveReflectionData(DxcDllHelper, OutDxilBlob);
-		if (FAILED(ReflectionStripResult))
+		if (Arguments.ShouldStripReflectionData())
 		{
-			return ReflectionStripResult;
+			HRESULT ReflectionStripResult = RemoveReflectionData(DxcDllHelper, OutDxilBlob);
+			if (FAILED(ReflectionStripResult))
+			{
+				return ReflectionStripResult;
+			}
 		}
 	}
 
diff --git a/Engine/Source/Developer/Windows/ShaderFormatD3D/Private/ShaderFormatD3D.cpp b/Engine/Source/Developer/Windows/ShaderFormatD3D/Private/ShaderFormatD3D.cpp
--- a/Engine/Source/Developer/Windows/ShaderFormatD3D/Private/ShaderFormatD3D.cpp
+++ b/Engine/Source/Developer/Windows/ShaderFormatD3D/Private/ShaderFormatD3D.cpp
@@ -18,8 +18,8 @@ static FName NAME_PCD3D_SM5(TEXT("PCD3D_SM5"));
 static FName NAME_PCD3D_ES3_1(TEXT("PCD3D_ES31"));
 
 static const FGuid UE_SHADER_PCD3D_SHARED_VER = FGuid("9D405D9B-44C0-4747-8E88-1EDB7413F5C3");
-static const FGuid UE_SHADER_PCD3D_SM6_VER    = FGuid("FC317769-CFF7-4007-A045-C1D4CD315BD4");
-static const FGuid UE_SHADER_PCD3D_SM5_VER    = FGuid("5B377D13-C70F-40C5-80C5-C9B228783469");
+static const FGuid UE_SHADER_PCD3D_SM6_VER    = FGuid("7F9AE1F2-9692-46B4-875A-28C2595E84AA");
+static const FGuid UE_SHADER_PCD3D_SM5_VER    = FGuid("F30C5AFA-77F7-48DE-8CC1-A075F179443C");
 static const FGuid UE_SHADER_PCD3D_ES3_1_VER  = FGuid("952939D9-1156-4347-97E9-9FFEA1A9FE14");
 
 class FShaderFormatD3D : public UE::ShaderCompilerCommon::FBaseShaderFormat 
