[Shaders] Move HlslParser into its own module to allow sharing it between Engine and ShaderCompilerCommon modules.

We want to use HlslParser in the material translator for future changes which is part of the Engine module. Due to its dependency graph, this means HlslParser has to be moved into its own module.
All existing modules that need to access the public functions such as CrossCompiler::Parser::Parse() need to explicitly import the new module since UE does not provide a transitive linking model.
This change also adds a simple function ForEachHlslIdentifier() to the HlslParser interface to scan for identifiers in a code string.

#rnx

diff --git a/Engine/Source/Developer/Apple/MetalShaderFormat/MetalShaderFormat.Build.cs b/Engine/Source/Developer/Apple/MetalShaderFormat/MetalShaderFormat.Build.cs
--- a/Engine/Source/Developer/Apple/MetalShaderFormat/MetalShaderFormat.Build.cs
+++ b/Engine/Source/Developer/Apple/MetalShaderFormat/MetalShaderFormat.Build.cs
@@ -15,6 +15,7 @@ public class MetalShaderFormat : ModuleRules
 				"Core",
 				"RenderCore",
 				"ShaderCompilerCommon",
+				"HlslParser",
 				"ShaderPreprocessor",
 				"FileUtilities",
 				"RHI"
diff --git a/Engine/Source/Developer/ShaderCompilerCommon/Public/ShaderCompilerCommon.h b/Engine/Source/Developer/ShaderCompilerCommon/Public/ShaderCompilerCommon.h
--- a/Engine/Source/Developer/ShaderCompilerCommon/Public/ShaderCompilerCommon.h
+++ b/Engine/Source/Developer/ShaderCompilerCommon/Public/ShaderCompilerCommon.h
@@ -384,7 +384,7 @@ extern SHADERCOMPILERCOMMON_API void ParseHLSLTypeName(const TCHAR* SearchString
 extern SHADERCOMPILERCOMMON_API FStringView FindNextHLSLDefinitionOfType(FStringView Typename, FStringView StartPos);
 
 // Structure to hold forward declarations for a specific scope/namespace chain for the HlslParser
-struct FScopedDeclarations
+struct UE_DEPRECATED(5.7, "FScopedDeclarations has been moved into internal module HlslParser and is no longer supported as public interface.") FScopedDeclarations
 {
 	FScopedDeclarations(TConstArrayView<FStringView> InScope, TConstArrayView<FStringView> InSymbols)
 		: Scope(InScope)
@@ -395,43 +395,31 @@ struct FScopedDeclarations
 	TConstArrayView<FStringView> Symbols;
 };
 
-extern SHADERCOMPILERCOMMON_API bool RemoveUnusedOutputs(
-	FString& InOutSourceCode,
-	TConstArrayView<FStringView> InUsedOutputs,
-	TConstArrayView<FStringView> InExceptions,
-	TConstArrayView<FScopedDeclarations> InScopedDeclarations,
-	FString& InOutEntryPoint,
-	TArray<FString>& OutErrors
-);
-
-extern SHADERCOMPILERCOMMON_API bool RemoveUnusedOutputs(FString& InOutSourceCode, const TArray<FString>& InUsedOutputs, const TArray<FString>& InExceptions, FString& InOutEntryPoint, TArray<FString>& OutErrors);
-
-extern SHADERCOMPILERCOMMON_API bool RemoveUnusedInputs(
-	FString& InOutSourceCode,
-	TConstArrayView<FStringView> InUsedInputs,
-	TConstArrayView<FScopedDeclarations> InScopedDeclarations,
-	FString& InOutEntryPoint,
-	TArray<FString>& OutErrors
-);
+UE_DEPRECATED(5.7, "RemoveUnusedOutputs() has been moved into internal module HlslParser and is no longer supported as public interface.")
+inline bool RemoveUnusedOutputs(...)
+{
+	return false; // dummy
+}
 
-extern SHADERCOMPILERCOMMON_API bool RemoveUnusedInputs(FString& InOutSourceCode, const TArray<FString>& InUsedInputs, FString& InOutEntryPoint, TArray<FString>& OutErrors);
+UE_DEPRECATED(5.7, "RemoveUnusedInputs() has been moved into internal module HlslParser and is no longer supported as public interface.")
+inline bool RemoveUnusedInputs(...)
+{
+	return false; // dummy
+}
 
 // Shader input/output parameter storage classes. Naming adopted from SPIR-V nomenclature.
-enum class EShaderParameterStorageClass
+enum class UE_DEPRECATED(5.7, "EShaderParameterStorageClass has been moved into internal module HlslParser and is no longer supported as public interface.") EShaderParameterStorageClass
 {
 	Input,
 	Output,
 };
 
 // Returns the semantic names of all individual entry point parameters (i.e. all structure fields are inlined)
-extern SHADERCOMPILERCOMMON_API bool FindEntryPointParameters(
-	const FString& InSourceCode,
-	const FString& InEntryPoint,
-	EShaderParameterStorageClass ParameterStorageClass,
-	TConstArrayView<FScopedDeclarations> InScopedDeclarations,
-	TArray<FString>& OutParameterSemantics,
-	TArray<FString>& OutErrors
-);
+UE_DEPRECATED(5.7, "FindEntryPointParameters() has been moved into internal module HlslParser and is no longer supported as public interface.")
+inline bool FindEntryPointParameters(...)
+{
+	return false; //dummy
+}
 
 extern SHADERCOMPILERCOMMON_API bool ConvertFromFP32ToFP16(FString& InOutSourceCode, TArray<FString>& OutErrors);
 
diff --git a/Engine/Source/Developer/ShaderCompilerCommon/ShaderCompilerCommon.Build.cs b/Engine/Source/Developer/ShaderCompilerCommon/ShaderCompilerCommon.Build.cs
--- a/Engine/Source/Developer/ShaderCompilerCommon/ShaderCompilerCommon.Build.cs
+++ b/Engine/Source/Developer/ShaderCompilerCommon/ShaderCompilerCommon.Build.cs
@@ -15,6 +15,11 @@ public class ShaderCompilerCommon : ModuleRules
 				"TargetPlatform",
 			}
 			);
+		PublicDependencyModuleNames.AddRange(
+			new string[] {
+				"HlslParser",
+			}
+			);
 		if (Target.bCompileAgainstEngine)
 		{
 			PublicDependencyModuleNames.AddRange(
diff --git a/Engine/Source/Developer/Windows/ShaderFormatD3D/Private/D3DShaderCompiler.cpp b/Engine/Source/Developer/Windows/ShaderFormatD3D/Private/D3DShaderCompiler.cpp
--- a/Engine/Source/Developer/Windows/ShaderFormatD3D/Private/D3DShaderCompiler.cpp
+++ b/Engine/Source/Developer/Windows/ShaderFormatD3D/Private/D3DShaderCompiler.cpp
@@ -15,6 +15,7 @@
 #include "ShaderPreprocessTypes.h"
 #include "SpirvCommon.h"
 #include "Algo/Transform.h"
+#include "HlslParserInternal.h"
 
 #define DEBUG_SHADERS 0
 
@@ -497,11 +498,11 @@ static bool PatchHlslWithReorderedIOVariables(
 	FString& HlslSourceString,
 	const FString& OriginalShaderSource,
 	const FString& OriginalEntryPoint,
-	EShaderParameterStorageClass StageVariablesStorageClass,
+	UE::HlslParser::EShaderParameterStorageClass StageVariablesStorageClass,
 	TArray<FShaderCompilerError>& OutErrors)
 {
 	// Find declaration struct for stage variables
-	const FStringView StageVariableDeclarationName = (StageVariablesStorageClass == EShaderParameterStorageClass::Input ? TEXTVIEW("SPIRV_Cross_Input") : TEXTVIEW("SPIRV_Cross_Output"));
+	const FStringView StageVariableDeclarationName = (StageVariablesStorageClass == UE::HlslParser::EShaderParameterStorageClass::Input ? TEXTVIEW("SPIRV_Cross_Input") : TEXTVIEW("SPIRV_Cross_Output"));
 	const int32 StageVariableDeclarationBegin = HlslSourceString.Find(StageVariableDeclarationName, ESearchCase::CaseSensitive);
 	if (StageVariableDeclarationBegin == INDEX_NONE)
 	{
@@ -522,7 +523,7 @@ static bool PatchHlslWithReorderedIOVariables(
 
 	// Parse variable names from SPIR-V input
 	TArray<FString> Variables, ParsingErrors;
-	if (!FindEntryPointParameters(OriginalShaderSource, OriginalEntryPoint, StageVariablesStorageClass, {}, Variables, ParsingErrors))
+	if (!UE::HlslParser::FindEntryPointParameters(OriginalShaderSource, OriginalEntryPoint, StageVariablesStorageClass, {}, Variables, ParsingErrors))
 	{
 		for (FString& Error : ParsingErrors)
 		{
@@ -698,7 +699,7 @@ static void PatchHlslForPrecompilation(
 	if (Frequency == SF_Vertex)
 	{
 		// Ensure order of output variables remains the same as declared in original shader source
-		PatchHlslWithReorderedIOVariables(HlslSourceString, OriginalShaderSource, OriginalEntryPoint, EShaderParameterStorageClass::Output, OutErrors);
+		PatchHlslWithReorderedIOVariables(HlslSourceString, OriginalShaderSource, OriginalEntryPoint, UE::HlslParser::EShaderParameterStorageClass::Output, OutErrors);
 	}
 	else if (Frequency == SF_Pixel)
 	{
@@ -718,7 +719,7 @@ static void PatchHlslForPrecompilation(
 		}
 
 		// Ensure order of input variables remains the same as declared in original shader source
-		PatchHlslWithReorderedIOVariables(HlslSourceString, OriginalShaderSource, OriginalEntryPoint, EShaderParameterStorageClass::Input, OutErrors);
+		PatchHlslWithReorderedIOVariables(HlslSourceString, OriginalShaderSource, OriginalEntryPoint, UE::HlslParser::EShaderParameterStorageClass::Input, OutErrors);
 	}
 
 	// Return new HLSL source
@@ -1506,7 +1507,7 @@ void CompileD3DShader(const FShaderCompilerInput& Input, const FShaderPreprocess
 			TEXTVIEW("SV_CullDistance7"),
 		};
 
-		TArray<FScopedDeclarations> ScopedDeclarations;
+		TArray<UE::HlslParser::FScopedDeclarations> ScopedDeclarations;
 		const FStringView GlobalSymbols[] =
 		{
 			TEXTVIEW("RayDesc"),
@@ -1514,7 +1515,7 @@ void CompileD3DShader(const FShaderCompilerInput& Input, const FShaderPreprocess
 		ScopedDeclarations.Emplace(TConstArrayView<FStringView>(), GlobalSymbols);
 
 		TArray<FString> Errors;
-		if (!RemoveUnusedOutputs(PreprocessedSource, UsedOutputs, Exceptions, ScopedDeclarations, EntryPointName, Errors))
+		if (!UE::HlslParser::RemoveUnusedOutputs(PreprocessedSource, UsedOutputs, Exceptions, ScopedDeclarations, EntryPointName, Errors))
 		{
 			UE_LOG(LogD3DShaderCompiler, Warning, TEXT("Failed to remove unused outputs from shader: %s"), *Input.GenerateShaderName());
 			for (const FString& ErrorReport : Errors)
diff --git a/Engine/Source/Developer/Windows/ShaderFormatD3D/Private/D3DShaderCompiler.inl b/Engine/Source/Developer/Windows/ShaderFormatD3D/Private/D3DShaderCompiler.inl
--- a/Engine/Source/Developer/Windows/ShaderFormatD3D/Private/D3DShaderCompiler.inl
+++ b/Engine/Source/Developer/Windows/ShaderFormatD3D/Private/D3DShaderCompiler.inl
@@ -675,7 +675,7 @@ bool RemoveUnusedInterpolators(
 					TArray<FString> RemoveErrors;
 					FString ModifiedShaderSource = PreprocessedShaderSource;
 					FString ModifiedEntryPointName = Input.EntryPointName;
-					if (RemoveUnusedInputs(ModifiedShaderSource, CompileData.ShaderInputs, ModifiedEntryPointName, RemoveErrors))
+					if (UE::HlslParser::RemoveUnusedInputs(ModifiedShaderSource, CompileData.ShaderInputs, ModifiedEntryPointName, RemoveErrors))
 					{
 						Output = OriginalOutput;
 						
diff --git a/Engine/Source/Developer/Windows/ShaderFormatD3D/Private/D3DShaderCompilerDXC.cpp b/Engine/Source/Developer/Windows/ShaderFormatD3D/Private/D3DShaderCompilerDXC.cpp
--- a/Engine/Source/Developer/Windows/ShaderFormatD3D/Private/D3DShaderCompilerDXC.cpp
+++ b/Engine/Source/Developer/Windows/ShaderFormatD3D/Private/D3DShaderCompilerDXC.cpp
@@ -46,6 +46,7 @@ THIRD_PARTY_INCLUDES_START
 THIRD_PARTY_INCLUDES_END
 
 #include "DXCUtils.inl"
+#include "HlslParserInternal.h"
 #include "D3DShaderCompiler.inl"
 
 static bool IsGlobalConstantBufferSupported(const FShaderTarget& Target)
diff --git a/Engine/Source/Developer/Windows/ShaderFormatD3D/ShaderFormatD3D.Build.cs b/Engine/Source/Developer/Windows/ShaderFormatD3D/ShaderFormatD3D.Build.cs
--- a/Engine/Source/Developer/Windows/ShaderFormatD3D/ShaderFormatD3D.Build.cs
+++ b/Engine/Source/Developer/Windows/ShaderFormatD3D/ShaderFormatD3D.Build.cs
@@ -22,6 +22,7 @@ public class ShaderFormatD3D : ModuleRules
 				"RenderCore",
 				"ShaderPreprocessor",
 				"ShaderCompilerCommon",
+				"HlslParser",
 				}
 			);
 
