[Shaders] Compile preview platform shaders with the preferred shader code format of the runtime RHI.

Preview platforms should always compile their shaders with the preferred shader code format (binary code or microcode) of the runtime RHI.
This avoids issues with the cross-compilation pipeline due to constraints between preview and runtime platform.
Non-preview platforms continue to support a mix of shader formats to optimize code sharing in cooks from the DDC, especially between D3D11 (DXBC) and D3D12 (DXIL).

Also remove ELanguage from D3D backend. This only mapped FName D3D shader formats to an enumeration, but FName comparisons are already fast enough. `ELanguage` doesn't add a valuable indirection and the name is misleading.

#rnx

diff --git a/Engine/Shaders/Public/Platform.ush b/Engine/Shaders/Public/Platform.ush
--- a/Engine/Shaders/Public/Platform.ush
+++ b/Engine/Shaders/Public/Platform.ush
@@ -61,7 +61,7 @@
 #elif COMPILER_GLSL || COMPILER_GLSL_ES3_1
 	// Helps with iteration when changing Vulkan shader code generation backend.
 	#include "Platform/GL/GLCommon.ush"
-#elif SM6_PROFILE || SM5_PROFILE
+#elif COMPILER_HLSL
 	#include "Platform/D3D/D3DCommon.ush"
 #endif
 
diff --git a/Engine/Shaders/Public/Platform/D3D/D3DCommon.ush b/Engine/Shaders/Public/Platform/D3D/D3DCommon.ush
--- a/Engine/Shaders/Public/Platform/D3D/D3DCommon.ush
+++ b/Engine/Shaders/Public/Platform/D3D/D3DCommon.ush
@@ -5,7 +5,7 @@ D3DCommon.ush: Common defines for the D3D shader compiler
 =============================================================================*/
 
 // Update this GUID to invalidate D3D shaders and force a shader recompilation for D3D RHIs.
-#pragma message("UESHADERMETADATA_VERSION 40E281EC-AA27-4407-B36E-45ABE0003D71")
+#pragma message("UESHADERMETADATA_VERSION 32F71811-58F3-486E-AFC8-F913E12047E6")
 
 #include "/Engine/Shared/HLSLReservedSpaces.h"
 
diff --git a/Engine/Source/Developer/Windows/ShaderFormatD3D/Private/ShaderFormatD3D.cpp b/Engine/Source/Developer/Windows/ShaderFormatD3D/Private/ShaderFormatD3D.cpp
--- a/Engine/Source/Developer/Windows/ShaderFormatD3D/Private/ShaderFormatD3D.cpp
+++ b/Engine/Source/Developer/Windows/ShaderFormatD3D/Private/ShaderFormatD3D.cpp
@@ -13,10 +13,14 @@
 
 DEFINE_LOG_CATEGORY(LogD3DShaderCompiler);
 
+// Shader formats
 static FName NAME_PCD3D_SM6(TEXT("PCD3D_SM6"));
 static FName NAME_PCD3D_SM5(TEXT("PCD3D_SM5"));
 static FName NAME_PCD3D_ES3_1(TEXT("PCD3D_ES31"));
 
+// Preferred shader code formats
+static FName NAME_DXIL(TEXT("DXIL"));
+
 static const FGuid UE_SHADER_PCD3D_SHARED_VER = FGuid("9D405D9B-44C0-4747-8E88-1EDB7413F5C3");
 static const FGuid UE_SHADER_PCD3D_SM6_VER    = FGuid("7F9AE1F2-9692-46B4-875A-28C2595E84AA");
 static const FGuid UE_SHADER_PCD3D_SM5_VER    = FGuid("F30C5AFA-77F7-48DE-8CC1-A075F179443C");
@@ -85,77 +89,59 @@ public:
 		OutFormats.Add(NAME_PCD3D_ES3_1);
 	}
 
-	enum class ELanguage
-	{
-		SM5,
-		SM6,
-		ES3_1,
-		Invalid,
-	};
-
-	static ELanguage LanguageFromFormat(FName Format)
+	static bool IsValidD3DShaderFormat(FName Format)
 	{
-		if (Format == NAME_PCD3D_SM6)
-		{
-			return ELanguage::SM6;
-		}
-		
-		if (Format == NAME_PCD3D_SM5)
-		{
-			return ELanguage::SM5;
-		}
-		
-		if (Format == NAME_PCD3D_ES3_1)
-		{
-			return ELanguage::ES3_1;
-		}
-
-		checkf(0, TEXT("Unknown format %s"), *Format.ToString());
-		return ELanguage::Invalid;
+		return
+		(
+			Format == NAME_PCD3D_SM6 ||
+			Format == NAME_PCD3D_SM5 ||
+			Format == NAME_PCD3D_ES3_1
+		);
 	}
 
-	static bool IsSM68(const FShaderCompilerInput& Input, ELanguage Language)
+	static bool IsSM68(const FShaderCompilerInput& Input)
 	{
 		return IsWorkGraphShaderFrequency(Input.Target.GetFrequency());
 	}
 
-	static bool IsSM66(const FShaderCompilerInput& Input, ELanguage Language)
+	static bool IsSM66(const FShaderCompilerInput& Input)
 	{
-		return Language == ELanguage::SM6 || IsRayTracingShaderFrequency(Input.Target.GetFrequency());
+		return Input.ShaderFormat == NAME_PCD3D_SM6 || IsRayTracingShaderFrequency(Input.Target.GetFrequency());
 	}
 
 	// Do we need any SM6.0 features?
-	static bool RequiresSM6Features(const FShaderCompilerInput& Input, ELanguage Language)
+	static bool RequiresSM6Features(const FShaderCompilerInput& Input)
 	{
-		return IsSM68(Input, Language)
-			|| IsSM66(Input, Language)
+		return IsSM68(Input)
+			|| IsSM66(Input)
 			|| Input.Environment.CompilerFlags.Contains(CFLAG_WaveOperations)
 			// TODO: Forcing DXC should not change platform flags, this needs to be moved to IsSM60 once existing uses are accounted for.
 			|| Input.Environment.CompilerFlags.Contains(CFLAG_ForceDXC)
+			|| Input.PreferredShaderCodeFormat == NAME_DXIL
 			;
 	}
 
 	// Do we need to compile with SM6.0 at all? We can compile with 6.0 but not require the language features
-	static bool IsSM60(const FShaderCompilerInput& Input, ELanguage Language)
+	static bool IsSM60(const FShaderCompilerInput& Input)
 	{
-		return RequiresSM6Features(Input, Language)
+		return RequiresSM6Features(Input)
 			|| Input.Environment.GetCompileArgument(TEXT("PLATFORM_MAX_SAMPLERS"), 0) > 16
 			;
 	}
 
-	static ED3DShaderModel DetermineShaderModel(const FShaderCompilerInput& Input, ELanguage Language)
+	static ED3DShaderModel DetermineShaderModel(const FShaderCompilerInput& Input)
 	{
-		if (IsSM68(Input, Language))
+		if (IsSM68(Input))
 		{
 			return ED3DShaderModel::SM6_8;
 		}
 
-		if (IsSM66(Input, Language))
+		if (IsSM66(Input))
 		{
 			return ED3DShaderModel::SM6_6;
 		}
 
-		if (IsSM60(Input, Language))
+		if (IsSM60(Input))
 		{
 			return ED3DShaderModel::SM6_0;
 		}
@@ -163,11 +149,6 @@ public:
 		return ED3DShaderModel::SM5_0;
 	}
 
-	static ED3DShaderModel DetermineShaderModel(const FShaderCompilerInput& Input)
-	{
-		return DetermineShaderModel(Input, LanguageFromFormat(Input.ShaderFormat));
-	}
-
 #if WITH_ENGINE
 	virtual void NotifyShaderCompiled(const FCompressedBuffer& SymbolData, FName Format, const FString& DebugInfo) const override
 	{
@@ -223,8 +204,8 @@ public:
 
 	void ModifyShaderCompilerInput(FShaderCompilerInput& Input) const final
 	{
-		const ELanguage Language = LanguageFromFormat(Input.ShaderFormat);
-		const ED3DShaderModel ShaderModel = DetermineShaderModel(Input, Language);
+		checkf(IsValidD3DShaderFormat(Input.ShaderFormat), TEXT("Unknown D3D shader format %s"), *Input.ShaderFormat.ToString());
+		const ED3DShaderModel ShaderModel = DetermineShaderModel(Input);
 
 		// Our compilers only support HLSL
 		Input.Environment.SetDefine(TEXT("COMPILER_HLSL"), true);
@@ -232,7 +213,7 @@ public:
 		// Assume min. spec HW supports with DX12/SM5
 		Input.Environment.SetDefine(TEXT("PLATFORM_SUPPORTS_ROV"), true);
 
-		const bool bSM6Features = RequiresSM6Features(Input, Language);
+		const bool bSM6Features = RequiresSM6Features(Input);
 		const bool bDXC         = DoesShaderModelRequireDXC(ShaderModel);
 
 		// Compiler specific defines
@@ -246,8 +227,8 @@ public:
 
 		// "profiles" are almost analogous to ERHIFeatureLevel but RT shaders are forcing themselves to SM6
 		Input.Environment.SetDefine(TEXT("SM6_PROFILE"),   ShaderModel >= ED3DShaderModel::SM6_6);
-		Input.Environment.SetDefine(TEXT("SM5_PROFILE"),   (Language == ELanguage::SM5));
-		Input.Environment.SetDefine(TEXT("ES3_1_PROFILE"), (Language == ELanguage::ES3_1));
+		Input.Environment.SetDefine(TEXT("SM5_PROFILE"),   (Input.ShaderFormat == NAME_PCD3D_SM5));
+		Input.Environment.SetDefine(TEXT("ES3_1_PROFILE"), (Input.ShaderFormat == NAME_PCD3D_ES3_1));
 
 		// Add SM6.6+ specific defines. None of these are intended to be enabled in lower SM's
 		if (ShaderModel >= ED3DShaderModel::SM6_6)
diff --git a/Engine/Source/Runtime/D3D12RHI/Private/D3D12Adapter.cpp b/Engine/Source/Runtime/D3D12RHI/Private/D3D12Adapter.cpp
--- a/Engine/Source/Runtime/D3D12RHI/Private/D3D12Adapter.cpp
+++ b/Engine/Source/Runtime/D3D12RHI/Private/D3D12Adapter.cpp
@@ -1293,6 +1293,12 @@ void FD3D12Adapter::InitializeDevices()
 					GRHISupportsMeshShadersTier0 = GRHISupportsMeshShadersTier1 = (D3D12Caps7.MeshShaderTier >= D3D12_MESH_SHADER_TIER_1);
 				}
 
+				// If SM 6.0 is supported, DXIL is the preferred shader code format for preview platforms. Multi-patform cooks can still share their shader code betweeen D3D11 and D3D12 from the DDC.
+				if (Desc.MaxSupportedShaderModel >= D3D_SHADER_MODEL_6_0)
+				{
+					GRHIGlobals.PreferredPreviewShaderCodeFormat = TEXT("DXIL");
+				}
+
 				if (D3D12Caps7.MeshShaderTier >= D3D12_MESH_SHADER_TIER_1)
 				{
 					UE_LOG(LogD3D12RHI, Log, TEXT("Mesh shader tier 1.0 is supported"));
diff --git a/Engine/Source/Runtime/Engine/Private/ShaderCompiler/ShaderCompiler.cpp b/Engine/Source/Runtime/Engine/Private/ShaderCompiler/ShaderCompiler.cpp
--- a/Engine/Source/Runtime/Engine/Private/ShaderCompiler/ShaderCompiler.cpp
+++ b/Engine/Source/Runtime/Engine/Private/ShaderCompiler/ShaderCompiler.cpp
@@ -3179,6 +3179,7 @@ void GlobalBeginCompileShader(
 	Input.Target = Target;
 	Input.ShaderPlatformName = FDataDrivenShaderPlatformInfo::GetName(ShaderPlatform);
 	Input.ShaderFormat = ShaderFormatName;
+	Input.PreferredShaderCodeFormat = FDataDrivenShaderPlatformInfo::GetIsPreviewPlatform(ShaderPlatform) ? GRHIGlobals.PreferredPreviewShaderCodeFormat : NAME_None;
 	Input.SupportedHardwareMask = TargetPlatform ? TargetPlatform->GetSupportedHardwareMask() : 0;
 	Input.CompressionFormat = GetShaderCompressionFormat();
 	GetShaderCompressionOodleSettings(Input.OodleCompressor, Input.OodleLevel);
diff --git a/Engine/Source/Runtime/RHI/Public/RHIGlobals.h b/Engine/Source/Runtime/RHI/Public/RHIGlobals.h
--- a/Engine/Source/Runtime/RHI/Public/RHIGlobals.h
+++ b/Engine/Source/Runtime/RHI/Public/RHIGlobals.h
@@ -739,6 +739,9 @@ PRAGMA_ENABLE_DEPRECATION_WARNINGS
 
 	/** True if RHI needs to manually apply an RHI transition when a resource is flag as SRVGraphicsNonPixel. */
 	bool NeedsSRVGraphicsNonPixelWorkaround = false;
+
+	/** Preferred format of FShaderCode (binary code or microcode). Only used for preview platforms with more than one shader code being supported such as D3D12 (supports DXIL and DXBC). */
+	FName PreferredPreviewShaderCodeFormat = NAME_None;
 };
 
 extern RHI_API FRHIGlobals GRHIGlobals;
diff --git a/Engine/Source/Runtime/RenderCore/Private/ShaderCore.cpp b/Engine/Source/Runtime/RenderCore/Private/ShaderCore.cpp
--- a/Engine/Source/Runtime/RenderCore/Private/ShaderCore.cpp
+++ b/Engine/Source/Runtime/RenderCore/Private/ShaderCore.cpp
@@ -4234,23 +4234,19 @@ FArchive& operator<<(FArchive& Ar, FShaderCode& Output)
 
 FArchive& operator<<(FArchive& Ar, FShaderCompilerInput& Input)
 {
+	auto SerializeName = [&Ar](FName& Name)
+		{
+			FString NameString(Name.ToString());
+			Ar << NameString;
+			Name = FName(*NameString);
+		};
+
 	// Note: this serialize is used to pass between UE and the shader compile worker, recompile both when modifying
 	Ar << Input.Target;
-	{
-		FString ShaderFormatString(Input.ShaderFormat.ToString());
-		Ar << ShaderFormatString;
-		Input.ShaderFormat = FName(*ShaderFormatString);
-	}
-	{
-		FString CompressionFormatString(Input.CompressionFormat.ToString());
-		Ar << CompressionFormatString;
-		Input.CompressionFormat = FName(*CompressionFormatString);
-	}
-	{
-		FString ShaderPlatformNameString(Input.ShaderPlatformName.ToString());
-		Ar << ShaderPlatformNameString;
-		Input.ShaderPlatformName = FName(*ShaderPlatformNameString);
-	}
+	SerializeName(Input.ShaderFormat);
+	SerializeName(Input.CompressionFormat);
+	SerializeName(Input.ShaderPlatformName);
+	SerializeName(Input.PreferredShaderCodeFormat);
 	Ar << Input.VirtualSourceFilePath;
 	Ar << Input.EntryPointName;
 	Ar << Input.ShaderName;
@@ -4370,6 +4366,8 @@ FShaderCompilerInputHash FShaderCompileJob::GetInputHash()
 
 	uint32 FormatVersion = GetTargetPlatformManagerRef().ShaderFormatVersion(Input.ShaderFormat);
 	Hasher << FormatVersion;
+
+	Hasher << Input.PreferredShaderCodeFormat;
 	
 	FShaderTarget Target = Input.Target;
 	Hasher << Target;
diff --git a/Engine/Source/Runtime/RenderCore/Public/ShaderCompilerCore.h b/Engine/Source/Runtime/RenderCore/Public/ShaderCompilerCore.h
--- a/Engine/Source/Runtime/RenderCore/Public/ShaderCompilerCore.h
+++ b/Engine/Source/Runtime/RenderCore/Public/ShaderCompilerCore.h
@@ -31,7 +31,7 @@ using FShaderSharedStringPtr = TSharedPtr<FString, ESPMode::ThreadSafe>;
 typedef TSharedPtr<TArray<ANSICHAR>, ESPMode::ThreadSafe> FShaderSharedAnsiStringPtr;
 
 // this is for the protocol, not the data, bump if FShaderCompilerInput/FShaderPreprocessOutput serialization, SerializeWorkerInput or ProcessInputFromArchive changes.
-inline const int32 ShaderCompileWorkerInputVersion = 31;
+inline const int32 ShaderCompileWorkerInputVersion = 32;
 // this is for the protocol, not the data, bump if FShaderCompilerOutput or WriteToOutputArchive changes.
 inline const int32 ShaderCompileWorkerOutputVersion = 28;
 // this is for the protocol, not the data.
@@ -213,6 +213,9 @@ struct FShaderCompilerInput
 	FName ShaderFormat;
 	FName CompressionFormat;
 	FName ShaderPlatformName;
+
+	// Preferred output format for FShaderCode, e.g. when DXIL is preferred over DXBC while both would be supported.
+	FName PreferredShaderCodeFormat;
 	
 	FString VirtualSourceFilePath;
 	FString EntryPointName;
