[HWRT]
- Refactor RayTracingRuntimeMode parsing to go through shared functions for cook and runtime parsing of platform settings
- Mark platforms only supporting inline raytracing correctly
- Add support for inline raytracing mode only for Windows and Linux to not compile raytracing shaders when not needed

diff --git a/Engine/Config/BaseEngine.ini b/Engine/Config/BaseEngine.ini
--- a/Engine/Config/BaseEngine.ini
+++ b/Engine/Config/BaseEngine.ini
@@ -3419,7 +3419,7 @@ bEnable2DPhysics=false
 +D3D12TargetedShaderFormats=PCD3D_SM5
 +D3D11TargetedShaderFormats=PCD3D_SM5
 bEnableDistanceFields=true
-bEnableRayTracing=true
+RayTracingMode=Full
 bTarget32Bit=false
 CacheSizeKB=65536
 MaxSampleRate=48000
@@ -3435,7 +3435,7 @@ AudioNumSourceWorkers=4
 
 [/Script/LinuxTargetPlatform.LinuxTargetSettings]
 +TargetedRHIs=SF_VULKAN_SM5
-bEnableRayTracing=true
+RayTracingMode=Full
 
 [/Script/MacTargetPlatform.MacTargetSettings]
 MetalLanguageVersion=7
diff --git a/Engine/Source/Developer/Android/AndroidTargetPlatformSettings/Public/AndroidTargetPlatformSettings.h b/Engine/Source/Developer/Android/AndroidTargetPlatformSettings/Public/AndroidTargetPlatformSettings.h
--- a/Engine/Source/Developer/Android/AndroidTargetPlatformSettings/Public/AndroidTargetPlatformSettings.h
+++ b/Engine/Source/Developer/Android/AndroidTargetPlatformSettings/Public/AndroidTargetPlatformSettings.h
@@ -120,6 +120,11 @@ public:
 		return bEnableRayTracing;
 	}
 
+	virtual ERayTracingRuntimeMode GetRayTracingMode() const override
+	{
+		return UsesRayTracing() ? ERayTracingRuntimeMode::Inline : ERayTracingRuntimeMode::Disabled;
+	}
+
 #if WITH_ENGINE
 	ANDROIDTARGETPLATFORMSETTINGS_API virtual void GetReflectionCaptureFormats(TArray<FName>& OutFormats) const override;
 	ANDROIDTARGETPLATFORMSETTINGS_API virtual const UTextureLODSettings& GetTextureLODSettings() const override;
diff --git a/Engine/Source/Developer/Linux/LinuxTargetPlatformSettings/Classes/LinuxTargetSettings.h b/Engine/Source/Developer/Linux/LinuxTargetPlatformSettings/Classes/LinuxTargetSettings.h
--- a/Engine/Source/Developer/Linux/LinuxTargetPlatformSettings/Classes/LinuxTargetSettings.h
+++ b/Engine/Source/Developer/Linux/LinuxTargetPlatformSettings/Classes/LinuxTargetSettings.h
@@ -9,6 +9,7 @@
 #include "CoreMinimal.h"
 #include "UObject/ObjectMacros.h"
 #include "UObject/Object.h"
+#include "Misc/TargetSettingsDefinitions.h"
 #include "LinuxTargetSettings.generated.h"
 
 
@@ -23,6 +24,8 @@ public:
 
 	GENERATED_UCLASS_BODY()
 
+	LINUXTARGETPLATFORMSETTINGS_API virtual void PostInitProperties() override;
+
 	/** Which of the currently enabled spatialization plugins to use. */
 	UPROPERTY(config, EditAnywhere, Category = "Audio")
 	FString SpatializationPlugin;
@@ -54,6 +57,14 @@ public:
 	UPROPERTY(config, EditAnywhere, Category = "Renderer", meta = (DisplayName = "Generate Nanite Fallback Meshes", ConfigRestartRequired = true))
 	bool bGenerateNaniteFallbackMeshes = true;
 
+	/** Deprecated use RayTracingMode instead. */
+	UPROPERTY(config, meta = (DeprecatedProperty))
+	bool bEnableRayTracing_DEPRECATED;
+
+	/** Controls the level of ray tracing support for Windows. r.RayTracing cvar also needs to be enabled. */
+	UPROPERTY(config, EditAnywhere, Category = "Renderer", meta = (DisplayName = "Hardware Ray Tracing Mode"))
+	ETargetSettingsRayTracingRuntimeMode RayTracingMode;
+
 	void virtual OverrideConfigSection(FString& InOutSectionName) override
 	{
 		InOutSectionName = TEXT("/Script/LinuxTargetPlatform.LinuxTargetSettings");
diff --git a/Engine/Source/Developer/Linux/LinuxTargetPlatformSettings/Private/LinuxTargetPlatformClasses.cpp b/Engine/Source/Developer/Linux/LinuxTargetPlatformSettings/Private/LinuxTargetPlatformClasses.cpp
--- a/Engine/Source/Developer/Linux/LinuxTargetPlatformSettings/Private/LinuxTargetPlatformClasses.cpp
+++ b/Engine/Source/Developer/Linux/LinuxTargetPlatformSettings/Private/LinuxTargetPlatformClasses.cpp
@@ -7,6 +7,7 @@
 #include "CoreMinimal.h"
 #include "LinuxTargetSettings.h"
 #include "LinuxTargetPlatformSettings.h"
+#include "Misc/ConfigCacheIni.h"
 
 
 /* ULinuxTargetSettings structors
@@ -15,3 +16,25 @@
 ULinuxTargetSettings::ULinuxTargetSettings( const FObjectInitializer& ObjectInitializer )
 	: Super(ObjectInitializer)
 { }
+
+void ULinuxTargetSettings::PostInitProperties()
+{
+	Super::PostInitProperties();
+
+	FProperty* RayTracingModeProperty = GetClass()->FindPropertyByName(GET_MEMBER_NAME_CHECKED(ULinuxTargetSettings, RayTracingMode));
+	if (RayTracingModeProperty)
+	{
+		FConfigFile LocalEngineIni;
+		FConfigFile* EngineIni = FConfigCacheIni::FindOrLoadPlatformConfig(LocalEngineIni, TEXT("Engine"), TEXT("Linux"));
+
+		// Convert and delete deprecated properties
+		// Deprecated properties are automatically deleted from .inis when updated via UpdateSinglePropertyInConfigFile()
+		bool bOldEnableRayTracing = false;
+		if (EngineIni->GetBool(LINUX_SECTION_TEXT, TEXT("bEnableRayTracing"), bOldEnableRayTracing))
+		{
+			RayTracingMode = bOldEnableRayTracing ? ETargetSettingsRayTracingRuntimeMode::Full : ETargetSettingsRayTracingRuntimeMode::Disabled;
+			UpdateSinglePropertyInConfigFile(RayTracingModeProperty, GetDefaultConfigFilename());
+			UpdateSinglePropertyInConfigFile(GetClass()->FindPropertyByName(FName(TEXT("bEnableRayTracing"))), GetDefaultConfigFilename());
+		}
+	}
+}
diff --git a/Engine/Source/Developer/Linux/LinuxTargetPlatformSettings/Public/LinuxTargetPlatformSettings.h b/Engine/Source/Developer/Linux/LinuxTargetPlatformSettings/Public/LinuxTargetPlatformSettings.h
--- a/Engine/Source/Developer/Linux/LinuxTargetPlatformSettings/Public/LinuxTargetPlatformSettings.h
+++ b/Engine/Source/Developer/Linux/LinuxTargetPlatformSettings/Public/LinuxTargetPlatformSettings.h
@@ -14,6 +14,7 @@
 #include "Linux/LinuxPlatformProperties.h"
 
 #define LOCTEXT_NAMESPACE "TLinuxTargetPlatformSettings"
+#define LINUX_SECTION_TEXT TEXT("/Script/LinuxTargetPlatform.LinuxTargetSettings")
 
 class UTextureLODSettings;
 
@@ -88,7 +89,7 @@ public:
 	{
 		// Get the Target RHIs for this platform, we do not always want all those that are supported. (reload in case user changed in the editor)
 		TArray<FString>TargetedShaderFormats;
-		GConfig->GetArray(TEXT("/Script/LinuxTargetPlatform.LinuxTargetSettings"), TEXT("TargetedRHIs"), TargetedShaderFormats, GEngineIni);
+		GConfig->GetArray(LINUX_SECTION_TEXT, TEXT("TargetedRHIs"), TargetedShaderFormats, GEngineIni);
 
 		// Gather the list of Target RHIs and filter out any that may be invalid.
 		TArray<FName> PossibleShaderFormats;
@@ -144,6 +145,21 @@ public:
 
 		return !bGenerateNaniteFallbackMeshes;
 	}
+	
+	virtual bool UsesRayTracing() const override
+	{
+		return GetRayTracingMode() != ERayTracingRuntimeMode::Disabled;
+	}
+
+	virtual ERayTracingRuntimeMode GetRayTracingMode() const override
+	{
+		if (!TTargetPlatformSettingsBase<TProperties>::UsesRayTracing())
+		{
+			return ERayTracingRuntimeMode::Disabled;
+		}
+
+		return TTargetPlatformSettingsBase<TProperties>::ParseRuntimeRayTracingMode(LINUX_SECTION_TEXT);
+	}
 
 	//~ End ITargetPlatform Interface
 
diff --git a/Engine/Source/Developer/Mac/MacTargetPlatformSettings/Public/GenericMacTargetPlatformSettings.h b/Engine/Source/Developer/Mac/MacTargetPlatformSettings/Public/GenericMacTargetPlatformSettings.h
--- a/Engine/Source/Developer/Mac/MacTargetPlatformSettings/Public/GenericMacTargetPlatformSettings.h
+++ b/Engine/Source/Developer/Mac/MacTargetPlatformSettings/Public/GenericMacTargetPlatformSettings.h
@@ -129,6 +129,12 @@ public:
 	 
 		return bEnableRayTracing;
 	}
+
+	virtual ERayTracingRuntimeMode GetRayTracingMode() const override
+	{
+		return UsesRayTracing() ? ERayTracingRuntimeMode::Inline : ERayTracingRuntimeMode::Disabled;
+	}
+
 	//~ End ITargetPlatform Interface
 
 private:
diff --git a/Engine/Source/Developer/TargetPlatform/Private/TargetPlatformSettingsBase.cpp b/Engine/Source/Developer/TargetPlatform/Private/TargetPlatformSettingsBase.cpp
--- a/Engine/Source/Developer/TargetPlatform/Private/TargetPlatformSettingsBase.cpp
+++ b/Engine/Source/Developer/TargetPlatform/Private/TargetPlatformSettingsBase.cpp
@@ -62,7 +62,22 @@ bool FTargetPlatformSettingsBase::UsesRayTracing() const
 
 ERayTracingRuntimeMode FTargetPlatformSettingsBase::GetRayTracingMode() const
 {
-	return UsesRayTracing() ? ERayTracingRuntimeMode::Full : ERayTracingRuntimeMode::Disabled;
+	return ERayTracingRuntimeMode::Disabled;
+}
+
+ERayTracingRuntimeMode FTargetPlatformSettingsBase::ParseRuntimeRayTracingMode(const TCHAR* SectionName) const
+{
+	FString ModeString;
+	GetConfigSystem()->GetString(SectionName, TEXT("RayTracingMode"), ModeString, GEngineIni);
+	if (ModeString.Equals(TEXT("Inline"), ESearchCase::IgnoreCase))
+	{
+		return ERayTracingRuntimeMode::Inline;
+	}
+	else if (ModeString.Equals(TEXT("Full"), ESearchCase::IgnoreCase))
+	{
+		return ERayTracingRuntimeMode::Full;
+	}
+	return ERayTracingRuntimeMode::Disabled;
 }
 
 uint32 FTargetPlatformSettingsBase::GetSupportedHardwareMask() const
diff --git a/Engine/Source/Developer/TargetPlatform/Public/Common/TargetPlatformSettingsBase.h b/Engine/Source/Developer/TargetPlatform/Public/Common/TargetPlatformSettingsBase.h
--- a/Engine/Source/Developer/TargetPlatform/Public/Common/TargetPlatformSettingsBase.h
+++ b/Engine/Source/Developer/TargetPlatform/Public/Common/TargetPlatformSettingsBase.h
@@ -68,6 +68,9 @@ public:
 	}
 
 protected:
+
+	TARGETPLATFORM_API ERayTracingRuntimeMode ParseRuntimeRayTracingMode(const TCHAR* SectionName) const;
+
 	FTargetPlatformSettingsBase() {};
 };
 
diff --git a/Engine/Source/Developer/TargetPlatform/Public/Interfaces/ITargetPlatformSettings.h b/Engine/Source/Developer/TargetPlatform/Public/Interfaces/ITargetPlatformSettings.h
--- a/Engine/Source/Developer/TargetPlatform/Public/Interfaces/ITargetPlatformSettings.h
+++ b/Engine/Source/Developer/TargetPlatform/Public/Interfaces/ITargetPlatformSettings.h
@@ -10,7 +10,7 @@ enum class EOfflineBVHMode
 	MinimizeMemory,
 };
 
-enum class ERayTracingRuntimeMode
+enum class ERayTracingRuntimeMode : uint8
 {
 	Disabled,
 	Inline,
diff --git a/Engine/Source/Developer/Windows/WindowsTargetPlatformSettings/Classes/WindowsTargetSettings.h b/Engine/Source/Developer/Windows/WindowsTargetPlatformSettings/Classes/WindowsTargetSettings.h
--- a/Engine/Source/Developer/Windows/WindowsTargetPlatformSettings/Classes/WindowsTargetSettings.h
+++ b/Engine/Source/Developer/Windows/WindowsTargetPlatformSettings/Classes/WindowsTargetSettings.h
@@ -9,6 +9,7 @@
 #include "CoreMinimal.h"
 #include "UObject/Object.h"
 #include "AudioCompressionSettings.h"
+#include "Misc/TargetSettingsDefinitions.h"
 #include "WindowsTargetSettings.generated.h"
 
 #define UE_API WINDOWSTARGETPLATFORMSETTINGS_API
@@ -67,6 +68,14 @@ public:
 	UPROPERTY(config, EditAnywhere, Category = "Renderer", meta = (DisplayName = "Generate Nanite Fallback Meshes", ConfigRestartRequired = true))
 	bool bGenerateNaniteFallbackMeshes = true;
 
+	/** Deprecated use RayTracingMode instead. */
+	UPROPERTY(config, meta = (DeprecatedProperty))
+	bool bEnableRayTracing_DEPRECATED;
+
+	/** Controls the level of ray tracing support for Windows. r.RayTracing cvar also needs to be enabled. */
+	UPROPERTY(config, EditAnywhere, Category = "Renderer", meta = (DisplayName = "Hardware Ray Tracing Mode"))
+	ETargetSettingsRayTracingRuntimeMode RayTracingMode;
+	
 	/** The compiler version to use for this project. May be different to the chosen IDE. */
 	UPROPERTY(EditAnywhere, config, Category = "Toolchain", Meta = (DisplayName = "Compiler Version"))
 	ECompilerVersion Compiler;
diff --git a/Engine/Source/Developer/Windows/WindowsTargetPlatformSettings/Private/WindowsTargetPlatformClasses.cpp b/Engine/Source/Developer/Windows/WindowsTargetPlatformSettings/Private/WindowsTargetPlatformClasses.cpp
--- a/Engine/Source/Developer/Windows/WindowsTargetPlatformSettings/Private/WindowsTargetPlatformClasses.cpp
+++ b/Engine/Source/Developer/Windows/WindowsTargetPlatformSettings/Private/WindowsTargetPlatformClasses.cpp
@@ -1,6 +1,7 @@
 // Copyright Epic Games, Inc. All Rights Reserved.
 
 #include "WindowsTargetSettings.h"
+#include "Misc/ConfigCacheIni.h"
 
 /* UWindowsTargetSettings structors
  *****************************************************************************/
@@ -41,6 +42,19 @@ static void AddToShaderFormatList(TArray<FString>& Dest, const TArray<FString>&
 void UWindowsTargetSettings::PostInitProperties()
 {
 	Super::PostInitProperties();
+	
+	FConfigFile LocalEngineIni;
+	FConfigFile* EngineIni = FConfigCacheIni::FindOrLoadPlatformConfig(LocalEngineIni, TEXT("Engine"), TEXT("Win64"));
+	
+	// Convert and delete deprecated properties
+	// Deprecated properties are automatically deleted from .inis when updated via UpdateSinglePropertyInConfigFile()
+	bool bOldEnableRayTracing = false;
+	if (EngineIni->GetBool(TEXT("/Script/WindowsTargetPlatform.WindowsTargetSettings"), TEXT("bEnableRayTracing"), bOldEnableRayTracing))
+	{
+		RayTracingMode = bOldEnableRayTracing ? ETargetSettingsRayTracingRuntimeMode::Full : ETargetSettingsRayTracingRuntimeMode::Disabled;
+		UpdateSinglePropertyInConfigFile(GetClass()->FindPropertyByName(GET_MEMBER_NAME_CHECKED(UWindowsTargetSettings, RayTracingMode)), GetDefaultConfigFilename());
+		UpdateSinglePropertyInConfigFile(GetClass()->FindPropertyByName(FName(TEXT("bEnableRayTracing"))), GetDefaultConfigFilename());
+	};
 
 	if (!TargetedRHIs_DEPRECATED.IsEmpty())
 	{
diff --git a/Engine/Source/Developer/Windows/WindowsTargetPlatformSettings/Public/GenericWindowsTargetPlatformSettings.h b/Engine/Source/Developer/Windows/WindowsTargetPlatformSettings/Public/GenericWindowsTargetPlatformSettings.h
--- a/Engine/Source/Developer/Windows/WindowsTargetPlatformSettings/Public/GenericWindowsTargetPlatformSettings.h
+++ b/Engine/Source/Developer/Windows/WindowsTargetPlatformSettings/Public/GenericWindowsTargetPlatformSettings.h
@@ -14,6 +14,7 @@
 #include "HAL/PlatformMisc.h"
 
 #define LOCTEXT_NAMESPACE "TGenericWindowsTargetPlatformSettings"
+#define WIN_SECTION_TEXT TEXT("/Script/WindowsTargetPlatform.WindowsTargetSettings")
 
 /**
  * Template for Windows target platforms settings
@@ -144,7 +145,7 @@ public:
 		for (const TCHAR* Name : RelevantSettings)
 		{
 			TArray<FString> NewTargetedShaderFormats;
-			ConfigSystem->GetArray(TEXT("/Script/WindowsTargetPlatform.WindowsTargetSettings"), Name, NewTargetedShaderFormats, GEngineIni);
+			ConfigSystem->GetArray(WIN_SECTION_TEXT, Name, NewTargetedShaderFormats, GEngineIni);
 
 			for (const FString& NewShaderFormat : NewTargetedShaderFormats)
 			{
@@ -257,7 +258,7 @@ public:
 	virtual bool ShouldStripNaniteFallbackMeshes() const override
 	{
 		bool bGenerateNaniteFallbackMeshes = true;
-		GConfig->GetBool(TEXT("/Script/WindowsTargetPlatform.WindowsTargetSettings"), TEXT("bGenerateNaniteFallbackMeshes"), bGenerateNaniteFallbackMeshes, GEngineIni);
+		GConfig->GetBool(WIN_SECTION_TEXT, TEXT("bGenerateNaniteFallbackMeshes"), bGenerateNaniteFallbackMeshes, GEngineIni);
 
 		return !bGenerateNaniteFallbackMeshes;
 	}
@@ -265,17 +266,24 @@ public:
 	virtual bool UsesDistanceFields() const override
 	{
 		bool bEnableDistanceFields = false;
-		GConfig->GetBool(TEXT("/Script/WindowsTargetPlatform.WindowsTargetSettings"), TEXT("bEnableDistanceFields"), bEnableDistanceFields, GEngineIni);
+		GConfig->GetBool(WIN_SECTION_TEXT, TEXT("bEnableDistanceFields"), bEnableDistanceFields, GEngineIni);
 
 		return bEnableDistanceFields && TSuper::UsesDistanceFields();
 	}
 
 	virtual bool UsesRayTracing() const override
 	{
-		bool bEnableRayTracing = false;
-		GConfig->GetBool(TEXT("/Script/WindowsTargetPlatform.WindowsTargetSettings"), TEXT("bEnableRayTracing"), bEnableRayTracing, GEngineIni);
+		return GetRayTracingMode() != ERayTracingRuntimeMode::Disabled;
+	}
+
+	virtual ERayTracingRuntimeMode GetRayTracingMode() const override
+	{
+		if (!TTargetPlatformSettingsBase<TProperties>::UsesRayTracing())
+		{
+			return ERayTracingRuntimeMode::Disabled;
+		}
 
-		return bEnableRayTracing && TSuper::UsesRayTracing();
+		return TTargetPlatformSettingsBase<TProperties>::ParseRuntimeRayTracingMode(WIN_SECTION_TEXT);
 	}
 
 	//~ End ITargetPlatform Interface
diff --git a/Engine/Source/Runtime/Apple/MetalRHI/Private/MetalRHI.cpp b/Engine/Source/Runtime/Apple/MetalRHI/Private/MetalRHI.cpp
--- a/Engine/Source/Runtime/Apple/MetalRHI/Private/MetalRHI.cpp
+++ b/Engine/Source/Runtime/Apple/MetalRHI/Private/MetalRHI.cpp
@@ -29,6 +29,7 @@
 #include "MetalBindlessDescriptors.h"
 #include "DataDrivenShaderPlatformInfo.h"
 #include "MetalResourceCollection.h"
+#include "RHICore.h"
  
 #if METAL_RHI_RAYTRACING
 #include "MetalRayTracing.h"
@@ -1185,7 +1186,7 @@ FMetalDynamicRHI::FMetalDynamicRHI(ERHIFeatureLevel::Type RequestedFeatureLevel)
 		if (!FParse::Param(FCommandLine::Get(), TEXT("noraytracing")))
 		{
 			GRHISupportsRayTracing = RHISupportsRayTracing(GMaxRHIShaderPlatform);
-			GRHISupportsRayTracingShaders = RHISupportsRayTracingShaders(GMaxRHIShaderPlatform);
+			GRHISupportsRayTracingShaders = RHISupportsRayTracingShaders(GMaxRHIShaderPlatform) && UE::RHICore::AreRayTracingShadersEnabled();
 
 			GRHISupportsRayTracingPSOAdditions = false;
 			GRHISupportsRayTracingAMDHitToken = false;
diff --git a/Engine/Source/Runtime/CoreUObject/Public/Misc/TargetSettingsDefinitions.h b/Engine/Source/Runtime/CoreUObject/Public/Misc/TargetSettingsDefinitions.h
new file mode 100644
--- /dev/null
+++ b/Engine/Source/Runtime/CoreUObject/Public/Misc/TargetSettingsDefinitions.h
@@ -0,0 +1,31 @@
+// Copyright Epic Games, Inc. All Rights Reserved.
+
+#pragma once
+
+#include "UObject/Object.h"
+
+#include "TargetSettingsDefinitions.generated.h"
+
+UENUM()
+enum class ETargetSettingsOfflineBVHMode
+{
+	Disabled UMETA(DisplayName = "Disabled", ToolTip = "Don't use offline BVHs"),
+	MaximizePerformance UMETA(DisplayName = "Maximize Performance", ToolTip = "Prefer high performance at the cost of extra memory usage"),
+	MinimizeMemory UMETA(DisplayName = "Minimize Memory", ToolTip = "Prefer reduced memory usage at the cost of some tracing performance"),
+};
+
+UENUM()
+enum class ETargetSettingsRayTracingRuntimeMode
+{
+	Disabled UMETA(DisplayName = "Disabled", ToolTip = "Ray tracing is disabled."),
+	Inline UMETA(DisplayName = "Inline", ToolTip = "Only inline ray tracing passes will be supported. Removes certain runtime and cook time overhead but also disables ray tracing passes that use material shaders (Lumen Hit Lighting, selected debug modes etc)."),
+	Full UMETA(DisplayName = "Full", ToolTip = "Full support for ray tracing features, including ray generation, hit and miss shaders as well as inline ray queries."),
+};
+
+UENUM()
+enum class ETargetSettingsShadingMode
+{
+	UseProjectSetting UMETA(DisplayName = "Use Project Setting", ToolTip = "Use the desktop project setting found under Engine - Rendering."),
+	Deferred UMETA(DisplayName = "Deferred", ToolTip = "Use deferred shading."),
+	Forward UMETA(DisplayName = "Forward", ToolTip = "Use forward shading."),
+};
\ No newline at end of file
diff --git a/Engine/Source/Runtime/D3D12RHI/Private/D3D12Adapter.cpp b/Engine/Source/Runtime/D3D12RHI/Private/D3D12Adapter.cpp
--- a/Engine/Source/Runtime/D3D12RHI/Private/D3D12Adapter.cpp
+++ b/Engine/Source/Runtime/D3D12RHI/Private/D3D12Adapter.cpp
@@ -1236,15 +1236,16 @@ void FD3D12Adapter::InitializeDevices()
 							UE_LOG(LogD3D12RHI, Log, TEXT("D3D12 ray tracing tier 1.1 and bindless resources are supported."));
 
 							GRHISupportsRayTracing = RHISupportsRayTracing(GMaxRHIShaderPlatform);
-							GRHISupportsRayTracingShaders = GRHISupportsRayTracing && RHISupportsRayTracingShaders(GMaxRHIShaderPlatform);
 
+							// inline raytracing requires bindless on d3d12
 							const ERHIBindlessConfiguration BindlessConfiguration = RHIGetRuntimeBindlessConfiguration(GMaxRHIShaderPlatform);
-
-							GRHISupportsRayTracingPSOAdditions = true;
 							GRHISupportsInlineRayTracing = GRHISupportsRayTracing && RHISupportsInlineRayTracing(GMaxRHIShaderPlatform) && (GMaxRHIFeatureLevel >= ERHIFeatureLevel::SM6) && !IsBindlessDisabled(BindlessConfiguration);
 														
 							// Inline RayTracing SBT is required
 							GRHIGlobals.RayTracing.RequiresInlineRayTracingSBT = GRHISupportsInlineRayTracing;
+														
+							GRHISupportsRayTracingPSOAdditions = true;
+							GRHISupportsRayTracingShaders = GRHISupportsRayTracing && RHISupportsRayTracingShaders(GMaxRHIShaderPlatform) && UE::RHICore::AreRayTracingShadersEnabled();
 						}
 						else
 						{
diff --git a/Engine/Source/Runtime/RHICore/Private/RHICore.cpp b/Engine/Source/Runtime/RHICore/Private/RHICore.cpp
--- a/Engine/Source/Runtime/RHICore/Private/RHICore.cpp
+++ b/Engine/Source/Runtime/RHICore/Private/RHICore.cpp
@@ -3,6 +3,7 @@
 #include "RHICore.h"
 #include "Modules/ModuleManager.h"
 #include "Misc/CommandLine.h"
+#include "Misc/ConfigCacheIni.h"
 
 DEFINE_LOG_CATEGORY(LogRHICore);
 
@@ -62,5 +63,30 @@ bool AllowVendorDevice()
 	return bAllowVendorDevice;
 }
 
+bool AreRayTracingShadersEnabled()
+{
+	struct FInit
+	{
+		bool bEnableRayTracingShaders = true;
+		FInit()
+		{
+			FString ModeString;
+			GConfig->GetString(FPlatformProperties::GetRuntimeSettingsClassName(), TEXT("RayTracingMode"), ModeString, GEngineIni);
+
+			if (ModeString.Equals(TEXT("Inline"), ESearchCase::IgnoreCase))
+			{
+				bEnableRayTracingShaders = false;
+			}
+			else if (ModeString.Equals(TEXT("Full"), ESearchCase::IgnoreCase))
+			{
+				bEnableRayTracingShaders = true;
+			}
+		}
+
+	} static Init;
+
+	return Init.bEnableRayTracingShaders;
+}
+
 } //! RHICore
 } //! UE
\ No newline at end of file
diff --git a/Engine/Source/Runtime/RHICore/Public/RHICore.h b/Engine/Source/Runtime/RHICore/Public/RHICore.h
--- a/Engine/Source/Runtime/RHICore/Public/RHICore.h
+++ b/Engine/Source/Runtime/RHICore/Public/RHICore.h
@@ -36,6 +36,8 @@ struct FResolveTextureInfo
 RHICORE_API void ResolveRenderPassTargets(const FRHIRenderPassInfo& Info, TFunction<void(FResolveTextureInfo)> Function);
 RHICORE_API FRHIViewDesc::EDimension AdjustViewInfoDimensionForNarrowing(const FRHIViewDesc::FTexture::FViewInfo& ViewInfo, const FRHITextureDesc& TextureDesc);
 
+RHICORE_API bool AreRayTracingShadersEnabled();
+
 RHICORE_API bool AllowVendorDevice();
 
 } //! UE::RHICore
diff --git a/Engine/Source/Runtime/VulkanRHI/Private/VulkanRHI.cpp b/Engine/Source/Runtime/VulkanRHI/Private/VulkanRHI.cpp
--- a/Engine/Source/Runtime/VulkanRHI/Private/VulkanRHI.cpp
+++ b/Engine/Source/Runtime/VulkanRHI/Private/VulkanRHI.cpp
@@ -20,6 +20,7 @@
 #include "VulkanLLM.h"
 #include "Misc/EngineVersion.h"
 #include "GlobalShader.h"
+#include "RHICore.h"
 #include "VulkanResourceCollection.h"
 #include "RHIValidation.h"
 #include "RHIUtilities.h"
@@ -907,7 +908,7 @@ void FVulkanDynamicRHI::InitInstance()
 
 		if (GRHISupportsRayTracing)
 		{
-			GRHISupportsRayTracingShaders = RHISupportsRayTracingShaders(GMaxRHIShaderPlatform) && Device->GetOptionalExtensions().HasRayTracingPipeline;
+			GRHISupportsRayTracingShaders = RHISupportsRayTracingShaders(GMaxRHIShaderPlatform) && Device->GetOptionalExtensions().HasRayTracingPipeline && UE::RHICore::AreRayTracingShadersEnabled();
 			GRHISupportsInlineRayTracing = !bUseAMDCompatibilityMode && RHISupportsInlineRayTracing(GMaxRHIShaderPlatform) && Device->GetOptionalExtensions().HasRayQuery;
 
 			// Inline RayTracing SBT is needed if raytracing position fetch isn't available
