Added iwyu pragma keep around forward declarations in fwd headers

diff --git a/Engine/Source/Developer/CookedEditor/Public/CookedEditorPackageManager.h b/Engine/Source/Developer/CookedEditor/Public/CookedEditorPackageManager.h
--- a/Engine/Source/Developer/CookedEditor/Public/CookedEditorPackageManager.h
+++ b/Engine/Source/Developer/CookedEditor/Public/CookedEditorPackageManager.h
@@ -2,6 +2,8 @@
 
 #pragma once
 
+// IWYU pragma: begin_exports
+
 #include "CoreMinimal.h"
 #include "Common/TargetPlatformBase.h"
 #include "Interfaces/IPluginManager.h"
@@ -47,6 +49,8 @@ typedef FMacEditorTargetPlatformSettingsParent FHostPlatformEditorTargetPlatform
 typedef FMacEditorTargetPlatformControlsParent FHostPlatformEditorTargetPlatformControlsParent;
 #endif
 
+// IWYU pragma: end_exports
+
 /**
  * Allows a project to control how packages are cooked when making a cooked editor.
  * 
diff --git a/Engine/Source/Developer/Linux/LinuxTargetPlatformControls/Public/LinuxTargetPlatformControls.h b/Engine/Source/Developer/Linux/LinuxTargetPlatformControls/Public/LinuxTargetPlatformControls.h
--- a/Engine/Source/Developer/Linux/LinuxTargetPlatformControls/Public/LinuxTargetPlatformControls.h
+++ b/Engine/Source/Developer/Linux/LinuxTargetPlatformControls/Public/LinuxTargetPlatformControls.h
@@ -13,7 +13,7 @@
 #include "AnalyticsEventAttribute.h"
 #include "InstalledPlatformInfo.h"
 #include "LinuxTargetDevice.h"
-#include "Linux/LinuxPlatformProperties.h"
+#include "Linux/LinuxPlatformProperties.h" // IWYU pragma: export
 #if WITH_ENGINE
 #include "TextureResource.h"
 #endif // WITH_ENGINE
diff --git a/Engine/Source/Developer/Linux/LinuxTargetPlatformSettings/Public/LinuxTargetPlatformSettings.h b/Engine/Source/Developer/Linux/LinuxTargetPlatformSettings/Public/LinuxTargetPlatformSettings.h
--- a/Engine/Source/Developer/Linux/LinuxTargetPlatformSettings/Public/LinuxTargetPlatformSettings.h
+++ b/Engine/Source/Developer/Linux/LinuxTargetPlatformSettings/Public/LinuxTargetPlatformSettings.h
@@ -11,7 +11,7 @@
 #if WITH_ENGINE
 #include "StaticMeshResources.h"
 #endif // WITH_ENGINE
-#include "Linux/LinuxPlatformProperties.h"
+#include "Linux/LinuxPlatformProperties.h" // IWYU pragma: export
 
 #define LOCTEXT_NAMESPACE "TLinuxTargetPlatformSettings"
 #define LINUX_SECTION_TEXT TEXT("/Script/LinuxTargetPlatform.LinuxTargetSettings")
diff --git a/Engine/Source/Developer/Windows/WindowsTargetPlatfomControls/Public/GenericWindowsTargetPlatformControls.h b/Engine/Source/Developer/Windows/WindowsTargetPlatfomControls/Public/GenericWindowsTargetPlatformControls.h
--- a/Engine/Source/Developer/Windows/WindowsTargetPlatfomControls/Public/GenericWindowsTargetPlatformControls.h
+++ b/Engine/Source/Developer/Windows/WindowsTargetPlatfomControls/Public/GenericWindowsTargetPlatformControls.h
@@ -2,6 +2,8 @@
 
 #pragma once
 
+// IWYU pragma: begin_exports
+
 #include "CoreMinimal.h"
 #include "Common/TargetPlatformControlsBase.h"
 #include "AnalyticsEventAttribute.h"
@@ -14,6 +16,8 @@
 #include "SteamDeck/SteamDeckDevice.h"
 #include "GenericWindowsTargetPlatformSettings.h"
 
+// IWYU pragma: end_exports
+
 #define LOCTEXT_NAMESPACE "TGenericWindowsTargetPlatformControls"
 
 /**
