[LiveCoding]
* Moved some registration logic to run on the game thread

#Synced-CL 48812887

diff --git a/Engine/Source/Developer/Windows/LiveCoding/Private/LiveCodingModule.cpp b/Engine/Source/Developer/Windows/LiveCoding/Private/LiveCodingModule.cpp
--- a/Engine/Source/Developer/Windows/LiveCoding/Private/LiveCodingModule.cpp
+++ b/Engine/Source/Developer/Windows/LiveCoding/Private/LiveCodingModule.cpp
@@ -1216,31 +1216,37 @@ bool FLiveCodingModule::StartLiveCoding(ELiveCodingStartupMode StartupMode)
 		}
 #endif
 
-		// Configure all the current modules. For non-commandlets, schedule it to be done in the first Tick() so we can batch everything together.
-		if (IsRunningCommandlet())
-		{
-			UpdateModules(true);
-		}
-		else
-		{
-			bUpdateModulesInTick = true;
-		}
-
 		if (StartupMode == ELiveCodingStartupMode::Automatic)
 		{
 			ShowConsole(true);
 		}
 
-		// Register a delegate to listen for new modules loaded from this point onwards
-		ModulesChangedDelegateHandle = FModuleManager::Get().OnModulesChanged().AddRaw(this, &FLiveCodingModule::OnModulesChanged);
-
 		// Mark it as started
 		State = EState::RunningAndEnabled;
 
-		FCoreDelegates::OnPreExit.AddLambda( [this]()
+		auto FinalizeStartup = [this]()
+			{
+				// Register a delegate to listen for new modules loaded from this point onwards
+				ModulesChangedDelegateHandle = FModuleManager::Get().OnModulesChanged().AddRaw(this, &FLiveCodingModule::OnModulesChanged);
+
+				bUpdateModulesInTick = true;
+
+				FCoreDelegates::OnPreExit.AddLambda( [this]()
+				{
+					StopLiveCoding();
+				});
+			};
+
+		// Configure all the current modules. For non-commandlets, schedule it to be done in the first Tick() so we can batch everything together.
+		if (IsRunningCommandlet())
 		{
-			StopLiveCoding();
-		});
+			UpdateModules(true);
+			FinalizeStartup();
+		}
+		else
+		{
+			FFunctionGraphTask::CreateAndDispatchWhenReady(MoveTemp(FinalizeStartup), TStatId(), nullptr, ENamedThreads::GameThread);
+		}
 	}
 	return true;
 }
