Cleanup string allocations in shader reflection code, focused on removing ANSI_TO_TCHAR in favor of the FString constructor as well as passing string views to ShaderCompilerCommon functions.

#Synced-CL 49389138

diff --git a/Engine/Source/Developer/Apple/MetalShaderFormat/Private/MetalCompileShaderMSC.cpp b/Engine/Source/Developer/Apple/MetalShaderFormat/Private/MetalCompileShaderMSC.cpp
--- a/Engine/Source/Developer/Apple/MetalShaderFormat/Private/MetalCompileShaderMSC.cpp
+++ b/Engine/Source/Developer/Apple/MetalShaderFormat/Private/MetalCompileShaderMSC.cpp
@@ -177,12 +177,11 @@ static void ProcessReflection(ID3D12ShaderReflection* ShaderReflection, const ui
 		D3D12_SHADER_INPUT_BIND_DESC BindDesc;
 		ShaderReflection->GetResourceBindingDesc(ResourceIndex, &BindDesc);
 
+		const FString BindDescName(BindDesc.Name);
+
 		IRResourceType ResourceType = QuantizeD3DResourceType(BindDesc.Type);
 		bool bIsResourceTyped = IsD3DResourceTypeTyped(BindDesc.Type);
 		const uint32 BindIndex = BindDesc.BindPoint;
-
-		const bool bRootConstantsCB = (FCStringAnsi::Strcmp(BindDesc.Name, "UERootConstants") == 0);
-		const bool bIsRootCB = FCString::Strcmp(ANSI_TO_TCHAR(BindDesc.Name), FShaderParametersMetadata::kRootUniformBufferBindingName) == 0;
 		
 		switch (ResourceType)
 		{
@@ -192,7 +191,7 @@ static void ProcessReflection(ID3D12ShaderReflection* ShaderReflection, const ui
 			else
 				OutputData.InvariantBuffers |= (1 << BindIndex);
 
-			CCHeaderWriter.WriteSRV(ANSI_TO_TCHAR(BindDesc.Name), BindIndex, BindDesc.BindCount);
+			CCHeaderWriter.WriteSRV(*BindDescName, BindIndex, BindDesc.BindCount);
 			NumSRVs = FMath::Max(NumSRVs, BindIndex + BindDesc.BindCount);
 			break;
 
@@ -207,21 +206,20 @@ static void ProcessReflection(ID3D12ShaderReflection* ShaderReflection, const ui
 				OutputData.InvariantBuffers |= (1 << BindIndex);
 			}
 
-			CCHeaderWriter.WriteUAV(ANSI_TO_TCHAR(BindDesc.Name), BindIndex, BindDesc.BindCount);
+			CCHeaderWriter.WriteUAV(*BindDescName, BindIndex, BindDesc.BindCount);
 			NumUAVs = FMath::Max(NumUAVs, BindIndex + BindDesc.BindCount);
 			break;
 
 		case IRResourceTypeSampler:
-			CCHeaderWriter.WriteSamplerState(ANSI_TO_TCHAR(BindDesc.Name), BindIndex);
+			CCHeaderWriter.WriteSamplerState(*BindDescName, BindIndex);
 			NumSamplers = FMath::Max(NumSamplers, BindIndex + BindDesc.BindCount);
 			break;
 
 		case IRResourceTypeCBV:
 		{
-			bool bIsGlobalCB = (FCStringAnsi::Strcmp(BindDesc.Name, "$Globals") == 0);
-			
-			int32 ConstantBufferSize = 0;
-			
+			const bool bIsGlobalCB = (FCStringAnsi::Strcmp(BindDesc.Name, "$Globals") == 0);
+			const bool bIsRootCB = (BindDescName == FShaderParametersMetadata::kRootUniformBufferBindingName);
+
 			OutputData.ConstantBuffers |= (1 << BindIndex);
 			
 			// Global uniform buffer - handled specially as we care about the internal layout
@@ -235,21 +233,19 @@ static void ProcessReflection(ID3D12ShaderReflection* ShaderReflection, const ui
 				D3D12_SHADER_BUFFER_DESC CBDesc;
 				ConstantBuffer->GetDesc(&CBDesc);
 
-				const uint32 CBIndex = BindIndex;
-
-				FString MbrString;
-
 				// Track all of the variables in this constant buffer.
 				for (uint32 ConstantIndex = 0; ConstantIndex < CBDesc.Variables; ConstantIndex++)
 				{
 					ID3D12ShaderReflectionVariable* Variable = ConstantBuffer->GetVariableByIndex(ConstantIndex);
 
-					D3D12_SHADER_VARIABLE_DESC VariableDesc;
+					D3D12_SHADER_VARIABLE_DESC VariableDesc{};
 					Variable->GetDesc(&VariableDesc);
 
 					if (VariableDesc.uFlags & D3D_SVF_USED)
 					{
-						CCHeaderWriter.WritePackedUBField(ANSI_TO_TCHAR(VariableDesc.Name), VariableDesc.StartOffset, VariableDesc.Size);
+						const FString VariableName(VariableDesc.Name);
+
+						CCHeaderWriter.WritePackedUBField(*VariableName, VariableDesc.StartOffset, VariableDesc.Size);
 						
 						const uint32 MbrOffset = VariableDesc.StartOffset / sizeof(float);
 						const uint32 MbrSize = VariableDesc.Size / sizeof(float);
@@ -260,15 +256,12 @@ static void ProcessReflection(ID3D12ShaderReflection* ShaderReflection, const ui
 						unsigned DestCBIndex = 0;
 						InsertRange(CBRanges, BindIndex, SourceOffset, DestSize, DestCBIndex, DestCBPrecision, DestOffset);
 						
-						{
-							HandleReflectedGlobalConstantBufferMember(
-													FString(VariableDesc.Name),
-													BindIndex,
-													VariableDesc.StartOffset,
-													VariableDesc.Size,
-													Output);
-						}
-
+						HandleReflectedGlobalConstantBufferMember(
+							VariableName,
+							BindIndex,
+							VariableDesc.StartOffset,
+							VariableDesc.Size,
+							Output);
 					}
 				}
 			}
@@ -276,10 +269,11 @@ static void ProcessReflection(ID3D12ShaderReflection* ShaderReflection, const ui
 			{
 				ID3D12ShaderReflectionConstantBuffer* ConstantBuffer = ShaderReflection->GetConstantBufferByName(BindDesc.Name);
 
-				D3D12_SHADER_BUFFER_DESC CBDesc;
+				D3D12_SHADER_BUFFER_DESC CBDesc{};
 				ConstantBuffer->GetDesc(&CBDesc);
 
 				const FString UniformBufferName(BindDesc.Name);
+
 				const EUniformBufferMemberReflectionReason Reason = ShouldReflectUniformBufferMembers(Input, UniformBufferName);
 				if (Reason != EUniformBufferMemberReflectionReason::None)
 				{
@@ -292,12 +286,14 @@ static void ProcessReflection(ID3D12ShaderReflection* ShaderReflection, const ui
 
 						if (VariableDesc.uFlags & D3D_SVF_USED)
 						{
+							const FString VariableName(VariableDesc.Name);
+
 							HandleReflectedUniformBufferConstantBufferMember(
 								Input,
 								Reason,
 								UniformBufferName,
 								BindIndex,
-								FString(VariableDesc.Name),
+								VariableName,
 								VariableDesc.StartOffset,
 								VariableDesc.Size,
 								Output
@@ -798,7 +794,7 @@ void FMetalCompileShaderMSC::DoCompileMetalShader(
 			if(bDumpDebugInfo)
 			{
 				ReflectionJSON = IRShaderReflectionCopyJSONString(AirReflection);
-				FString ReflectionString = ANSI_TO_TCHAR(ReflectionJSON);
+				FString ReflectionString(ReflectionJSON);
 				DumpDebugShaderText(Input, ReflectionString, TEXT("reflection.json"));
 				checkSlow(ReflectionJSON);
 			}
diff --git a/Engine/Source/Developer/Apple/MetalShaderFormat/Private/MetalCompileShaderSPIRV.cpp b/Engine/Source/Developer/Apple/MetalShaderFormat/Private/MetalCompileShaderSPIRV.cpp
--- a/Engine/Source/Developer/Apple/MetalShaderFormat/Private/MetalCompileShaderSPIRV.cpp
+++ b/Engine/Source/Developer/Apple/MetalShaderFormat/Private/MetalCompileShaderSPIRV.cpp
@@ -431,18 +431,20 @@ void FMetalCompileShaderSPIRV::DoCompileMetalShader(
 				// Which is necessary to match the D3D binding scheme.
 				for (SpvReflectDescriptorBinding* Binding : Bindings)
 				{
-					if (Binding->resource_type != SPV_REFLECT_RESOURCE_FLAG_CBV && ResourceTable.Contains(UTF8_TO_TCHAR(Binding->name)))
+					const FString BindingName(Binding->name);
+
+					if (Binding->resource_type != SPV_REFLECT_RESOURCE_FLAG_CBV && ResourceTable.Contains(BindingName))
 					{
 						ResourceBindings.Add(Binding);
 						
-						FMetalResourceTableEntry Entry = ResourceTable.FindRef(UTF8_TO_TCHAR(Binding->name));
+						FMetalResourceTableEntry Entry = ResourceTable.FindRef(BindingName);
 						UsedSets.Add(FString(Entry.GetUniformBufferName()));
 						
 						continue;
 					}
 
 					// Add descriptor binding to argument bindings if it's a constant buffer with a name from 'TableNames'. Otherwise, add to common binding container.
-					if (Binding->resource_type == SPV_REFLECT_RESOURCE_FLAG_CBV && Binding->accessed && TableNames.Contains(UTF8_TO_TCHAR(Binding->name)))
+					if (Binding->resource_type == SPV_REFLECT_RESOURCE_FLAG_CBV && Binding->accessed && TableNames.Contains(BindingName))
 					{
 						check(Binding->descriptor_type == SPV_REFLECT_DESCRIPTOR_TYPE_UNIFORM_BUFFER);
 						ArgumentBindings.Add(Binding);
@@ -570,7 +572,7 @@ void FMetalCompileShaderSPIRV::DoCompileMetalShader(
 				{
 					for (auto const& Binding : ResourceBindings)
 					{
-						FMetalResourceTableEntry* Entry = ResourceTable.Find(UTF8_TO_TCHAR(Binding->name));
+						FMetalResourceTableEntry* Entry = ResourceTable.Find(FString(Binding->name));
 						FString EntryUniformBufferName(Entry->GetUniformBufferName());
 						auto* ResourceArray = IABs.Find(EntryUniformBufferName);
 						if (!IABTier1Index.Contains(EntryUniformBufferName))
@@ -608,7 +610,7 @@ void FMetalCompileShaderSPIRV::DoCompileMetalShader(
 				
 				for (auto const& Binding : ResourceBindings)
 				{
-					FMetalResourceTableEntry* Entry = ResourceTable.Find(UTF8_TO_TCHAR(Binding->name));
+					FMetalResourceTableEntry* Entry = ResourceTable.Find(FString(Binding->name));
 					FString EntryUniformBufferName(Entry->GetUniformBufferName());
 					
 					for (uint32 j = 0; j < (uint32)TableNames.Num(); j++)
@@ -662,7 +664,7 @@ void FMetalCompileShaderSPIRV::DoCompileMetalShader(
 				
 				for (auto const& Pair : IABs)
 				{
-					FString Name = Pair.Key;
+					const FString& Name = Pair.Key;
 					auto const& ResourceArray = Pair.Value;
 					uint32 SetIndex = ResourceArray[0].SetIndex + IABOffsetIndex;
 
@@ -885,7 +887,7 @@ void FMetalCompileShaderSPIRV::DoCompileMetalShader(
 							
 						FString TypeQualifier;
 
-						auto const type = *Var->type_description;
+						const SpvReflectTypeDescription& type = *Var->type_description;
 						uint32_t masked_type = type.type_flags & 0xF;
 							
 						switch (masked_type) {
diff --git a/Engine/Source/Developer/Apple/MetalShaderFormat/Private/MetalShaderCompiler.cpp b/Engine/Source/Developer/Apple/MetalShaderFormat/Private/MetalShaderCompiler.cpp
--- a/Engine/Source/Developer/Apple/MetalShaderFormat/Private/MetalShaderCompiler.cpp
+++ b/Engine/Source/Developer/Apple/MetalShaderFormat/Private/MetalShaderCompiler.cpp
@@ -323,7 +323,7 @@ void BuildMetalShaderOutput(
 	if (Frequency == SF_Vertex)
 	{
 		static const FString AttributePrefix = TEXT("in_ATTRIBUTE");
-		for (auto& Input : CCHeader.Inputs)
+		for (const CrossCompiler::FHlslccHeader::FInOut& Input : CCHeader.Inputs)
 		{
 			// Only process attributes.
 			if (Input.Name.StartsWith(AttributePrefix))
@@ -341,7 +341,7 @@ void BuildMetalShaderOutput(
 	// Only outputs for pixel shaders must be tracked.
 	if (Frequency == SF_Pixel)
 	{
-		for (auto& Output : CCHeader.Outputs)
+		for (const CrossCompiler::FHlslccHeader::FInOut& Output : CCHeader.Outputs)
 		{
 			// Handle targets.
 			if (Output.Name.StartsWith(TargetPrefix))
@@ -375,7 +375,7 @@ void BuildMetalShaderOutput(
 
 	// Then 'normal' uniform buffers.
 	bool bOutOfBounds = false;
-	for (auto& UniformBlock : CCHeader.UniformBlocks)
+	for (const CrossCompiler::FHlslccHeader::FAttribute& UniformBlock : CCHeader.UniformBlocks)
 	{
 		uint16 UBIndex = UniformBlock.Index;
 		if (UBIndex >= Header.Bindings.NumUniformBuffers)
@@ -402,7 +402,7 @@ void BuildMetalShaderOutput(
 	// Packed global uniforms
 	const uint16 BytesPerComponent = 4;
 	TMap<ANSICHAR, uint16> PackedGlobalArraySize;
-	for (auto& PackedGlobal : CCHeader.PackedGlobals)
+	for (const CrossCompiler::FHlslccHeader::FPackedGlobal& PackedGlobal : CCHeader.PackedGlobals)
 	{
 		HandleReflectedGlobalConstantBufferMember(
 			PackedGlobal.Name,
@@ -423,9 +423,9 @@ void BuildMetalShaderOutput(
 	
 	// Packed Uniform Buffers
 	TMap<int, TMap<CrossCompiler::EPackedTypeName, uint16> > PackedUniformBuffersSize;
-	for (auto& PackedUB : CCHeader.PackedUBs)
+	for (const CrossCompiler::FHlslccHeader::FPackedUB& PackedUB : CCHeader.PackedUBs)
 	{
-		for (auto& Member : PackedUB.Members)
+		for (const CrossCompiler::FHlslccHeader::FPackedUB::FMember& Member : PackedUB.Members)
 		{
 			uint32 ConstantBufferIndex = bBindlessEnabled ? 0 : (uint32)CrossCompiler::EPackedTypeName::HighP;
 
@@ -485,7 +485,7 @@ void BuildMetalShaderOutput(
 	
 	// Then samplers.
 	TMap<FString, uint32> SamplerMap;
-	for (auto& Sampler : CCHeader.Samplers)
+	for (const CrossCompiler::FHlslccHeader::FSampler& Sampler : CCHeader.Samplers)
 	{
 		HandleReflectedShaderResource(Sampler.Name, Sampler.Offset, Sampler.Count, ShaderOutput);
 
@@ -500,7 +500,7 @@ void BuildMetalShaderOutput(
 	Header.Bindings.NumSamplers = CCHeader.SamplerStates.Num();
 
 	// Then UAVs (images in Metal)
-	for (auto& UAV : CCHeader.UAVs)
+	for (const CrossCompiler::FHlslccHeader::FUAV& UAV : CCHeader.UAVs)
 	{
 		HandleReflectedShaderUAV(UAV.Name, UAV.Offset, UAV.Count, ShaderOutput);
 
@@ -510,7 +510,7 @@ void BuildMetalShaderOutput(
 			);
 	}
 
-	for (auto& SamplerState : CCHeader.SamplerStates)
+	for (const CrossCompiler::FHlslccHeader::FAttribute& SamplerState : CCHeader.SamplerStates)
 	{
 		if (!SamplerMap.Contains(SamplerState.Name))
 		{
@@ -538,7 +538,7 @@ void BuildMetalShaderOutput(
 		// Only needed for VS Input (to generate the stage-in function used to convert inputs).
         if (Frequency == SF_Vertex)
         {
-            Header.Bindings.IRConverterReflectionJSON = ANSI_TO_TCHAR(ShaderReflectionJSON);
+            Header.Bindings.IRConverterReflectionJSON = FString(ShaderReflectionJSON);
 			IRShaderReflectionReleaseString(ShaderReflectionJSON);
 			check(ShaderReflectionJSON && Header.Bindings.IRConverterReflectionJSON.Len() > 0);
         }
diff --git a/Engine/Source/Developer/ShaderCompilerCommon/Private/ShaderCompilerCommon.cpp b/Engine/Source/Developer/ShaderCompilerCommon/Private/ShaderCompilerCommon.cpp
--- a/Engine/Source/Developer/ShaderCompilerCommon/Private/ShaderCompilerCommon.cpp
+++ b/Engine/Source/Developer/ShaderCompilerCommon/Private/ShaderCompilerCommon.cpp
@@ -579,7 +579,7 @@ bool UE::ShaderCompilerCommon::RemoveDeadCode(FShaderSource& InOutPreprocessedSh
 }
 
 void HandleReflectedGlobalConstantBufferMember(
-	const FString& InMemberName,
+	FStringView InMemberName,
 	uint32 ConstantBufferIndex,
 	int32 ReflectionOffset,
 	int32 ReflectionSize,
@@ -648,7 +648,7 @@ void HandleReflectedUniformBufferConstantBufferMember(
 void HandleReflectedRootConstantBufferMember(
 	const FShaderCompilerInput& Input,
 	const FShaderParameterParser& ShaderParameterParser,
-	const FString& InMemberName,
+	FStringView InMemberName,
 	int32 ReflectionOffset,
 	int32 ReflectionSize,
 	FShaderCompilerOutput& Output
@@ -760,7 +760,7 @@ EUniformBufferMemberReflectionReason ShouldReflectUniformBufferMembers(const FSh
 }
 
 void HandleReflectedShaderResource(
-	const FString& ResourceName,
+	FStringView ResourceName,
 	int32 BindOffset,
 	int32 ReflectionSlot,
 	int32 BindCount,
@@ -778,7 +778,7 @@ void HandleReflectedShaderResource(
 
 void UpdateStructuredBufferStride(
 	const FShaderCompilerInput& Input,
-	const FString& ResourceName,
+	FStringView ResourceName,
 	uint16 BindPoint,
 	uint16 Stride,
 	FShaderCompilerOutput& CompilerOutput
@@ -790,7 +790,13 @@ void UpdateStructuredBufferStride(
 	}
 	else
 	{
-		FString ErrorMessage = FString::Printf(TEXT("%s: Failed to set stride on parameter %s: Bind point %d, Stride %d"), *Input.GenerateShaderName(), *ResourceName, BindPoint, Stride);
+		FString ErrorMessage =
+			FString::Printf(TEXT("%s: Failed to set stride on parameter %.*s: Bind point %d, Stride %d"),
+				*Input.GenerateShaderName(),
+				ResourceName.Len(), ResourceName.GetData(),
+				BindPoint,
+				Stride
+			);
 		CompilerOutput.Errors.Add(FShaderCompilerError(*ErrorMessage));
 	}
 }
@@ -826,7 +832,7 @@ void AddShaderValidationUBSize(uint16 BindPoint,
 }
  
 void HandleReflectedShaderUAV(
-	const FString& UAVName,
+	FStringView UAVName,
 	int32 BindOffset,
 	int32 ReflectionSlot,
 	int32 BindCount,
@@ -843,7 +849,7 @@ void HandleReflectedShaderUAV(
 }
 
 void HandleReflectedShaderSampler(
-	const FString& SamplerName,
+	FStringView SamplerName,
 	int32 BindOffset,
 	int32 ReflectionSlot,
 	int32 BindCount,
@@ -876,7 +882,7 @@ void AddNoteToDisplayShaderParameterStructureOnCppSide(
 void AddUnboundShaderParameterError(
 	const FShaderCompilerInput& CompilerInput,
 	const FShaderParameterParser& ShaderParameterParser,
-	const FString& ParameterBindingName,
+	FStringView ParameterBindingName,
 	FShaderCompilerOutput& CompilerOutput)
 {
 	check(CompilerInput.RootParametersStructure);
@@ -885,8 +891,8 @@ void AddUnboundShaderParameterError(
 	check(!Member.bIsBindable);
 
 	FShaderCompilerError Error(FString::Printf(
-		TEXT("Error: Shader parameter %s could not be bound to %s's shader parameter structure %s."),
-		*ParameterBindingName,
+		TEXT("Error: Shader parameter %.*s could not be bound to %s's shader parameter structure %s."),
+		ParameterBindingName.Len(), ParameterBindingName.GetData(),
 		*CompilerInput.ShaderName,
 		CompilerInput.RootParametersStructure->GetStructTypeName()));
 	ShaderParameterParser.GetParameterFileAndLine(Member, Error.ErrorVirtualFilePath, Error.ErrorLineString);
diff --git a/Engine/Source/Developer/ShaderCompilerCommon/Public/ShaderCompilerCommon.h b/Engine/Source/Developer/ShaderCompilerCommon/Public/ShaderCompilerCommon.h
--- a/Engine/Source/Developer/ShaderCompilerCommon/Public/ShaderCompilerCommon.h
+++ b/Engine/Source/Developer/ShaderCompilerCommon/Public/ShaderCompilerCommon.h
@@ -209,7 +209,7 @@ SHADERCOMPILERCOMMON_API EUniformBufferMemberReflectionReason ShouldReflectUnifo
 );
 
 extern SHADERCOMPILERCOMMON_API void HandleReflectedGlobalConstantBufferMember(
-	const FString& MemberName,
+	FStringView MemberName,
 	uint32 ConstantBufferIndex,
 	int32 ReflectionOffset,
 	int32 ReflectionSize,
@@ -230,7 +230,7 @@ extern SHADERCOMPILERCOMMON_API void HandleReflectedUniformBufferConstantBufferM
 extern SHADERCOMPILERCOMMON_API void HandleReflectedRootConstantBufferMember(
 	const FShaderCompilerInput& Input,
 	const FShaderParameterParser& ShaderParameterParser,
-	const FString& MemberName,
+	FStringView MemberName,
 	int32 ReflectionOffset,
 	int32 ReflectionSize,
 	FShaderCompilerOutput& CompilerOutput
@@ -261,7 +261,7 @@ inline void HandleReflectedUniformBuffer(const FShaderCompilerInput& Input, FStr
 }
 
 extern SHADERCOMPILERCOMMON_API void HandleReflectedShaderResource(
-	const FString& ResourceName,
+	FStringView ResourceName,
 	int32 BindOffset,
 	int32 ReflectionSlot,
 	int32 BindCount,
@@ -270,7 +270,7 @@ extern SHADERCOMPILERCOMMON_API void HandleReflectedShaderResource(
 
 extern SHADERCOMPILERCOMMON_API void UpdateStructuredBufferStride(
 	const FShaderCompilerInput& Input,
-	const FString& ResourceName,
+	FStringView ResourceName,
 	uint16 BindPoint,
 	uint16 Stride,
 	FShaderCompilerOutput& CompilerOutput
@@ -291,48 +291,48 @@ extern SHADERCOMPILERCOMMON_API void AddShaderValidationUBSize(
 	uint32_t Size,
 	FShaderCompilerOutput& CompilerOutput);
 
-inline void HandleReflectedShaderResource(const FString& ResourceName, int32 ReflectionSlot, int32 BindCount, FShaderCompilerOutput& CompilerOutput)
+inline void HandleReflectedShaderResource(FStringView ResourceName, int32 ReflectionSlot, int32 BindCount, FShaderCompilerOutput& CompilerOutput)
 {
 	HandleReflectedShaderResource(ResourceName, 0, ReflectionSlot, BindCount, CompilerOutput);
 }
 
-inline void HandleReflectedShaderResource(const FString& ResourceName, int32 ReflectionSlot, FShaderCompilerOutput& CompilerOutput)
+inline void HandleReflectedShaderResource(FStringView ResourceName, int32 ReflectionSlot, FShaderCompilerOutput& CompilerOutput)
 {
 	HandleReflectedShaderResource(ResourceName, ReflectionSlot, 1, CompilerOutput);
 }
 
 extern SHADERCOMPILERCOMMON_API void HandleReflectedShaderUAV(
-	const FString& UAVName,
+	FStringView UAVName,
 	int32 BindOffset,
 	int32 ReflectionSlot,
 	int32 BindCount,
 	FShaderCompilerOutput& CompilerOutput
 );
 
-inline void HandleReflectedShaderUAV(const FString& UAVName, int32 ReflectionSlot, int32 BindCount, FShaderCompilerOutput& CompilerOutput)
+inline void HandleReflectedShaderUAV(FStringView UAVName, int32 ReflectionSlot, int32 BindCount, FShaderCompilerOutput& CompilerOutput)
 {
 	HandleReflectedShaderUAV(UAVName, 0, ReflectionSlot, BindCount, CompilerOutput);
 }
 
-inline void HandleReflectedShaderUAV(const FString& UAVName, int32 ReflectionSlot, FShaderCompilerOutput& CompilerOutput)
+inline void HandleReflectedShaderUAV(FStringView UAVName, int32 ReflectionSlot, FShaderCompilerOutput& CompilerOutput)
 {
 	HandleReflectedShaderUAV(UAVName, ReflectionSlot, 1, CompilerOutput);
 }
 
 extern SHADERCOMPILERCOMMON_API void HandleReflectedShaderSampler(
-	const FString& SamplerName,
+	FStringView SamplerName,
 	int32 BindOffset,
 	int32 ReflectionSlot,
 	int32 BindCount,
 	FShaderCompilerOutput& CompilerOutput
 );
 
-inline void HandleReflectedShaderSampler(const FString& SamplerName, int32 ReflectionSlot, int32 BindCount, FShaderCompilerOutput& CompilerOutput)
+inline void HandleReflectedShaderSampler(FStringView SamplerName, int32 ReflectionSlot, int32 BindCount, FShaderCompilerOutput& CompilerOutput)
 {
 	HandleReflectedShaderSampler(SamplerName, 0, ReflectionSlot, BindCount, CompilerOutput);
 }
 
-inline void HandleReflectedShaderSampler(const FString& SamplerName, int32 ReflectionSlot, FShaderCompilerOutput& CompilerOutput)
+inline void HandleReflectedShaderSampler(FStringView SamplerName, int32 ReflectionSlot, FShaderCompilerOutput& CompilerOutput)
 {
 	HandleReflectedShaderSampler(SamplerName, ReflectionSlot, 1, CompilerOutput);
 }
@@ -346,7 +346,7 @@ extern SHADERCOMPILERCOMMON_API void AddNoteToDisplayShaderParameterStructureOnC
 extern SHADERCOMPILERCOMMON_API void AddUnboundShaderParameterError(
 	const FShaderCompilerInput& CompilerInput,
 	const FShaderParameterParser& ShaderParameterParser,
-	const FString& ParameterBindingName,
+	FStringView ParameterBindingName,
 	FShaderCompilerOutput& CompilerOutput);
 
 // Convert generated UniformBuffer code and references into something the shader compilers can use.
diff --git a/Engine/Source/Developer/VulkanShaderFormat/Private/SpirVShaderCompiler.inl b/Engine/Source/Developer/VulkanShaderFormat/Private/SpirVShaderCompiler.inl
--- a/Engine/Source/Developer/VulkanShaderFormat/Private/SpirVShaderCompiler.inl
+++ b/Engine/Source/Developer/VulkanShaderFormat/Private/SpirVShaderCompiler.inl
@@ -370,7 +370,7 @@ static void BuildShaderOutput(
 				continue;
 			}
 
-			const FString InputAttrName(ANSI_TO_TCHAR(Attribute->semantic));
+			const FString InputAttrName(Attribute->semantic);
 			if (InputAttrName.StartsWith(AttributePrefix))
 			{
 				const uint32 AttributeIndex = ParseNumber(*InputAttrName + AttributePrefix.Len(), /*bEmptyIsZero:*/ true);
@@ -400,7 +400,7 @@ static void BuildShaderOutput(
 			else
 			{
 				// Only targets for pixel shaders must be tracked.
-				const FString OutputAttrName(ANSI_TO_TCHAR(Attribute->semantic));
+				const FString OutputAttrName(Attribute->semantic);
 				if (OutputAttrName.StartsWith(TargetPrefix))
 				{
 					const uint32 TargetIndex = ParseNumber(*OutputAttrName + TargetPrefix.Len(), /*bEmptyIsZero:*/ true);
@@ -567,7 +567,7 @@ static void GatherSpirvReflectionBindings(
 			for (int32 Index = BindingArray.Num() - 1; Index >= 0; --Index)
 			{
 				const SpvReflectDescriptorBinding* pBinding = BindingArray[Index];
-				const FString BindingName(ANSI_TO_TCHAR(pBinding->name));
+				const FString BindingName(pBinding->name);
 				if (BindingName.StartsWith(HeapPrefix))
 				{
 					const uint32 Binding = 0;  // single bindless heap per descriptor set
@@ -596,7 +596,7 @@ static void GatherSpirvReflectionBindings(
 			for (int32 Index = OutBindings.UniformBuffers.Num() - 1; Index >= 0; --Index)
 			{
 				const SpvReflectDescriptorBinding* pBinding = OutBindings.UniformBuffers[Index];
-				const FString BindingName(ANSI_TO_TCHAR(pBinding->name));
+				const FString BindingName(pBinding->name);
 				if (BindingName.StartsWith(kBindlessCBPrefix))
 				{
 					check(InternalState.bUseBindlessUniformBuffer || InternalState.bUseStaticUniformBufferBindings);
@@ -693,7 +693,7 @@ static bool BuildShaderOutputFromSpirv(
 	FSpirvReflectBindings Bindings;
 	GatherSpirvReflectionBindings(Reflection, Bindings, SerializedOutput.UsedBindlessUB, InternalState);
 
-	const FString UBOGlobalsNameSpv(ANSI_TO_TCHAR(CrossCompiler::FShaderConductorContext::GetIdentifierTable().GlobalsUniformBuffer));
+	const FString UBOGlobalsNameSpv(CrossCompiler::FShaderConductorContext::GetIdentifierTable().GlobalsUniformBuffer);
 	const FString UBORootParamNameSpv(FShaderParametersMetadata::kRootUniformBufferBindingName);
 
 	TBitArray<> UsedUniformBufferSlots;
@@ -728,7 +728,7 @@ static bool BuildShaderOutputFromSpirv(
 				checkf(!InternalState.bSupportsBindless || (DescriptorType == VK_DESCRIPTOR_TYPE_UNIFORM_BUFFER),
 					TEXT("Bindless shaders should only have uniform buffers."));
 
-				const FString ResourceName(ANSI_TO_TCHAR(Binding->name));
+				const FString ResourceName(Binding->name);
 
 				const bool bIsGlobalOrRootBuffer = ((UBOGlobalsNameSpv == ResourceName) || (UBORootParamNameSpv == ResourceName));
 				if ((bIsPackedUniformBuffer && !bIsGlobalOrRootBuffer) || ((!bIsPackedUniformBuffer) && bIsGlobalOrRootBuffer))
@@ -814,7 +814,7 @@ static bool BuildShaderOutputFromSpirv(
 							{
 								const SpvReflectBlockVariable& Member = Binding->block.members[MemberIndex];
 
-								FString MemberName(Member.name);
+								const FString MemberName(Member.name);
 								FStringView AdjustedMemberName(MemberName);
 
 								const EShaderParameterType BindlessParameterType = FShaderParameterParser::ParseAndRemoveBindlessParameterPrefix(AdjustedMemberName);
diff --git a/Engine/Source/Developer/Windows/ShaderFormatD3D/Private/D3DShaderCompiler.inl b/Engine/Source/Developer/Windows/ShaderFormatD3D/Private/D3DShaderCompiler.inl
--- a/Engine/Source/Developer/Windows/ShaderFormatD3D/Private/D3DShaderCompiler.inl
+++ b/Engine/Source/Developer/Windows/ShaderFormatD3D/Private/D3DShaderCompiler.inl
@@ -91,7 +91,7 @@ template <typename ID3D1xShaderReflection, typename D3D1x_SHADER_DESC, typename
 	// Add parameters for shader resources (constant buffers, textures, samplers, etc. */
 	for (uint32 ResourceIndex = 0; ResourceIndex < ShaderDesc.BoundResources; ResourceIndex++)
 	{
-		D3D1x_SHADER_INPUT_BIND_DESC BindDesc;
+		D3D1x_SHADER_INPUT_BIND_DESC BindDesc{};
 		Reflector->GetResourceBindingDesc(ResourceIndex, &BindDesc);
 
 		if (!IsCompatibleBinding(BindDesc, BindingSpace))
@@ -99,11 +99,14 @@ template <typename ID3D1xShaderReflection, typename D3D1x_SHADER_DESC, typename
 			continue;
 		}
 
+		const FString BindDescName(BindDesc.Name);
+
 		if (BindDesc.Type == D3D_SIT_CBUFFER || BindDesc.Type == D3D_SIT_TBUFFER)
 		{
 			const uint32 CBIndex = BindDesc.BindPoint;
 			ID3D1xShaderReflectionConstantBuffer* ConstantBuffer = Reflector->GetConstantBufferByName(BindDesc.Name);
-			D3D1x_SHADER_BUFFER_DESC CBDesc;
+
+			D3D1x_SHADER_BUFFER_DESC CBDesc{};
 			ConstantBuffer->GetDesc(&CBDesc);
 
 			const FString ConstantBufferName(CBDesc.Name);
@@ -121,14 +124,18 @@ template <typename ID3D1xShaderReflection, typename D3D1x_SHADER_DESC, typename
 					for (uint32 ConstantIndex = 0; ConstantIndex < CBDesc.Variables; ConstantIndex++)
 					{
 						ID3D1xShaderReflectionVariable* Variable = ConstantBuffer->GetVariableByIndex(ConstantIndex);
-						D3D1x_SHADER_VARIABLE_DESC VariableDesc;
+
+						D3D1x_SHADER_VARIABLE_DESC VariableDesc{};
 						Variable->GetDesc(&VariableDesc);
+
 						if (VariableDesc.uFlags & D3D_SVF_USED)
 						{
+							const FString VariableName(VariableDesc.Name);
+
 							AddUnboundShaderParameterError(
 								Input,
 								ShaderParameterParser,
-								ANSI_TO_TCHAR(VariableDesc.Name),
+								VariableName,
 								Output);
 						}
 					}
@@ -139,20 +146,23 @@ template <typename ID3D1xShaderReflection, typename D3D1x_SHADER_DESC, typename
 					for (uint32 ConstantIndex = 0; ConstantIndex < CBDesc.Variables; ConstantIndex++)
 					{
 						ID3D1xShaderReflectionVariable* Variable = ConstantBuffer->GetVariableByIndex(ConstantIndex);
-						D3D1x_SHADER_VARIABLE_DESC VariableDesc;
+
+						D3D1x_SHADER_VARIABLE_DESC VariableDesc{};
 						Variable->GetDesc(&VariableDesc);
+
 						if (VariableDesc.uFlags & D3D_SVF_USED)
 						{
-							CompileData.bGlobalUniformBufferUsed = true;
+							const FString VariableName(VariableDesc.Name);
 
 							HandleReflectedGlobalConstantBufferMember(
-								FString(VariableDesc.Name),
+								VariableName,
 								CBIndex,
 								VariableDesc.StartOffset,
 								VariableDesc.Size,
 								Output
 							);
 
+							CompileData.bGlobalUniformBufferUsed = true;
 							CompileData.UsedUniformBufferSlots[CBIndex] = true;
 						}
 					}
@@ -173,14 +183,18 @@ template <typename ID3D1xShaderReflection, typename D3D1x_SHADER_DESC, typename
 					for (uint32 ConstantIndex = 0; ConstantIndex < CBDesc.Variables; ConstantIndex++)
 					{
 						ID3D1xShaderReflectionVariable* Variable = ConstantBuffer->GetVariableByIndex(ConstantIndex);
-						D3D1x_SHADER_VARIABLE_DESC VariableDesc;
+
+						D3D1x_SHADER_VARIABLE_DESC VariableDesc{};
 						Variable->GetDesc(&VariableDesc);
+
 						if (VariableDesc.uFlags & D3D_SVF_USED)
 						{
+							const FString VariableName(VariableDesc.Name);
+
 							HandleReflectedRootConstantBufferMember(
 								Input,
 								ShaderParameterParser,
-								FString(VariableDesc.Name),
+								VariableName,
 								VariableDesc.StartOffset,
 								VariableDesc.Size,
 								Output
@@ -224,19 +238,19 @@ template <typename ID3D1xShaderReflection, typename D3D1x_SHADER_DESC, typename
 					{
 						ID3D1xShaderReflectionVariable* Variable = ConstantBuffer->GetVariableByIndex(ConstantIndex);
 
-						D3D1x_SHADER_VARIABLE_DESC VariableDesc;
+						D3D1x_SHADER_VARIABLE_DESC VariableDesc{};
 						Variable->GetDesc(&VariableDesc);
 
 						if (VariableDesc.uFlags & D3D_SVF_USED)
 						{
-							const FString MemberName(VariableDesc.Name);
+							const FString VariableName(VariableDesc.Name);
 
 							HandleReflectedUniformBufferConstantBufferMember(
 								Input,
 								Reason,
 								ConstantBufferName,
 								CBIndex,
-								MemberName,
+								VariableName,
 								VariableDesc.StartOffset,
 								VariableDesc.Size,
 								Output
@@ -272,7 +286,7 @@ template <typename ID3D1xShaderReflection, typename D3D1x_SHADER_DESC, typename
 			}
 			else if (ParameterType == EShaderParameterType::Sampler)
 			{
-				HandleReflectedShaderSampler(FString(BindDesc.Name), BindDesc.BindPoint, Output);
+				HandleReflectedShaderSampler(BindDescName, BindDesc.BindPoint, Output);
 				CompileData.NumSamplers = FMath::Max(CompileData.NumSamplers, BindDesc.BindPoint + BindCount);
 			}
 			else
@@ -280,7 +294,7 @@ template <typename ID3D1xShaderReflection, typename D3D1x_SHADER_DESC, typename
 				EShaderCodeResourceBindingType ResourceBindingType = D3DBindDescToShaderCodeResourceBinding(BindDesc);
 				AddShaderValidationSRVType(BindDesc.BindPoint, ResourceBindingType, Output);
 
-				HandleReflectedShaderResource(FString(BindDesc.Name), BindDesc.BindPoint, Output);
+				HandleReflectedShaderResource(BindDescName, BindDesc.BindPoint, Output);
 				CompileData.NumSRVs = FMath::Max(CompileData.NumSRVs, BindDesc.BindPoint + BindCount);
 			}
 		}
@@ -326,14 +340,13 @@ template <typename ID3D1xShaderReflection, typename D3D1x_SHADER_DESC, typename
 				EShaderCodeResourceBindingType ResourceBindingType = D3DBindDescToShaderCodeResourceBinding(BindDesc);
 				AddShaderValidationUAVType(BindDesc.BindPoint, ResourceBindingType, Output);
 
-				HandleReflectedShaderUAV(FString(BindDesc.Name), BindDesc.BindPoint, Output);
+				HandleReflectedShaderUAV(BindDescName, BindDesc.BindPoint, Output);
 				CompileData.NumUAVs = FMath::Max(CompileData.NumUAVs, BindDesc.BindPoint + BindCount);
 			}
 		}
 		else if (BindDesc.Type == D3D_SIT_STRUCTURED || BindDesc.Type == D3D_SIT_BYTEADDRESS)
 		{
 			check(BindDesc.BindCount == 1);
-			FString BindDescName(BindDesc.Name);
 
 			EShaderCodeResourceBindingType ResourceBindingType = D3DBindDescToShaderCodeResourceBinding(BindDesc);
 			AddShaderValidationSRVType(BindDesc.BindPoint, ResourceBindingType, Output);
@@ -342,7 +355,7 @@ template <typename ID3D1xShaderReflection, typename D3D1x_SHADER_DESC, typename
 
 			// https://learn.microsoft.com/en-us/windows/win32/api/d3d12shader/ns-d3d12shader-d3d12_shader_input_bind_desc
 			// If the shader resource is a structured buffer, the field contains the stride of the type in bytes
-			if ( BindDesc.Type == D3D_SIT_STRUCTURED)
+			if (BindDesc.Type == D3D_SIT_STRUCTURED)
 			{
 				UpdateStructuredBufferStride(Input, BindDescName, BindDesc.BindPoint, BindDesc.NumSamples, Output);
 			}
@@ -357,7 +370,7 @@ template <typename ID3D1xShaderReflection, typename D3D1x_SHADER_DESC, typename
 			EShaderCodeResourceBindingType ResourceBindingType = D3DBindDescToShaderCodeResourceBinding(BindDesc);
 			AddShaderValidationSRVType(BindDesc.BindPoint, ResourceBindingType, Output);
 
-			HandleReflectedShaderResource(FString(BindDesc.Name), BindDesc.BindPoint, Output);
+			HandleReflectedShaderResource(BindDescName, BindDesc.BindPoint, Output);
 			CompileData.NumSRVs = FMath::Max(CompileData.NumSRVs, BindDesc.BindPoint + 1);
 		}
 	}
@@ -626,7 +639,7 @@ bool RemoveUnusedInterpolators(
 				{
 					if (ParamDesc.ReadWriteMask != 0)
 					{
-						FString SemanticName = ANSI_TO_TCHAR(ParamDesc.SemanticName);
+						FString SemanticName(ParamDesc.SemanticName);
 
 						CompileData.ShaderInputs.AddUnique(SemanticName);
 
@@ -648,7 +661,7 @@ bool RemoveUnusedInterpolators(
 					// if (ParamDesc.ReadWriteMask != 0)
 					{
 						// Keep system values
-						CompileData.ShaderInputs.AddUnique(FString(ANSI_TO_TCHAR(ParamDesc.SemanticName)));
+						CompileData.ShaderInputs.AddUnique(FString(ParamDesc.SemanticName));
 					}
 				}
 			}
diff --git a/Engine/Source/Runtime/RenderCore/Private/ShaderParameterParser.cpp b/Engine/Source/Runtime/RenderCore/Private/ShaderParameterParser.cpp
--- a/Engine/Source/Runtime/RenderCore/Private/ShaderParameterParser.cpp
+++ b/Engine/Source/Runtime/RenderCore/Private/ShaderParameterParser.cpp
@@ -1266,7 +1266,7 @@ bool FShaderParameterParser::ParseAndModify(const FShaderCompilerInput& Compiler
 
 void FShaderParameterParser::ValidateShaderParameterType(
 	const FShaderCompilerInput& CompilerInput,
-	const FString& ShaderBindingName,
+	FStringView ShaderBindingName,
 	int32 ReflectionOffset,
 	int32 ReflectionSize,
 	bool bPlatformSupportsPrecisionModifier,
@@ -1361,9 +1361,9 @@ void FShaderParameterParser::ValidateShaderParameterType(
 
 			FShaderCompilerError Error;
 			Error.StrippedErrorMessage = FString::Printf(
-				TEXT("Error: Type %.*s of shader parameter %s in shader mismatch the shader parameter structure: %s expects a %s"),
+				TEXT("Error: Type %.*s of shader parameter %.*s in shader mismatch the shader parameter structure: %s expects a %s"),
 				ParsedParameter.ParsedType.Len(), ParsedParameter.ParsedType.GetData(),
-				*ShaderBindingName,
+				ShaderBindingName.Len(), ShaderBindingName.GetData(),
 				*CppCodeName,
 				*ExpectedShaderType);
 			GetParameterFileAndLine(ParsedParameter, Error.ErrorVirtualFilePath, Error.ErrorLineString);
@@ -1382,8 +1382,8 @@ void FShaderParameterParser::ValidateShaderParameterType(
 
 		FShaderCompilerError Error;
 		Error.StrippedErrorMessage = FString::Printf(
-			TEXT("Error: The size required to bind shader parameter %s is %i bytes, smaller than %s that is %i bytes in the parameter structure."),
-			*ShaderBindingName,
+			TEXT("Error: The size required to bind shader parameter %.*s is %i bytes, smaller than %s that is %i bytes in the parameter structure."),
+			ShaderBindingName.Len(), ShaderBindingName.GetData(),
 			ReflectionSize,
 			*CppCodeName,
 			ParsedParameter.MemberSize);
diff --git a/Engine/Source/Runtime/RenderCore/Public/ShaderParameterParser.h b/Engine/Source/Runtime/RenderCore/Public/ShaderParameterParser.h
--- a/Engine/Source/Runtime/RenderCore/Public/ShaderParameterParser.h
+++ b/Engine/Source/Runtime/RenderCore/Public/ShaderParameterParser.h
@@ -151,22 +151,23 @@ public:
 	);
 
 	/** Gets parsing information from a parameter binding name. */
-	const FParsedShaderParameter& FindParameterInfos(const FString& ParameterName) const
+	const FParsedShaderParameter& FindParameterInfos(FStringView ParameterName) const
 	{
-		const FParsedShaderParameter* ShaderParameter = ParsedParameters.Find(ParameterName);
-		checkf(ShaderParameter != nullptr, TEXT("Missing shader parameter '%s'"), *ParameterName);
+		const FParsedShaderParameter* ShaderParameter = FindParameterInfosUnsafe(ParameterName);
+		checkf(ShaderParameter != nullptr, TEXT("Missing shader parameter '%.*s'"), ParameterName.Len(), ParameterName.GetData());
 		return *ShaderParameter;
 	}
 
-	const FParsedShaderParameter* FindParameterInfosUnsafe(const FString& ParameterName) const
+	const FParsedShaderParameter* FindParameterInfosUnsafe(FStringView ParameterName) const
 	{
-		return ParsedParameters.Find(ParameterName);
+		const uint32 KeyHash = decltype(ParsedParameters)::KeyFuncsType::GetKeyHash(ParameterName);
+		return ParsedParameters.FindByHash(KeyHash, ParameterName);
 	}
 
 	/** Validates the shader parameter in code is compatible with the shader parameter structure. */
 	RENDERCORE_API void ValidateShaderParameterType(
 		const FShaderCompilerInput& CompilerInput,
-		const FString& ShaderBindingName,
+		FStringView ShaderBindingName,
 		int32 ReflectionOffset,
 		int32 ReflectionSize,
 		bool bPlatformSupportsPrecisionModifier,
@@ -174,7 +175,7 @@ public:
 
 	void ValidateShaderParameterType(
 		const FShaderCompilerInput& CompilerInput,
-		const FString& ShaderBindingName,
+		FStringView ShaderBindingName,
 		int32 ReflectionOffset,
 		int32 ReflectionSize,
 		FShaderCompilerOutput& CompilerOutput) const
