Adding a printf style method to add errors to FShaderCompilerOutput.

#Synced-CL 49623085

diff --git a/Engine/Source/Developer/Apple/MetalShaderFormat/Private/MetalCompileShaderMSC.cpp b/Engine/Source/Developer/Apple/MetalShaderFormat/Private/MetalCompileShaderMSC.cpp
--- a/Engine/Source/Developer/Apple/MetalShaderFormat/Private/MetalCompileShaderMSC.cpp
+++ b/Engine/Source/Developer/Apple/MetalShaderFormat/Private/MetalCompileShaderMSC.cpp
@@ -735,8 +735,7 @@ void FMetalCompileShaderMSC::DoCompileMetalShader(
 		IRRootSignature* RootSignature = IRRootSignatureCreateFromDescriptor(&RootSignatureDesc, &RootSignatureCreationError);
 		if (RootSignature == nullptr || RootSignatureCreationError != nullptr)
 		{
-			FShaderCompilerError Error(FString::Printf(TEXT("Error: MetalShaderConverter failed to create a root signature for '%s' (%s)!"), *Input.EntryPointName, ANSI_TO_TCHAR((const char *)IRErrorGetPayload(RootSignatureCreationError))));
-			Output.Errors.Add(Error);
+			Output.AddErrorf(TEXT("Error: MetalShaderConverter failed to create a root signature for '%s' (%s)!"), *Input.EntryPointName, ANSI_TO_TCHAR((const char*)IRErrorGetPayload(RootSignatureCreationError)));
 			Output.bSucceeded = false;
 			
 			return;
@@ -771,8 +770,7 @@ void FMetalCompileShaderMSC::DoCompileMetalShader(
 		IRObject* AirBytecode = IRCompilerAllocCompileAndLink(CompilerInstance, nullptr, DXILBytecode, &CompileError);
 		if (!AirBytecode || CompileError != nullptr)
 		{
-			FShaderCompilerError Error(FString::Printf(TEXT("Error: MetalShaderConverter failed to produce air bytecode for '%s' (%s)!"), *Input.EntryPointName, ANSI_TO_TCHAR((const char *)IRErrorGetPayload(CompileError))));
-			Output.Errors.Add(Error);
+			Output.AddErrorf(TEXT("Error: MetalShaderConverter failed to produce air bytecode for '%s' (%s)!"), *Input.EntryPointName, ANSI_TO_TCHAR((const char*)IRErrorGetPayload(CompileError)));
 			Output.bSucceeded = false;
 			
 			return;
@@ -868,8 +866,7 @@ void FMetalCompileShaderMSC::DoCompileMetalShader(
 		IRMetalLibBinary* GeneratedMetalLib = IRMetalLibBinaryCreate();
 		if (!IRObjectGetMetalLibBinary(AirBytecode, ShaderStage, GeneratedMetalLib))
 		{
-			FShaderCompilerError Error(FString::Printf(TEXT("Error: MetalShaderConverter failed to produce a metallib for '%s'!"), *Input.EntryPointName));
-			Output.Errors.Add(Error);
+			Output.AddErrorf(TEXT("Error: MetalShaderConverter failed to produce a metallib for '%s'!"), *Input.EntryPointName);
 			Output.bSucceeded = false;
 			
 			return;
diff --git a/Engine/Source/Developer/ShaderCompilerCommon/Private/ShaderCompilerCommon.cpp b/Engine/Source/Developer/ShaderCompilerCommon/Private/ShaderCompilerCommon.cpp
--- a/Engine/Source/Developer/ShaderCompilerCommon/Private/ShaderCompilerCommon.cpp
+++ b/Engine/Source/Developer/ShaderCompilerCommon/Private/ShaderCompilerCommon.cpp
@@ -437,21 +437,21 @@ bool UE::ShaderCompilerCommon::ValidatePackedResourceCounts(FShaderCompilerOutpu
 		if (EnumHasAnyFlags(PackedResourceCounts.UsageFlags, EShaderResourceUsageFlags::BindlessResources) && PackedResourceCounts.NumSRVs > 0)
 		{
 			const FString Names = GetAllResourcesOfType(EShaderParameterType::SRV);
-			Output.Errors.Add(FString::Printf(TEXT("Shader is mixing bindless resources with non-bindless resources. %d SRV slots were detected: %s"), PackedResourceCounts.NumSRVs, *Names));
+			Output.AddErrorf(TEXT("Shader is mixing bindless resources with non-bindless resources. %d SRV slots were detected: %s"), PackedResourceCounts.NumSRVs, *Names);
 			Output.bSucceeded = false;
 		}
 
 		if (EnumHasAnyFlags(PackedResourceCounts.UsageFlags, EShaderResourceUsageFlags::BindlessResources) && PackedResourceCounts.NumUAVs > 0)
 		{
 			const FString Names = GetAllResourcesOfType(EShaderParameterType::UAV);
-			Output.Errors.Add(FString::Printf(TEXT("Shader is mixing bindless resources with non-bindless resources. %d UAV slots were detected: %s"), PackedResourceCounts.NumUAVs, *Names));
+			Output.AddErrorf(TEXT("Shader is mixing bindless resources with non-bindless resources. %d UAV slots were detected: %s"), PackedResourceCounts.NumUAVs, *Names);
 			Output.bSucceeded = false;
 		}
 
 		if (EnumHasAnyFlags(PackedResourceCounts.UsageFlags, EShaderResourceUsageFlags::BindlessSamplers) && PackedResourceCounts.NumSamplers > 0)
 		{
 			const FString Names = GetAllResourcesOfType(EShaderParameterType::Sampler);
-			Output.Errors.Add(FString::Printf(TEXT("Shader is mixing bindless samplers with non-bindless samplers. %d sampler slots were detected: %s"), PackedResourceCounts.NumSamplers, *Names));
+			Output.AddErrorf(TEXT("Shader is mixing bindless samplers with non-bindless samplers. %d sampler slots were detected: %s"), PackedResourceCounts.NumSamplers, *Names);
 			Output.bSucceeded = false;
 		}
 	}
@@ -621,15 +621,12 @@ void HandleReflectedUniformBufferConstantBufferMember(
 	if (EnumHasAnyFlags(Reason, EUniformBufferMemberReflectionReason::MaterialUBBindError))
 	{
 #define SVARG(N) N.Len(), N.GetData()
-		FString ErrorMessage = FString::Printf(TEXT("%s: Unauthorized to use %.*s from %.*s at frequency %s"), 
-			*Input.GenerateShaderName(), 
+		Output.AddErrorf(TEXT("%s: Unauthorized to use %.*s from %.*s at frequency %s"),
+			*Input.GenerateShaderName(),
 			SVARG(InMemberName),
 			SVARG(UniformBufferName),
 			GetShaderFrequencyString(Output.Target.GetFrequency()));
 #undef SVARG
-
-		FShaderCompilerError Error(MoveTemp(ErrorMessage));
-		Output.Errors.Add(Error);
 		Output.bSucceeded = false;
 	}
 
diff --git a/Engine/Source/Developer/VulkanShaderFormat/Private/SpirVShaderCompiler.inl b/Engine/Source/Developer/VulkanShaderFormat/Private/SpirVShaderCompiler.inl
--- a/Engine/Source/Developer/VulkanShaderFormat/Private/SpirVShaderCompiler.inl
+++ b/Engine/Source/Developer/VulkanShaderFormat/Private/SpirVShaderCompiler.inl
@@ -424,7 +424,7 @@ static void BuildShaderOutput(
 		FShaderCompilerResourceTable CompilerSRT;
 		if (!BuildResourceTableMapping(ShaderInput.Environment.ResourceTableMap, ShaderInput.Environment.UniformBufferMap, UsedUniformBufferSlots, ShaderOutput.ParameterMap, CompilerSRT))
 		{
-			ShaderOutput.Errors.Add(TEXT("Internal error on BuildResourceTableMapping."));
+			ShaderOutput.AddErrorf(TEXT("Internal error on BuildResourceTableMapping."));
 			return;
 		}
 		UE::ShaderCompilerCommon::BuildShaderResourceTable(CompilerSRT, SerializedOutput.ShaderResourceTable);
@@ -942,7 +942,7 @@ static bool BuildShaderOutputFromSpirv(
 		const char* OptArgs[] = { "--strip-reflect", "-O"};
 		if (!CompilerContext.OptimizeSpirv(SerializedOutput.Spirv.Data, OptArgs, UE_ARRAY_COUNT(OptArgs)))
 		{
-			Output.Errors.Add(TEXT("Failed to strip debug instructions from SPIR-V module"));
+			Output.AddErrorf(TEXT("Failed to strip debug instructions from SPIR-V module"));
 			return false;
 		}
 	}
@@ -958,7 +958,7 @@ static bool BuildShaderOutputFromSpirv(
 		};
 		if (!CompilerContext.OptimizeSpirv(SerializedOutput.Spirv.Data, OptArgs, UE_ARRAY_COUNT(OptArgs)))
 		{
-			Output.Errors.Add(TEXT("Failed to apply driver patches for Android"));
+			Output.AddErrorf(TEXT("Failed to apply driver patches for Android"));
 			return false;
 		}
 	}
diff --git a/Engine/Source/Developer/Windows/ShaderFormatD3D/Private/D3DShaderCompiler.cpp b/Engine/Source/Developer/Windows/ShaderFormatD3D/Private/D3DShaderCompiler.cpp
--- a/Engine/Source/Developer/Windows/ShaderFormatD3D/Private/D3DShaderCompiler.cpp
+++ b/Engine/Source/Developer/Windows/ShaderFormatD3D/Private/D3DShaderCompiler.cpp
@@ -271,7 +271,7 @@ static FString D3D11CreateShaderCompileCommandLine(
 			break;
 
 		default:
-			Output.Errors.Emplace(TEXT("Unknown D3DCOMPILE optimization level"));
+			Output.AddErrorf(TEXT("Unknown D3DCOMPILE optimization level"));
 			break;
 	}
 
@@ -529,7 +529,7 @@ static bool PatchHlslWithReorderedIOVariables(
 	{
 		for (FString& Error : ParsingErrors)
 		{
-			OutErrors.Add(FShaderCompilerError(MoveTemp(Error)));
+			OutErrors.Emplace(MoveTemp(Error));
 		}
 		return false;
 	}
@@ -955,7 +955,7 @@ static bool CompileAndProcessD3DShaderFXCExt(
 			// Compile again with FXC and disable special compiler rule to simplify control flow - 2nd try
 			if (Result == E_FAIL && D3DCompileErrorContainsValidationErrors(Errors.GetReference()))
 			{
-				Output.Errors.Add(FShaderCompilerError(TEXT("Validation error in FXC encountered: Compiling intermediate HLSL a second time with simplified control flow")));
+				Output.AddErrorf(TEXT("Validation error in FXC encountered: Compiling intermediate HLSL a second time with simplified control flow"));
 
 				// Rule 0x08024065 is described as "simplify flow control that writes the same value in each flow control path"
 				const FAnsiStringView PragmaDirectiveCode = "#pragma ruledisable 0x08024065\n";
@@ -966,7 +966,7 @@ static bool CompileAndProcessD3DShaderFXCExt(
 				// If FXC compilation still fails with a validation error, compile again and skip optimizations entirely as a last resort - 3rd try
 				if (Result == E_FAIL && D3DCompileErrorContainsValidationErrors(Errors.GetReference()))
 				{
-					Output.Errors.Add(FShaderCompilerError(TEXT("Validation error in FXC encountered: Compiling intermediate HLSL a third time without optimization (D3DCOMPILE_SKIP_OPTIMIZATION)")));
+					Output.AddErrorf(TEXT("Validation error in FXC encountered: Compiling intermediate HLSL a third time without optimization (D3DCOMPILE_SKIP_OPTIMIZATION)"));
 
 					const uint32 CompileFlagsSkipOptimizations = CompileFlagsNoWarningsAsErrors | D3DCOMPILE_SKIP_OPTIMIZATION;
 					Result = CompileCrossCompiledHlsl(CrossCompiledSource, CompileFlagsSkipOptimizations);
@@ -979,7 +979,7 @@ static bool CompileAndProcessD3DShaderFXCExt(
 				FSCWErrorCode::Reset();
 
 				// Let the user know this shader had to be cross-compiled due to a crash in FXC. Only shows up if CVar 'r.ShaderDevelopmentMode' is enabled.
-				Output.Errors.Add(FShaderCompilerError(TEXT("Cross-compiled shader to intermediate HLSL after first attempt crashed FXC")));
+				Output.AddErrorf(TEXT("Cross-compiled shader to intermediate HLSL after first attempt crashed FXC"));
 
 				// Output the errors from the initial run so that the user can know what failed and eventually correct it or at least make an informed decision about whether to bypass the error using DXC pre-compilation:
 				Algo::Transform(InitialFXCRunFilteredErrors, Output.Errors, [](const FString& InInitialError) { return FShaderCompilerError(*InInitialError); });
@@ -988,9 +988,7 @@ static bool CompileAndProcessD3DShaderFXCExt(
 	}
 	else
 	{
-		FShaderCompilerError NewError;
-		NewError.StrippedErrorMessage = TEXT("Couldn't find D3D shader compiler DLL");
-		Output.Errors.Add(NewError);
+		Output.AddErrorf(TEXT("Couldn't find D3D shader compiler DLL"));
 		Result = E_FAIL;
 	}
 
@@ -1024,7 +1022,7 @@ static bool CompileAndProcessD3DShaderFXCExt(
 			{
 				NewError.StrippedErrorMessage = CurrentError;
 			}
-			Output.Errors.Add(NewError);
+			Output.Errors.Emplace(MoveTemp(NewError));
 		}
 	}
 
@@ -1053,9 +1051,7 @@ static bool CompileAndProcessD3DShaderFXCExt(
 					// dcl_globalFlags will contain enableDoublePrecisionFloatOps when the shader uses doubles, even though the docs on dcl_globalFlags don't say anything about this
 					if (DisassemblyStringW.Contains(TEXT("enableDoublePrecisionFloatOps")))
 					{
-						FShaderCompilerError NewError;
-						NewError.StrippedErrorMessage = TEXT("Shader uses double precision floats, which are not supported on all D3D11 hardware!");
-						Output.Errors.Add(NewError);
+						Output.AddErrorf(TEXT("Shader uses double precision floats, which are not supported on all D3D11 hardware!"));
 						return false;
 					}
 				}
@@ -1069,9 +1065,7 @@ static bool CompileAndProcessD3DShaderFXCExt(
 					// https://docs.microsoft.com/en-us/windows/win32/direct3ddxgi/format-support-for-direct3d-11-0-feature-level-hardware
 					if (DisassemblyStringW.Contains(TEXT("Typed UAV Load Additional Formats")))
 					{
-						FShaderCompilerError NewError;
-						NewError.StrippedErrorMessage = TEXT("Shader uses UAV loads from additional typed formats, which are not supported on all D3D11 hardware! Set r.D3D.CheckedForTypedUAVs=0 if you want to allow typed UAV loads for your project, or individual shaders can opt-in by specifying CFLAG_AllowTypedUAVLoads.");
-						Output.Errors.Add(NewError);
+						Output.AddErrorf(TEXT("Shader uses UAV loads from additional typed formats, which are not supported on all D3D11 hardware! Set r.D3D.CheckedForTypedUAVs=0 if you want to allow typed UAV loads for your project, or individual shaders can opt-in by specifying CFLAG_AllowTypedUAVLoads."));
 						return false;
 					}
 				}
@@ -1169,9 +1163,7 @@ static bool CompileAndProcessD3DShaderFXCExt(
 		}
 		else
 		{
-			FShaderCompilerError NewError;
-			NewError.StrippedErrorMessage = TEXT("Couldn't find shader reflection function in D3D Compiler DLL");
-			Output.Errors.Add(NewError);
+			Output.AddErrorf(TEXT("Couldn't find shader reflection function in D3D Compiler DLL"));
 			Result = E_FAIL;
 			Output.bSucceeded = false;
 		}
@@ -1185,9 +1177,7 @@ static bool CompileAndProcessD3DShaderFXCExt(
 		// Check for resource limits for feature level 11.0
 		if (CompileData.NumUAVs > GD3DMaximumNumUAVs)
 		{
-			FShaderCompilerError NewError;
-			NewError.StrippedErrorMessage = FString::Printf(TEXT("Number of UAVs exceeded limit: %d slots used, but limit is %d due to maximum feature level 11.0"), CompileData.NumUAVs, GD3DMaximumNumUAVs);
-			Output.Errors.Add(NewError);
+			Output.AddErrorf(TEXT("Number of UAVs exceeded limit: %d slots used, but limit is %d due to maximum feature level 11.0"), CompileData.NumUAVs, GD3DMaximumNumUAVs);
 			Result = E_FAIL;
 			Output.bSucceeded = false;
 		}
@@ -1250,9 +1240,7 @@ static bool CompileAndProcessD3DShaderFXCExt(
 						}
 						else
 						{
-							FShaderCompilerError NewError;
-							NewError.StrippedErrorMessage = TEXT("Symbol generation was requested, but no PDB blob exists in the compiler output.");
-							Output.Errors.Add(NewError);
+							Output.AddErrorf(TEXT("Symbol generation was requested, but no PDB blob exists in the compiler output."));
 							Result = E_FAIL;
 							Output.bSucceeded = false;
 						}
@@ -1263,9 +1251,7 @@ static bool CompileAndProcessD3DShaderFXCExt(
 				}
 				else
 				{
-					FShaderCompilerError NewError;
-					NewError.StrippedErrorMessage = TEXT("Symbol or symbols info generation was requested, but no debug name blob exists in the compiler output.");
-					Output.Errors.Add(NewError);
+					Output.AddErrorf(TEXT("Symbol or symbols info generation was requested, but no debug name blob exists in the compiler output."));
 					Output.bSucceeded = false;
 				}
 			}
@@ -1453,7 +1439,7 @@ void CompileD3DShader(const FShaderCompilerInput& Input, const FShaderPreprocess
 
 	if (!ShaderProfile)
 	{
-		Output.Errors.Add(FShaderCompilerError(*FString::Printf(TEXT("Unrecognized shader frequency %s"), GetShaderFrequencyString((EShaderFrequency)Input.Target.Frequency))));
+		Output.AddErrorf(TEXT("Unrecognized shader frequency %s"), GetShaderFrequencyString(Input.Target.GetFrequency()));
 		return;
 	}
 
@@ -1528,7 +1514,7 @@ void CompileD3DShader(const FShaderCompilerInput& Input, const FShaderPreprocess
 
 	if (!bSuccess && !Output.Errors.Num())
 	{
-		Output.Errors.Add(TEXT("Compile failed without errors!"));
+		Output.AddErrorf(TEXT("Compile failed without errors!"));
 	}
 
 	ShaderParameterParser.ValidateShaderParameterTypes(Input, Output);
diff --git a/Engine/Source/Developer/Windows/ShaderFormatD3D/Private/D3DShaderCompiler.inl b/Engine/Source/Developer/Windows/ShaderFormatD3D/Private/D3DShaderCompiler.inl
--- a/Engine/Source/Developer/Windows/ShaderFormatD3D/Private/D3DShaderCompiler.inl
+++ b/Engine/Source/Developer/Windows/ShaderFormatD3D/Private/D3DShaderCompiler.inl
@@ -214,12 +214,12 @@ template <typename ID3D1xShaderReflection, typename D3D1x_SHADER_DESC, typename
 				}
 				else
 				{
-					FString ErrorMessage = FString::Printf(
+					Output.AddErrorf(
 						TEXT("Error: %s is expected to always be in the API slot %d, but is actually in slot %d."),
 						FShaderParametersMetadata::kRootUniformBufferBindingName,
 						FShaderParametersMetadata::kRootCBufferBindingIndex,
-						CBIndex);
-					Output.Errors.Add(FShaderCompilerError(*ErrorMessage));
+						CBIndex
+					);
 					Output.bSucceeded = false;
 				}
 			}
@@ -708,10 +708,9 @@ bool RemoveUnusedInterpolators(
 						if (Output.UsedAttributes.Num() > CompileData.ShaderInputs.Num())
 						{
 							UE_LOG(LogD3DShaderCompiler, Warning, TEXT("Second pass had more used attributes (%d) than first pass (%d)"), Output.UsedAttributes.Num(), CompileData.ShaderInputs.Num());
-							FShaderCompilerError NewError;
-							NewError.StrippedErrorMessage = FString::Printf(TEXT("Second pass had more used attributes (%d) than first pass (%d)"), Output.UsedAttributes.Num(), CompileData.ShaderInputs.Num());
+							FShaderCompilerError NewError(FString::Printf(TEXT("Second pass had more used attributes (%d) than first pass (%d)"), Output.UsedAttributes.Num(), CompileData.ShaderInputs.Num()));
 							Output = OriginalOutput;
-							Output.Errors.Add(NewError);
+							Output.Errors.Emplace(MoveTemp(NewError));
 							break;
 						}
 
@@ -722,13 +721,12 @@ bool RemoveUnusedInterpolators(
 								Attempt + 1,
 								Output.UsedAttributes.Num(),
 								CompileData.ShaderInputs.Num());
-							FShaderCompilerError NewError;
-							NewError.StrippedErrorMessage = FString::Printf(TEXT("Unable to determine unused inputs after %d attempts (last number of used attributes: %d, previous step:%d)!"),
+							FShaderCompilerError NewError(FString::Printf(TEXT("Unable to determine unused inputs after %d attempts (last number of used attributes: %d, previous step:%d)!"),
 								Attempt + 1,
 								Output.UsedAttributes.Num(),
-								CompileData.ShaderInputs.Num());
+								CompileData.ShaderInputs.Num()));
 							Output = OriginalOutput;
-							Output.Errors.Add(NewError);
+							Output.Errors.Emplace(MoveTemp(NewError));
 							break;
 						}
 
diff --git a/Engine/Source/Developer/Windows/ShaderFormatD3D/Private/D3DShaderCompilerDXC.cpp b/Engine/Source/Developer/Windows/ShaderFormatD3D/Private/D3DShaderCompilerDXC.cpp
--- a/Engine/Source/Developer/Windows/ShaderFormatD3D/Private/D3DShaderCompilerDXC.cpp
+++ b/Engine/Source/Developer/Windows/ShaderFormatD3D/Private/D3DShaderCompilerDXC.cpp
@@ -1110,7 +1110,7 @@ bool CompileAndProcessD3DShaderDXC(
 						}
 					}
 
-					Output.Errors.Add(MoveTemp(ErrorString));
+					Output.Errors.Emplace(MoveTemp(ErrorString));
 					Output.bSucceeded = false;
 				}
 			}
@@ -1302,7 +1302,7 @@ bool CompileAndProcessD3DShaderDXC(
 			const bool bKnownError = ErrorMsg[0] != TEXT('\0');
 
 			FString ErrorString = FString::Printf(TEXT("D3DCompileToDxil failed. Error code: %s (0x%08X)."), bKnownError ? ErrorMsg : TEXT("Unknown error"), (int)D3DCompileToDxilResult);
-			Output.Errors.Add(MoveTemp(ErrorString));
+			Output.Errors.Emplace(MoveTemp(ErrorString));
 		}
 	}
 
diff --git a/Engine/Source/Runtime/RenderCore/Private/ShaderCore.cpp b/Engine/Source/Runtime/RenderCore/Private/ShaderCore.cpp
--- a/Engine/Source/Runtime/RenderCore/Private/ShaderCore.cpp
+++ b/Engine/Source/Runtime/RenderCore/Private/ShaderCore.cpp
@@ -1941,7 +1941,7 @@ public:
 		// then we could potentially merge these here instead, but that _seems_ unlikely to be the case in practice.
 		if (Job.Output.UsedAttributes != Job.SecondaryOutput->UsedAttributes)
 		{
-			Job.Output.Errors.Add(FShaderCompilerError(TEXT("Mismatch in UsedAttributes between primary and secondary compilation; these must match since later stages may not have a secondary compile and so require both versions of the earlier stage to output the same set of attributes.")));
+			Job.Output.AddErrorf(TEXT("Mismatch in UsedAttributes between primary and secondary compilation; these must match since later stages may not have a secondary compile and so require both versions of the earlier stage to output the same set of attributes."));
 			return false;
 		}
 		return true;
@@ -2066,10 +2066,10 @@ bool PreprocessShader(FShaderCommonCompileJob* Job)
 			else
 			{
 				// skip subsequent stage preprocessing if a prior stage failed to avoid unnecessary work, but log an error to indicate this
-				FString Error = FString::Printf(
+				StageJob->Output.AddErrorf(
 					TEXT("Preprocessing %s stage skipped due to earlier stage preprocessing failure."),
-					GetShaderFrequencyString(StageJob->Input.Target.GetFrequency()));
-				StageJob->Output.Errors.Add(FShaderCompilerError(MoveTemp(Error)));
+					GetShaderFrequencyString(StageJob->Input.Target.GetFrequency())
+				);
 			}
 		}
 		return !bAnyFailed;
@@ -2107,12 +2107,11 @@ void CompileShader(const TArray<const IShaderFormat*>& ShaderFormats, FShaderCom
 	FInternalShaderCompilerFunctions::CompileShaderInternal(Compiler, Job, WorkingDirectory, ExceptionMsg, ExceptionCallstack, CompileCount);
 	if (!Job.Output.bSucceeded && (!ExceptionMsg.IsEmpty() || !ExceptionCallstack.IsEmpty()))
 	{
-		FShaderCompilerError Error;
-		Error.StrippedErrorMessage = FString::Printf(
-			TEXT("Exception encountered in platform compiler: %s\nException Callstack:\n%s"), 
-			*ExceptionMsg, 
-			*ExceptionCallstack);
-		Job.Output.Errors.Add(Error);
+		Job.Output.AddErrorf(
+			TEXT("Exception encountered in platform compiler: %s\nException Callstack:\n%s"),
+			*ExceptionMsg,
+			*ExceptionCallstack
+		);
 	}
 
 	Job.bSucceeded = Job.Output.bSucceeded;
diff --git a/Engine/Source/Runtime/RenderCore/Private/ShaderParameterParser.cpp b/Engine/Source/Runtime/RenderCore/Private/ShaderParameterParser.cpp
--- a/Engine/Source/Runtime/RenderCore/Private/ShaderParameterParser.cpp
+++ b/Engine/Source/Runtime/RenderCore/Private/ShaderParameterParser.cpp
@@ -1359,13 +1359,12 @@ void FShaderParameterParser::ValidateShaderParameterType(
 		{
 			FString CppCodeName = CompilerInput.RootParametersStructure->GetFullMemberCodeName(ParsedParameter.ConstantBufferOffset);
 
-			FShaderCompilerError Error;
-			Error.StrippedErrorMessage = FString::Printf(
+			FShaderCompilerError Error(FString::Printf(
 				TEXT("Error: Type %.*s of shader parameter %.*s in shader mismatch the shader parameter structure: %s expects a %s"),
 				ParsedParameter.ParsedType.Len(), ParsedParameter.ParsedType.GetData(),
 				ShaderBindingName.Len(), ShaderBindingName.GetData(),
 				*CppCodeName,
-				*ExpectedShaderType);
+				*ExpectedShaderType));
 			GetParameterFileAndLine(ParsedParameter, Error.ErrorVirtualFilePath, Error.ErrorLineString);
 
 			CompilerOutput.Errors.Add(Error);
@@ -1380,13 +1379,12 @@ void FShaderParameterParser::ValidateShaderParameterType(
 	{
 		FString CppCodeName = CompilerInput.RootParametersStructure->GetFullMemberCodeName(ParsedParameter.ConstantBufferOffset);
 
-		FShaderCompilerError Error;
-		Error.StrippedErrorMessage = FString::Printf(
+		FShaderCompilerError Error(FString::Printf(
 			TEXT("Error: The size required to bind shader parameter %.*s is %i bytes, smaller than %s that is %i bytes in the parameter structure."),
 			ShaderBindingName.Len(), ShaderBindingName.GetData(),
 			ReflectionSize,
 			*CppCodeName,
-			ParsedParameter.MemberSize);
+			ParsedParameter.MemberSize));
 		GetParameterFileAndLine(ParsedParameter, Error.ErrorVirtualFilePath, Error.ErrorLineString);
 
 		CompilerOutput.Errors.Add(Error);
diff --git a/Engine/Source/Runtime/RenderCore/Public/ShaderCompilerCore.h b/Engine/Source/Runtime/RenderCore/Public/ShaderCompilerCore.h
--- a/Engine/Source/Runtime/RenderCore/Public/ShaderCompilerCore.h
+++ b/Engine/Source/Runtime/RenderCore/Public/ShaderCompilerCore.h
@@ -418,11 +418,7 @@ struct FShaderCompilerInput
 struct FShaderCompilerError
 {
 	FShaderCompilerError(const TCHAR* InStrippedErrorMessage = TEXT(""))
-		: ErrorVirtualFilePath(TEXT(""))
-		, ErrorLineString(TEXT(""))
-		, StrippedErrorMessage(InStrippedErrorMessage)
-		, HighlightedLine(TEXT(""))
-		, HighlightedLineMarker(TEXT(""))
+		: StrippedErrorMessage(InStrippedErrorMessage)
 	{}
 
 	FShaderCompilerError(const TCHAR* InVirtualFilePath, const TCHAR* InLineString, const TCHAR* InStrippedErrorMessage)
@@ -434,17 +430,11 @@ struct FShaderCompilerError
 	{}
 
 	FShaderCompilerError(FString&& InStrippedErrorMessage)
-		: ErrorVirtualFilePath(TEXT(""))
-		, ErrorLineString(TEXT(""))
-		, StrippedErrorMessage(MoveTemp(InStrippedErrorMessage))
-		, HighlightedLine(TEXT(""))
-		, HighlightedLineMarker(TEXT(""))
+		: StrippedErrorMessage(MoveTemp(InStrippedErrorMessage))
 	{}
 
 	FShaderCompilerError(FString&& InStrippedErrorMessage, FString&& InHighlightedLine, FString&& InHighlightedLineMarker)
-		: ErrorVirtualFilePath(TEXT(""))
-		, ErrorLineString(TEXT(""))
-		, StrippedErrorMessage(MoveTemp(InStrippedErrorMessage))
+		: StrippedErrorMessage(MoveTemp(InStrippedErrorMessage))
 		, HighlightedLine(MoveTemp(InHighlightedLine))
 		, HighlightedLineMarker(MoveTemp(InHighlightedLineMarker))
 	{}
@@ -558,6 +548,12 @@ struct FShaderCompilerOutput
 	/** Add optional diagnostic data in ShaderCode to perform assert translation at runtime*/
 	RENDERCORE_API void SerializeShaderDiagnosticData();
 
+	template<typename... Types>
+	void AddErrorf(UE::Core::TCheckedFormatString<FString::FmtCharType, Types...> Fmt, Types... Args)
+	{
+		Errors.Emplace(FString::Printf(Fmt, Args...));
+	}
+
 	template<typename TValue>
 	void AddStatistic(const TCHAR* Name, TValue Value, FGenericShaderStat::EFlags Flags = FGenericShaderStat::EFlags::None, FName TagName = NAME_None)
 	{
