IsCompatibleBinding needs to handle NvApi's g_NvidiaExt space.
Bumping SM6 version.

#Synced-CL 50075005

diff --git a/Engine/Source/Developer/Windows/ShaderFormatD3D/Private/D3DShaderCompilerDXC.cpp b/Engine/Source/Developer/Windows/ShaderFormatD3D/Private/D3DShaderCompilerDXC.cpp
--- a/Engine/Source/Developer/Windows/ShaderFormatD3D/Private/D3DShaderCompilerDXC.cpp
+++ b/Engine/Source/Developer/Windows/ShaderFormatD3D/Private/D3DShaderCompilerDXC.cpp
@@ -679,6 +679,11 @@ inline bool IsCompatibleBinding(const D3D12_SHADER_INPUT_BIND_DESC& BindDesc, ui
 		bIsCompatibleBinding = bIsAMDExtensionDX12 && (BindDesc.Space == AGS_DX12_SHADER_INSTRINSICS_SPACE_ID);
 	}
 	if (!bIsCompatibleBinding)
+	{
+		const bool bIsNVExtension = (FCStringAnsi::Strcmp(BindDesc.Name, "g_NvidiaExt") == 0);
+		bIsCompatibleBinding = bIsNVExtension && (BindDesc.Space == UE_HLSL_SPACE_NV_SHADER_EXTN);
+	}
+	if (!bIsCompatibleBinding)
 	{
 		const bool bIsUEDebugBuffer = (FCStringAnsi::Strcmp(BindDesc.Name, "UEDiagnosticBuffer") == 0);
 		bIsCompatibleBinding = bIsUEDebugBuffer && (BindDesc.Space == UE_HLSL_SPACE_DIAGNOSTIC);
diff --git a/Engine/Source/Developer/Windows/ShaderFormatD3D/Private/ShaderFormatD3D.cpp b/Engine/Source/Developer/Windows/ShaderFormatD3D/Private/ShaderFormatD3D.cpp
--- a/Engine/Source/Developer/Windows/ShaderFormatD3D/Private/ShaderFormatD3D.cpp
+++ b/Engine/Source/Developer/Windows/ShaderFormatD3D/Private/ShaderFormatD3D.cpp
@@ -22,7 +22,7 @@ static FName NAME_PCD3D_ES3_1(TEXT("PCD3D_ES31"));
 static FName NAME_DXIL(TEXT("DXIL"));
 
 static const FGuid UE_SHADER_PCD3D_SHARED_VER = FGuid("9D405D9B-44C0-4747-8E88-1EDB7413F5C3");
-static const FGuid UE_SHADER_PCD3D_SM6_VER    = FGuid("7F9AE1F2-9692-46B4-875A-28C2595E84AA");
+static const FGuid UE_SHADER_PCD3D_SM6_VER    = FGuid("123ac6e6-a0a2-4a5f-b8be-4102376788ab");
 static const FGuid UE_SHADER_PCD3D_SM5_VER    = FGuid("F30C5AFA-77F7-48DE-8CC1-A075F179443C");
 static const FGuid UE_SHADER_PCD3D_ES3_1_VER  = FGuid("952939D9-1156-4347-97E9-9FFEA1A9FE14");
 
