Shaders: Silence for loop warning from dxc

Ignore the warning about re-using the same index variable in consecutive for loops. We have this pattern in a few places and it seems innocuous.

#rnx
#Synced-CL 51869758

diff --git a/Engine/Source/Developer/Windows/ShaderFormatD3D/Private/D3DShaderCompilerDXC.cpp b/Engine/Source/Developer/Windows/ShaderFormatD3D/Private/D3DShaderCompilerDXC.cpp
--- a/Engine/Source/Developer/Windows/ShaderFormatD3D/Private/D3DShaderCompilerDXC.cpp
+++ b/Engine/Source/Developer/Windows/ShaderFormatD3D/Private/D3DShaderCompilerDXC.cpp
@@ -266,6 +266,7 @@ public:
 
 		// disable undesired warnings
 		ExtraArguments.Add(TEXT("-Wno-parentheses-equality"));
+        ExtraArguments.Add(TEXT("-Wno-for-redefinition"));
 
 		// working around bindless conversion specific issue where globallycoherent on a function return type is flagged as ignored even though it is necessary.
 		// github issue: [REDACTED_PROFILE_URL]/DirectXShaderCompiler/issues/4537
