Enable flow control attributes for D3D when using the mobile renderer

#Synced-CL 52668772
#Synced-CL 52668772

diff --git a/Engine/Source/Developer/Windows/ShaderFormatD3D/Private/ShaderFormatD3D.cpp b/Engine/Source/Developer/Windows/ShaderFormatD3D/Private/ShaderFormatD3D.cpp
--- a/Engine/Source/Developer/Windows/ShaderFormatD3D/Private/ShaderFormatD3D.cpp
+++ b/Engine/Source/Developer/Windows/ShaderFormatD3D/Private/ShaderFormatD3D.cpp
@@ -224,11 +224,19 @@ public:
 		Input.Environment.SetDefine(TEXT("PLATFORM_SUPPORTS_SM6_0_WAVE_OPERATIONS"), bSM6Features);
 		Input.Environment.SetDefine(TEXT("PLATFORM_SUPPORTS_DIAGNOSTIC_BUFFER"),     bSM6Features);
 
+		const bool bSM6Profile = ShaderModel >= ED3DShaderModel::SM6_6;
+		const bool bSM5Profile = Input.ShaderFormat == NAME_PCD3D_SM5;
+
 		// "profiles" are almost analogous to ERHIFeatureLevel but RT shaders are forcing themselves to SM6
-		Input.Environment.SetDefine(TEXT("SM6_PROFILE"),   ShaderModel >= ED3DShaderModel::SM6_6);
-		Input.Environment.SetDefine(TEXT("SM5_PROFILE"),   (Input.ShaderFormat == NAME_PCD3D_SM5));
+		Input.Environment.SetDefine(TEXT("SM6_PROFILE"), bSM6Profile);
+		Input.Environment.SetDefine(TEXT("SM5_PROFILE"), bSM5Profile);
 		Input.Environment.SetDefine(TEXT("ES3_1_PROFILE"), (Input.ShaderFormat == NAME_PCD3D_ES3_1));
 
+		if (!bSM6Profile && !bSM5Profile)
+		{
+			Input.Environment.SetDefine(TEXT("COMPILER_SUPPORTS_ATTRIBUTES"), true);
+		}
+
 		// Add SM6.6+ specific defines. None of these are intended to be enabled in lower SM's
 		if (ShaderModel >= ED3DShaderModel::SM6_6)
 		{
